Add multiple clumpIDs support for groom geometry.

diff --git a/Engine/Plugins/Runtime/HairStrands/Source/HairStrandsCore/Private/GroomBuilder.cpp b/Engine/Plugins/Runtime/HairStrands/Source/HairStrandsCore/Private/GroomBuilder.cpp
--- a/Engine/Plugins/Runtime/HairStrands/Source/HairStrandsCore/Private/GroomBuilder.cpp
+++ b/Engine/Plugins/Runtime/HairStrands/Source/HairStrandsCore/Private/GroomBuilder.cpp
@@ -40,7 +40,7 @@ static FAutoConsoleVariableRef CVarHairGroupIndexBuilder_MaxVoxelResolution(TEXT
 
 FString FGroomBuilder::GetVersion()
 {
-	return TEXT("v8g");
+	return TEXT("v8h");
 }
 
 namespace FHairStrandsDecimation
@@ -179,7 +179,8 @@ namespace HairStrandsBuilder
 		}
 
 		const uint32 Attributes = HairStrands.GetAttributes();
-
+		const uint32 AttributeFlags = HairStrands.GetAttributeFlags();
+		const bool bHasMultipleClumpIDs = HasHairAttributeFlags(AttributeFlags, EHairAttributeFlags::HasMultipleClumpIDs);
 		TArray<uint32> AttributeRootUV;
 		TArray<uint32> AttributeSeed;
 		TArray<uint32> AttributeLength;
@@ -190,11 +191,11 @@ namespace HairStrandsBuilder
 
 		AttributeSeed.SetNum(FMath::DivideAndRoundUp(NumCurves, 4u));		// 1 bytes encoding - Per-Curve
 		AttributeLength.SetNum(FMath::DivideAndRoundUp(NumCurves, 2u));		// 2 bytes encoding - Per-Curve
-		if (HasHairAttribute(Attributes, EHairAttribute::RootUV)) 	{ AttributeRootUV.SetNum(NumCurves); }									// 4 bytes encoding - Per-Curve
-		if (HasHairAttribute(Attributes, EHairAttribute::ClumpID))	{ AttributeClumpIDs.SetNum(FMath::DivideAndRoundUp(NumCurves, 2u)); }	// 2 bytes encoding - Per-Curve
-		if (HasHairAttribute(Attributes, EHairAttribute::Color))	{ AttributeBaseColor.SetNum(NumPoints); }								// 4 bytes encoding - Per-Vertex
-		if (HasHairAttribute(Attributes, EHairAttribute::Roughness)){ AttributeRoughness.SetNum(FMath::DivideAndRoundUp(NumPoints, 2u)); }	// 1 bytes encoding - Per-Vertex
-		if (HasHairAttribute(Attributes, EHairAttribute::AO))		{ AttributeAO.SetNum(FMath::DivideAndRoundUp(NumPoints, 4u)); }			// 1 bytes encoding - Per-Vertex
+		if (HasHairAttribute(Attributes, EHairAttribute::RootUV)) 	{ AttributeRootUV.SetNum(NumCurves); }																		// 4 bytes encoding - Per-Curve
+		if (HasHairAttribute(Attributes, EHairAttribute::ClumpID))	{ AttributeClumpIDs.SetNum(bHasMultipleClumpIDs ? NumCurves * 2u : FMath::DivideAndRoundUp(NumCurves, 2u));}// 8 or 2 bytes encoding - Per-Curve
+		if (HasHairAttribute(Attributes, EHairAttribute::Color))	{ AttributeBaseColor.SetNum(NumPoints); }																	// 4 bytes encoding - Per-Vertex
+		if (HasHairAttribute(Attributes, EHairAttribute::Roughness)){ AttributeRoughness.SetNum(FMath::DivideAndRoundUp(NumPoints, 2u)); }										// 1 bytes encoding - Per-Vertex
+		if (HasHairAttribute(Attributes, EHairAttribute::AO))		{ AttributeAO.SetNum(FMath::DivideAndRoundUp(NumPoints, 4u)); }												// 1 bytes encoding - Per-Vertex
 
 		const FVector HairBoxCenter = HairStrands.BoundingBox.GetCenter();
 
@@ -313,8 +314,18 @@ namespace HairStrandsBuilder
 			// Per-Curve Clump ID (16 bits max)
 			if (HasHairAttribute(Attributes, EHairAttribute::ClumpID))
 			{
-				uint16* Packed = (uint16*)AttributeClumpIDs.GetData();
-				Packed[CurveIndex] = uint16(FMath::Clamp(Curves.ClumpIDs[CurveIndex], 0, 65535));
+				if (bHasMultipleClumpIDs)
+				{
+					uint64* Packed = (uint64*)AttributeClumpIDs.GetData();
+					Packed[CurveIndex] |= uint64(FMath::Clamp(Curves.ClumpIDs[CurveIndex].X, 0, 65535));
+					Packed[CurveIndex] |= uint64(FMath::Clamp(Curves.ClumpIDs[CurveIndex].Y, 0, 65535))<<16;
+					Packed[CurveIndex] |= uint64(FMath::Clamp(Curves.ClumpIDs[CurveIndex].Z, 0, 65535))<<32;
+				}
+				else
+				{
+					uint16* Packed = (uint16*)AttributeClumpIDs.GetData();
+					Packed[CurveIndex] = uint16(FMath::Clamp(Curves.ClumpIDs[CurveIndex].X, 0, 65535));
+				}
 			}
 
 			// Per-Curve Strand ID
@@ -360,7 +371,7 @@ namespace HairStrandsBuilder
 
 			for (uint32 AttributeIt=0; AttributeIt< HAIR_ATTRIBUTE_COUNT; ++AttributeIt)
 			{
-				OutBulkData.AttributeOffsets[AttributeIt] = HAIR_ATRIBUTE_INVALID_OFFSET;
+				OutBulkData.AttributeOffsets[AttributeIt] = HAIR_ATTRIBUTE_INVALID_OFFSET;
 			}
 
 			if (HasHairAttribute(Attributes, EHairAttribute::RootUV))
@@ -381,8 +392,16 @@ namespace HairStrandsBuilder
 
 			if (HasHairAttribute(Attributes, EHairAttribute::ClumpID))
 			{
-				OutBulkData.AttributeOffsets[HAIR_ATTRIBUTE_CLUMPID] = OutPackedAttributes.Num() * UintToByte;
-				OutPackedAttributes.Append(AttributeClumpIDs);
+				if (HasHairAttributeFlags(OutBulkData.ImportedAttributeFlags, EHairAttributeFlags::HasMultipleClumpIDs))
+				{
+					OutBulkData.AttributeOffsets[HAIR_ATTRIBUTE_CLUMPID3] = OutPackedAttributes.Num() * UintToByte;
+					OutPackedAttributes.Append(AttributeClumpIDs);
+				}
+				else
+				{
+					OutBulkData.AttributeOffsets[HAIR_ATTRIBUTE_CLUMPID] = OutPackedAttributes.Num() * UintToByte;
+					OutPackedAttributes.Append(AttributeClumpIDs);
+				}
 			}
 
 			if (HasHairAttribute(Attributes, EHairAttribute::Color))
@@ -1503,6 +1522,7 @@ bool FGroomBuilder::BuildHairDescriptionGroups(const FHairDescription& HairDescr
 	TVertexAttributesConstRef<float> VertexAO				= HairDescription.VertexAttributes().GetAttributesRef<float>(HairAttribute::Vertex::AO);
 	TStrandAttributesConstRef<int> StrandNumVertices		= HairDescription.StrandAttributes().GetAttributesRef<int>(HairAttribute::Strand::VertexCount);
 	TStrandAttributesConstRef<int> ClumpIDs					= HairDescription.StrandAttributes().GetAttributesRef<int>(HairAttribute::Strand::ClumpID);
+	TStrandAttributesConstRef<FVector3f> ClumpID3s			= HairDescription.StrandAttributes().GetAttributesRef<FVector3f>(HairAttribute::Strand::ClumpID);
 
 	if (!VertexPositions.IsValid() || !StrandNumVertices.IsValid())
 	{
@@ -1513,7 +1533,7 @@ bool FGroomBuilder::BuildHairDescriptionGroups(const FHairDescription& HairDescr
 	const bool bHasBaseColorAttribute = VertexBaseColor.IsValid();
 	const bool bHasRoughnessAttribute = VertexRoughness.IsValid();
 	const bool bHasAOAttribute = VertexAO.IsValid();
-	const bool bHasClumpIDs = ClumpIDs.IsValid();
+	const bool bHasClumpIDs = ClumpIDs.IsValid() || ClumpID3s.IsValid();
 
 	// Sanity check
 	check(bHasBaseColorAttribute == HairDescription.HasAttribute(EHairAttribute::Color));
@@ -1540,8 +1560,8 @@ bool FGroomBuilder::BuildHairDescriptionGroups(const FHairDescription& HairDescr
 	TStrandAttributesConstRef<int> ClosestGuide		  = HairDescription.StrandAttributes().GetAttributesRef<int>(HairAttribute::Strand::ClosestGuides);
 	TStrandAttributesConstRef<float> GuideWeight	  = HairDescription.StrandAttributes().GetAttributesRef<float>(HairAttribute::Strand::GuideWeights);
 	// Triplet
-	TStrandAttributesConstRef<FVector3f> ClosestGuides  = HairDescription.StrandAttributes().GetAttributesRef<FVector3f>(HairAttribute::Strand::ClosestGuides);
-	TStrandAttributesConstRef<FVector3f> GuideWeights	  = HairDescription.StrandAttributes().GetAttributesRef<FVector3f>(HairAttribute::Strand::GuideWeights);
+	TStrandAttributesConstRef<FVector3f> ClosestGuides= HairDescription.StrandAttributes().GetAttributesRef<FVector3f>(HairAttribute::Strand::ClosestGuides);
+	TStrandAttributesConstRef<FVector3f> GuideWeights = HairDescription.StrandAttributes().GetAttributesRef<FVector3f>(HairAttribute::Strand::GuideWeights);
 
 	// To use ClosestGuides and GuideWeights attributes, guides must be imported from HairDescription and
 	// must include StrandID attribute since ClosestGuides references those IDs
@@ -1684,11 +1704,15 @@ bool FGroomBuilder::BuildHairDescriptionGroups(const FHairDescription& HairDescr
 
 		if (bHasClumpIDs)
 		{
-			CurrentHairStrandsDatas->StrandsCurves.ClumpIDs.Add(ClumpIDs[StrandID]);
-			if (false)
+			if (ClumpID3s.IsValid())
 			{
+				CurrentHairStrandsDatas->StrandsCurves.ClumpIDs.Add(FIntVector(ClumpID3s[StrandID].X, ClumpID3s[StrandID].Y, ClumpID3s[StrandID].Z));
 				SetHairAttributeFlags(CurrentHairStrandsDatas->StrandsCurves.AttributeFlags, EHairAttributeFlags::HasMultipleClumpIDs);
 			}
+			else
+			{
+				CurrentHairStrandsDatas->StrandsCurves.ClumpIDs.Add(FIntVector(ClumpIDs[StrandID]));
+			}
 		}
 
 		if (bHasUVData)
diff --git a/Engine/Plugins/Runtime/HairStrands/Source/HairStrandsCore/Private/GroomComponent.cpp b/Engine/Plugins/Runtime/HairStrands/Source/HairStrandsCore/Private/GroomComponent.cpp
--- a/Engine/Plugins/Runtime/HairStrands/Source/HairStrandsCore/Private/GroomComponent.cpp
+++ b/Engine/Plugins/Runtime/HairStrands/Source/HairStrandsCore/Private/GroomComponent.cpp
@@ -68,6 +68,9 @@ static FAutoConsoleVariableRef CVarHairStrands_BoundsMode(TEXT("r.HairStrands.Bo
 static int32 GHairStrands_UseAttachedSimulationComponents = 0;
 static FAutoConsoleVariableRef CVarHairStrands_UseAttachedSimulationComponents(TEXT("r.HairStrands.UseAttachedSimulationComponents"), GHairStrands_UseAttachedSimulationComponents, TEXT("Boolean to check if we are using already attached niagara components for simulation (WIP)"));
 
+static int32 GHairStrands_ViewModeClumpIndex = 0;
+static FAutoConsoleVariableRef CVarHairStrands_ViewModeClumpIndex(TEXT("r.HairStrands.ViewMode.ClumpIndex"), GHairStrands_ViewModeClumpIndex, TEXT("Define the ClumpID index (0, 1, or 2) which should be visualized"));
+
 #define LOCTEXT_NAMESPACE "GroomComponent"
 
 ////////////////////////////////////////////////////////////////////////////////////////////////////////////////////////
@@ -705,6 +708,12 @@ public:
 							HairMaxRadius = Instances[GroupIt]->HairGroupPublicData ? Instances[GroupIt]->HairGroupPublicData->LODIndex : 0;
 						}
 
+						// Reuse the HairMaxRadius field to send the clumpID index selection
+						if (ViewMode == EGroomViewMode::ClumpID)
+						{
+							HairMaxRadius = FMath::Clamp(GHairStrands_ViewModeClumpIndex, 0, 2);
+						}
+
 						FVector HairColor = FVector::ZeroVector;
 						if (ViewMode == EGroomViewMode::Group)
 						{
diff --git a/Engine/Plugins/Runtime/HairStrands/Source/HairStrandsCore/Private/GroomEdit.cpp b/Engine/Plugins/Runtime/HairStrands/Source/HairStrandsCore/Private/GroomEdit.cpp
--- a/Engine/Plugins/Runtime/HairStrands/Source/HairStrandsCore/Private/GroomEdit.cpp
+++ b/Engine/Plugins/Runtime/HairStrands/Source/HairStrandsCore/Private/GroomEdit.cpp
@@ -47,7 +47,7 @@ static void HairStrandsDataToEditableHairStrands(const FHairStrandsDatas& In, TA
 		const FVector3f GuideWeights = bHasValidClosestGuide ? (FVector3f)In.StrandsCurves.CurvesClosestGuideIDs[StrandsIt] : FVector3f::ZeroVector;
 
 		Out[StrandsIt].StrandID		= bHasValidStrandsIDs ? In.StrandsCurves.StrandIDs[StrandsIt] : StrandsIt;
-		Out[StrandsIt].ClumpID		= bHasValidClumpIDs ? In.StrandsCurves.ClumpIDs[StrandsIt] : 0;
+		Out[StrandsIt].ClumpID		= bHasValidClumpIDs ? In.StrandsCurves.ClumpIDs[StrandsIt].X : 0;
 		Out[StrandsIt].RootUV		= In.StrandsCurves.CurvesRootUV[StrandsIt];
 		Out[StrandsIt].GuideIDs[0]  = GuideIDs.X;
 		Out[StrandsIt].GuideIDs[1]  = GuideIDs.Y;
diff --git a/Engine/Plugins/Runtime/HairStrands/Source/HairStrandsCore/Private/HairDescription.cpp b/Engine/Plugins/Runtime/HairStrands/Source/HairStrandsCore/Private/HairDescription.cpp
--- a/Engine/Plugins/Runtime/HairStrands/Source/HairStrandsCore/Private/HairDescription.cpp
+++ b/Engine/Plugins/Runtime/HairStrands/Source/HairStrandsCore/Private/HairDescription.cpp
@@ -74,7 +74,7 @@ bool FHairDescription::HasAttribute(EHairAttribute InAttribute) const
 	switch (InAttribute)
 	{
 	case EHairAttribute::RootUV:		return StrandAttributes().GetAttributesRef<FVector2f>(HairAttribute::Strand::RootUV).IsValid();
-	case EHairAttribute::ClumpID:		return StrandAttributes().GetAttributesRef<int>(HairAttribute::Strand::ClumpID).IsValid();
+	case EHairAttribute::ClumpID:		return StrandAttributes().GetAttributesRef<int>(HairAttribute::Strand::ClumpID).IsValid() || StrandAttributes().GetAttributesRef<FVector3f>(HairAttribute::Strand::ClumpID).IsValid();
 	case EHairAttribute::StrandID:		return StrandAttributes().GetAttributesRef<int>(HairAttribute::Strand::ID).IsValid();
 	case EHairAttribute::Color:			return VertexAttributes().GetAttributesRef<FVector3f>(HairAttribute::Vertex::Color).IsValid();
 	case EHairAttribute::Roughness:		return VertexAttributes().GetAttributesRef<float>(HairAttribute::Vertex::Roughness).IsValid();
diff --git a/Engine/Plugins/Runtime/HairStrands/Source/HairStrandsCore/Public/HairStrandsDatas.h b/Engine/Plugins/Runtime/HairStrands/Source/HairStrandsCore/Public/HairStrandsDatas.h
--- a/Engine/Plugins/Runtime/HairStrands/Source/HairStrandsCore/Public/HairStrandsDatas.h
+++ b/Engine/Plugins/Runtime/HairStrands/Source/HairStrandsCore/Public/HairStrandsDatas.h
@@ -404,8 +404,8 @@ struct HAIRSTRANDSCORE_API FHairStrandsCurves
 	/** Strand ID associated with each curve (optional) */
 	TArray<int> StrandIDs;
 
-	/** Clump ID associated with each curve (optional) */
-	TArray<uint16> ClumpIDs;
+	/** Clump ID associated with each curve (optional) */	
+	TArray<FIntVector> ClumpIDs;
 
 	/** Mapping of imported Groom ID to index */
 	TMap<int, int> GroomIDToIndex;
diff --git a/Engine/Shaders/Private/HairStrands/HairCardsAttributeCommon.ush b/Engine/Shaders/Private/HairStrands/HairCardsAttributeCommon.ush
--- a/Engine/Shaders/Private/HairStrands/HairCardsAttributeCommon.ush
+++ b/Engine/Shaders/Private/HairStrands/HairCardsAttributeCommon.ush
@@ -111,7 +111,7 @@ float  GetHairStrandsSeed(float2 HairPrimitiveUV)
 	return GetHairStrandsAttributes(HairPrimitiveUV).HairSeed;
 }
 
-uint GetHairStrandsClumpID(float2 HairPrimitiveUV)
+uint3 GetHairStrandsClumpID(float2 HairPrimitiveUV)
 {
 	return 0;
 }
diff --git a/Engine/Shaders/Private/HairStrands/HairStrandsAttributeCommon.ush b/Engine/Shaders/Private/HairStrands/HairStrandsAttributeCommon.ush
--- a/Engine/Shaders/Private/HairStrands/HairStrandsAttributeCommon.ush
+++ b/Engine/Shaders/Private/HairStrands/HairStrandsAttributeCommon.ush
@@ -10,6 +10,7 @@
 #define HAIR_ATTRIBUTE_STRIDE_SEED			1
 #define HAIR_ATTRIBUTE_STRIDE_LENGTH		2
 #define HAIR_ATTRIBUTE_STRIDE_CLUMPID		2
+#define HAIR_ATTRIBUTE_STRIDE_CLUMPID3		8
 #define HAIR_ATTRIBUTE_STRIDE_BASECOLOR		4
 #define HAIR_ATTRIBUTE_STRIDE_ROUGHNESS		1
 #define HAIR_ATTRIBUTE_STRIDE_AO			1
@@ -18,10 +19,11 @@
 #define HAIR_ATTRIBUTE_OFFSET_ROOTUV(InOffsets)		(InOffsets[0].x)
 #define HAIR_ATTRIBUTE_OFFSET_SEED(InOffsets)		(InOffsets[0].y)
 #define HAIR_ATTRIBUTE_OFFSET_LENGTH(InOffsets)		(InOffsets[0].z)
-#define HAIR_ATTRIBUTE_OFFSET_CLUMPID(InOffsets)	(InOffsets[0].w)
+#define HAIR_ATTRIBUTE_OFFSET_AO(InOffsets)			(InOffsets[0].w)
 #define HAIR_ATTRIBUTE_OFFSET_BASECOLOR(InOffsets)	(InOffsets[1].x)
 #define HAIR_ATTRIBUTE_OFFSET_ROUGHNESS(InOffsets)	(InOffsets[1].y)
-#define HAIR_ATTRIBUTE_OFFSET_AO(InOffsets)			(InOffsets[1].z)
+#define HAIR_ATTRIBUTE_OFFSET_CLUMPID(InOffsets)	(InOffsets[1].z)
+#define HAIR_ATTRIBUTE_OFFSET_CLUMPID3(InOffsets)	(InOffsets[1].w)
 
 ////////////////////////////////////////////////////////////////////////////////
 // Attribute access functions
@@ -37,12 +39,31 @@ uint InternalGetHairStrandsAttributes(uint Index, uint AttributeOffsetInBytes, u
 	return Packed >> (8u * ElementIndex);
 }
 
+uint2 InternalGetHairStrandsAttributes64(uint Index, uint AttributeOffsetInBytes, uint AttributeStrideInBytes, ByteAddressBuffer InAttributeBuffer)
+{
+	// Ensure the reading address is 4-byte aligned
+	const uint Address = AttributeOffsetInBytes + Index * AttributeStrideInBytes;
+	const uint Address_4BytesAligned = Address & (~0x3);
+
+	// For now, do two separate load to as we don't ensure the offset is 2byte-aligned
+	uint2 OutPacked = 0;
+	OutPacked.x = InAttributeBuffer.Load(Address_4BytesAligned);
+	OutPacked.y = InAttributeBuffer.Load(Address_4BytesAligned + 4u);
+	return OutPacked;
+}
+
 uint InternalGetHairStrandsAttributesPerCurve(uint HairPrimitiveId, uint AttributeOffsetInBytes, uint AttributeStrideInBytes, ByteAddressBuffer InAttributeBuffer, Buffer<uint> InPointToCurveBuffer)
 {
 	const uint CurveIndex = InPointToCurveBuffer[HairPrimitiveId];
 	return InternalGetHairStrandsAttributes(CurveIndex, AttributeOffsetInBytes, AttributeStrideInBytes, InAttributeBuffer);
 }
 
+uint2 InternalGetHairStrandsAttributesPerCurve64(uint HairPrimitiveId, uint AttributeOffsetInBytes, uint AttributeStrideInBytes, ByteAddressBuffer InAttributeBuffer, Buffer<uint> InPointToCurveBuffer)
+{
+	const uint CurveIndex = InPointToCurveBuffer[HairPrimitiveId];
+	return InternalGetHairStrandsAttributes64(CurveIndex, AttributeOffsetInBytes, AttributeStrideInBytes, InAttributeBuffer);
+}
+
 uint InternalGetHairStrandsAttributesPerVertex(uint HairPrimitiveId, uint AttributeOffsetInBytes, uint AttributeStrideInBytes, ByteAddressBuffer InAttributeBuffer)
 {
 	return InternalGetHairStrandsAttributes(HairPrimitiveId, AttributeOffsetInBytes, AttributeStrideInBytes, InAttributeBuffer);
@@ -60,7 +81,7 @@ float InternalGetHairStrandsWorldRadius(Buffer<uint4> InPositionBuffer, uint Ver
 
 bool IsAttributeValid(uint InOffset)
 {
-	return InOffset != HAIR_ATRIBUTE_INVALID_OFFSET;
+	return InOffset != HAIR_ATTRIBUTE_INVALID_OFFSET;
 }
 
 ////////////////////////////////////////////////////////////////////////////////
@@ -125,13 +146,18 @@ float GetHairStrandsSeed(uint HairPrimitiveId, float2 HairPrimitiveUV)
 	return UnpackHairSeed(InternalGetHairStrandsAttributesPerCurve(HairPrimitiveId, HAIR_ATTRIBUTE_OFFSET_SEED(HairStrandsVF.AttributeOffsets), HAIR_ATTRIBUTE_STRIDE_SEED, HairStrandsVF.AttributeBuffer, HairStrandsVF.PointToCurveBuffer));
 }
 
-uint GetHairStrandsClumpID(uint HairPrimitiveId, float2 HairPrimitiveUV)
+uint3 GetHairStrandsClumpID(uint HairPrimitiveId, float2 HairPrimitiveUV)
 {
-	const uint ClumpIdOffset = HAIR_ATTRIBUTE_OFFSET_CLUMPID(HairStrandsVF.AttributeOffsets);
+	const uint ClumpIdOffset  = HAIR_ATTRIBUTE_OFFSET_CLUMPID(HairStrandsVF.AttributeOffsets);
+	const uint ClumpId3Offset = HAIR_ATTRIBUTE_OFFSET_CLUMPID3(HairStrandsVF.AttributeOffsets);
 	if (IsAttributeValid(ClumpIdOffset))
 	{
 		return UnpackHairClumpID(InternalGetHairStrandsAttributesPerCurve(HairPrimitiveId, ClumpIdOffset, HAIR_ATTRIBUTE_STRIDE_CLUMPID, HairStrandsVF.AttributeBuffer, HairStrandsVF.PointToCurveBuffer));
 	}
+	else if (IsAttributeValid(ClumpId3Offset))
+	{
+		return UnpackHairClumpID(InternalGetHairStrandsAttributesPerCurve64(HairPrimitiveId, ClumpId3Offset, HAIR_ATTRIBUTE_STRIDE_CLUMPID3, HairStrandsVF.AttributeBuffer, HairStrandsVF.PointToCurveBuffer));
+	}
 	else
 	{
 		return 0;
diff --git a/Engine/Shaders/Private/HairStrands/HairStrandsVertexFactoryCommon.ush b/Engine/Shaders/Private/HairStrands/HairStrandsVertexFactoryCommon.ush
--- a/Engine/Shaders/Private/HairStrands/HairStrandsVertexFactoryCommon.ush
+++ b/Engine/Shaders/Private/HairStrands/HairStrandsVertexFactoryCommon.ush
@@ -94,9 +94,14 @@ float UnpackHairSeed(uint In)
 	return (In & 0xFF) / 255.f;
 }
 
-uint UnpackHairClumpID(uint In)
+uint3 UnpackHairClumpID(uint In)
 {
-	return (In & 0xFFFF);
+	return (In & 0xFFFF).xxx;
+}
+
+uint3 UnpackHairClumpID(uint2 In)
+{
+	return uint3(In.x & 0xFFFF, (In.x>>16) & 0xFFFF, (In.x>>32) & 0xFFFF);
 }
 
 float3 UnpackHairColor(uint In)
diff --git a/Engine/Shaders/Private/MaterialTemplate.ush b/Engine/Shaders/Private/MaterialTemplate.ush
--- a/Engine/Shaders/Private/MaterialTemplate.ush
+++ b/Engine/Shaders/Private/MaterialTemplate.ush
@@ -1522,7 +1522,7 @@ float MaterialExpressionGetHairSeed(FMaterialPixelParameters Parameters)
 #endif
 }
 
-float MaterialExpressionGetHairClumpID(FMaterialPixelParameters Parameters)
+float3 MaterialExpressionGetHairClumpID(FMaterialPixelParameters Parameters)
 {
 #if HAIR_STRAND_MESH_FACTORY
 	return GetHairStrandsClumpID(Parameters.HairPrimitiveId, Parameters.HairPrimitiveUV);
diff --git a/Engine/Shaders/Shared/HairStrandsDefinitions.h b/Engine/Shaders/Shared/HairStrandsDefinitions.h
--- a/Engine/Shaders/Shared/HairStrandsDefinitions.h
+++ b/Engine/Shaders/Shared/HairStrandsDefinitions.h
@@ -14,18 +14,19 @@
 #define HAIR_ATTRIBUTE_ROOTUV 0
 #define HAIR_ATTRIBUTE_SEED 1
 #define HAIR_ATTRIBUTE_LENGTH 2
-#define HAIR_ATTRIBUTE_CLUMPID 3
-#define HAIR_ATTRIBUTE_BASECOLOR 4
-#define HAIR_ATTRIBUTE_ROUGHNESS 5
-#define HAIR_ATTRIBUTE_AO 6
-#define HAIR_ATTRIBUTE_COUNT 7
+#define HAIR_ATTRIBUTE_BASECOLOR 3
+#define HAIR_ATTRIBUTE_ROUGHNESS 4
+#define HAIR_ATTRIBUTE_AO 5
+#define HAIR_ATTRIBUTE_CLUMPID 6
+#define HAIR_ATTRIBUTE_CLUMPID3 7
+#define HAIR_ATTRIBUTE_COUNT 8
 
 // Groom limits (based on encoding)
 #define HAIR_MAX_NUM_POINT_PER_CURVE ((1u<<8)-1u)
 #define HAIR_MAX_NUM_POINT_PER_GROUP ((1u<<24)-1u)
 #define HAIR_MAX_NUM_CURVE_PER_GROUP ((1u<<22)-1u)
 
-#define HAIR_ATRIBUTE_INVALID_OFFSET 0xFFFFFFFF
+#define HAIR_ATTRIBUTE_INVALID_OFFSET 0xFFFFFFFF
 
 // Max number of discrete LOD that a hair group can have
 #define MAX_HAIR_LOD 8
diff --git a/Engine/Source/Runtime/Engine/Private/Materials/HLSLMaterialTranslator.cpp b/Engine/Source/Runtime/Engine/Private/Materials/HLSLMaterialTranslator.cpp
--- a/Engine/Source/Runtime/Engine/Private/Materials/HLSLMaterialTranslator.cpp
+++ b/Engine/Source/Runtime/Engine/Private/Materials/HLSLMaterialTranslator.cpp
@@ -9677,7 +9677,7 @@ int32 FHLSLMaterialTranslator::GetHairSeed()
 
 int32 FHLSLMaterialTranslator::GetHairClumpID()
 {
-	return AddCodeChunk(MCT_Float1, TEXT("MaterialExpressionGetHairClumpID(Parameters)"));
+	return AddCodeChunk(MCT_Float3, TEXT("MaterialExpressionGetHairClumpID(Parameters)"));
 }
 
 int32 FHLSLMaterialTranslator::GetHairTangent(bool bUseTangentSpace)
diff --git a/Engine/Source/Runtime/Engine/Private/Materials/MaterialExpressionHLSL.cpp b/Engine/Source/Runtime/Engine/Private/Materials/MaterialExpressionHLSL.cpp
--- a/Engine/Source/Runtime/Engine/Private/Materials/MaterialExpressionHLSL.cpp
+++ b/Engine/Source/Runtime/Engine/Private/Materials/MaterialExpressionHLSL.cpp
@@ -2319,7 +2319,7 @@ bool UMaterialExpressionHairAttributes::GenerateHLSLExpression(FMaterialHLSLGene
 	case 12: OutExpression = Generator.GetTree().NewExpression<FExpressionInlineCustomHLSL>(EValueType::Float2, TEXT("MaterialExpressionGetAtlasUVs(Parameters)")); break;
 	case 13: OutExpression = Generator.GetTree().NewExpression<FExpressionInlineCustomHLSL>(EValueType::Float1, TEXT("MaterialExpressionGetHairGroupIndex(Parameters)")); break;
 	case 14: OutExpression = Generator.GetTree().NewExpression<FExpressionInlineCustomHLSL>(EValueType::Float1, TEXT("MaterialExpressionGetHairAO(Parameters)")); break;
-	case 15: OutExpression = Generator.GetTree().NewExpression<FExpressionInlineCustomHLSL>(EValueType::Int1, TEXT("MaterialExpressionGetHairClumpID(Parameters)")); break;
+	case 15: OutExpression = Generator.GetTree().NewExpression<FExpressionInlineCustomHLSL>(EValueType::Int3, TEXT("MaterialExpressionGetHairClumpID(Parameters)")); break;
 	default: return Generator.Error(TEXT("Invalid output"));
 	}
 	return OutExpression != nullptr;
