Fix race from ALT by changing the HairStands verbosity from GT only

#rnx

diff --git a/Engine/Plugins/Runtime/HairStrands/Source/HairStrandsCore/Private/GroomAsset.cpp b/Engine/Plugins/Runtime/HairStrands/Source/HairStrandsCore/Private/GroomAsset.cpp
--- a/Engine/Plugins/Runtime/HairStrands/Source/HairStrandsCore/Private/GroomAsset.cpp
+++ b/Engine/Plugins/Runtime/HairStrands/Source/HairStrandsCore/Private/GroomAsset.cpp
@@ -71,10 +71,12 @@ static FAutoConsoleVariableRef CVarHairStrandsLoadAsset(TEXT("r.HairStrands.Load
 static int32 GEnableGroomAsyncLoad = 0;
 static FAutoConsoleVariableRef CVarGroomAsyncLoad(TEXT("r.HairStrands.AsyncLoad"), GEnableGroomAsyncLoad, TEXT("Allow groom asset to be loaded asynchronously in the editor"));
 
+void UpdateHairStrandsVerbosity(IConsoleVariable* InCVarVerbosity);
 static TAutoConsoleVariable<int32> GHairStrandsWarningLogVerbosity(
 	TEXT("r.HairStrands.Log"),
 	-1,
-	TEXT("Enable warning log report for groom related asset (0: no logging, 1: error only, 2: error & warning only, other: all logs). By default all logging are enabled (-1). Value needs to be set at startup time."));
+	TEXT("Enable warning log report for groom related asset (0: no logging, 1: error only, 2: error & warning only, other: all logs). By default all logging are enabled (-1). Value needs to be set at startup time."),
+	FConsoleVariableDelegate::CreateStatic(UpdateHairStrandsVerbosity));
 
 static int32 GHairStrandsDDCLogEnable = 0;
 static FAutoConsoleVariableRef CVarHairStrandsDDCLogEnable(TEXT("r.HairStrands.DDCLog"), GHairStrandsDDCLogEnable, TEXT("Enable DDC logging for groom assets and groom binding assets"));
@@ -384,9 +386,9 @@ void DumpLoadedGroomData(IConsoleVariable* InCVarPakTesterEnabled)
 
 /////////////////////////////////////////////////////////////////////////////////////////
 
-void UpdateHairStrandsLogVerbosity()
+void UpdateHairStrandsVerbosity(IConsoleVariable* InCVarVerbosity)
 {
-	const int32 Verbosity = GHairStrandsWarningLogVerbosity->GetInt();
+	const int32 Verbosity = InCVarVerbosity->GetInt();
 	switch (Verbosity)
 	{
 	case 0:  UE_SET_LOG_VERBOSITY(LogHairStrands, NoLogging); break;
@@ -779,8 +781,6 @@ UGroomAsset::UGroomAsset(const FObjectInitializer& ObjectInitializer)
 
 	MinLOD.Default = 0;
 	DisableBelowMinLodStripping.Default = false;
-
-	UpdateHairStrandsLogVerbosity();
 }
 
 bool UGroomAsset::HasGeometryType(uint32 GroupIndex, EGroomGeometryType Type) const
