Shared hair instance shader parameter between various part of the hair code.

diff --git a/Engine/Plugins/Runtime/HairStrands/Source/HairStrandsCore/Private/Deformer/DeformerDataInterfaceGroom.cpp b/Engine/Plugins/Runtime/HairStrands/Source/HairStrandsCore/Private/Deformer/DeformerDataInterfaceGroom.cpp
--- a/Engine/Plugins/Runtime/HairStrands/Source/HairStrandsCore/Private/Deformer/DeformerDataInterfaceGroom.cpp
+++ b/Engine/Plugins/Runtime/HairStrands/Source/HairStrandsCore/Private/Deformer/DeformerDataInterfaceGroom.cpp
@@ -216,12 +216,13 @@ void FOptimusGroomDataProviderProxy::GatherDispatchData(FDispatchData const& InD
 			FParameters& Parameters = ParameterArray[InvocationIndex];
 			Parameters.NumControlPoints = NumControlPoints;
 			Parameters.NumCurves = NumCurves;
-			Parameters.VF_HairRadius = VFInput.Strands.HairRadius;
-			Parameters.VF_HairLength = VFInput.Strands.HairLength;
-			Parameters.VF_RootScale = VFInput.Strands.HairRootScale;
-			Parameters.VF_TipScale = VFInput.Strands.HairTipScale;
-			Parameters.VF_GroupIndex = Instance->Debug.GroupIndex;
-			PACK_HAIR_ATTRIBUTE_OFFSETS(Parameters.VF_AttributeOffsets, VFInput.Strands.AttributeOffsets);
+			Parameters.VF_HairRadius = VFInput.Strands.Common.Radius;
+			Parameters.VF_HairLength = VFInput.Strands.Common.Length;
+			Parameters.VF_RootScale  = VFInput.Strands.Common.RootScale;
+			Parameters.VF_TipScale   = VFInput.Strands.Common.TipScale;
+			Parameters.VF_GroupIndex = VFInput.Strands.Common.GroupIndex;// Instance->Debug.GroupIndex;
+			Parameters.VF_AttributeOffsets = VFInput.Strands.Common.AttributeOffsets;
+			//PACK_HAIR_ATTRIBUTE_OFFSETS(Parameters.VF_AttributeOffsets, VFInput.Strands.AttributeOffsets);
 
 			if (bIsSRVValid)
 			{
diff --git a/Engine/Plugins/Runtime/HairStrands/Source/HairStrandsCore/Private/GroomComponent.cpp b/Engine/Plugins/Runtime/HairStrands/Source/HairStrandsCore/Private/GroomComponent.cpp
--- a/Engine/Plugins/Runtime/HairStrands/Source/HairStrandsCore/Private/GroomComponent.cpp
+++ b/Engine/Plugins/Runtime/HairStrands/Source/HairStrandsCore/Private/GroomComponent.cpp
@@ -822,7 +822,7 @@ public:
 			VertexFactory = (FVertexFactory*)Instance->Strands.VertexFactory;
 			HairVertexCount = Instance->HairGroupPublicData->GetActiveStrandsPointCount();
 			MaxVertexIndex = HairVertexCount * HAIR_POINT_TO_VERTEX;
-			bUseCulling = Instance->Strands.bIsCullingEnabled;
+			bUseCulling = Instance->Strands.bCullingEnable;
 			NumPrimitive = bUseCulling ? 0 : HairVertexCount * HAIR_POINT_TO_TRIANGLE;
 			if (MaterialRenderProxy == nullptr)
 			{
diff --git a/Engine/Plugins/Runtime/HairStrands/Source/HairStrandsCore/Private/GroomDebug.cpp b/Engine/Plugins/Runtime/HairStrands/Source/HairStrandsCore/Private/GroomDebug.cpp
--- a/Engine/Plugins/Runtime/HairStrands/Source/HairStrandsCore/Private/GroomDebug.cpp
+++ b/Engine/Plugins/Runtime/HairStrands/Source/HairStrandsCore/Private/GroomDebug.cpp
@@ -597,13 +597,13 @@ static void AddDrawDebugStrandsCVsPass(
 	FDrawDebugStrandsCVsCS::FParameters* Parameters = GraphBuilder.AllocParameters<FDrawDebugStrandsCVsCS::FParameters>();
 	Parameters->ViewUniformBuffer = View.ViewUniformBuffer;
 	Parameters->HairStrandsVF = Instance->Strands.UniformBuffer;
-	Parameters->LocalToWorld = FMatrix44f(Instance->LocalToWorld.ToMatrixWithScale());		// LWC_TODO: Precision loss
+	Parameters->LocalToWorld = FMatrix44f(Instance->LocalToWorld.ToMatrixWithScale());		// LWC_TODO: Precision loss // TODO change this to Uniform buffer parameters..
 	Parameters->MaxVertexCount = Instance->Strands.Data->GetNumPoints();
 	Parameters->ColorTexture = GraphBuilder.CreateUAV(ColorTexture);
 	Parameters->DepthTexture = DepthTexture;
 	Parameters->LinearSampler = TStaticSamplerState<SF_Bilinear, AM_Clamp, AM_Clamp, AM_Clamp>::GetRHI();
 
-	const uint32 PointCount = Instance->HairGroupPublicData->VFInput.Strands.PointCount;
+	const uint32 PointCount = Instance->HairGroupPublicData->VFInput.Strands.Common.PointCount;
 	FComputeShaderUtils::AddPass(
 		GraphBuilder, 
 		RDG_EVENT_NAME("HairStrands::DrawCVs"), 
@@ -954,14 +954,14 @@ static void AddHairDebugPrintInstancePass(
 
 				{
 					Data2 = FUintVector4(0);
-					Data2.X |= InstanceIndex < InstanceCountPerType[HairInstanceCount_StrandsPrimaryView] ? 0x1u  : 0u;
-					Data2.X |= InstanceIndex < InstanceCountPerType[HairInstanceCount_StrandsShadowView]  ? 0x2u  : 0u;
-					Data2.X |= Instance->HairGroupPublicData->VFInput.Strands.bScatterSceneLighting       ? 0x4u  : 0u;
-					Data2.X |= Instance->HairGroupPublicData->VFInput.Strands.bUseRaytracingGeometry      ? 0x8u  : 0u;
-					Data2.X |= Instance->HairGroupPublicData->VFInput.Strands.bUseStableRasterization     ? 0x10u : 0u;
-					Data2.X |= Instance->HairGroupPublicData->bSupportVoxelization                        ? 0x20u : 0u;
-					Data2.X |= ActiveGroomCacheType == EGroomCacheType::Guides                            ? 0x40u : 0u;
-					Data2.X |= ActiveGroomCacheType == EGroomCacheType::Strands                           ? 0x80u : 0u;
+					Data2.X |= InstanceIndex < InstanceCountPerType[HairInstanceCount_StrandsPrimaryView]  ? 0x1u  : 0u;
+					Data2.X |= InstanceIndex < InstanceCountPerType[HairInstanceCount_StrandsShadowView]   ? 0x2u  : 0u;
+					Data2.X |= Instance->HairGroupPublicData->VFInput.Strands.Common.bScatterSceneLighting ? 0x4u  : 0u;
+					Data2.X |= Instance->HairGroupPublicData->VFInput.Strands.Common.bRaytracingGeometry   ? 0x8u  : 0u;
+					Data2.X |= Instance->HairGroupPublicData->VFInput.Strands.Common.bStableRasterization  ? 0x10u : 0u;
+					Data2.X |= Instance->HairGroupPublicData->bSupportVoxelization                         ? 0x20u : 0u;
+					Data2.X |= ActiveGroomCacheType == EGroomCacheType::Guides                             ? 0x40u : 0u;
+					Data2.X |= ActiveGroomCacheType == EGroomCacheType::Strands                            ? 0x80u : 0u;
 					Data2.X |= uint32(FFloat16(Instance->HairGroupPublicData->ContinuousLODScreenSize).Encoded) << 16u;
 
 					Data2.Y = Instance->HairGroupPublicData->GetActiveStrandsPointCount();
@@ -979,14 +979,14 @@ static void AddHairDebugPrintInstancePass(
 				}
 				
 				Data3 = FUintVector4(0);
-				Data3.X |= FFloat16(Instance->HairGroupPublicData->VFInput.Strands.HairRadius).Encoded;
-				Data3.X |= FFloat16(Instance->HairGroupPublicData->VFInput.Strands.HairDensity).Encoded << 16u;
+				Data3.X |= FFloat16(Instance->HairGroupPublicData->VFInput.Strands.Common.Radius).Encoded;
+				Data3.X |= FFloat16(Instance->HairGroupPublicData->VFInput.Strands.Common.Density).Encoded << 16u;
 
-				Data3.Y |= FFloat16(Instance->HairGroupPublicData->VFInput.Strands.HairRootScale).Encoded;
-				Data3.Y |= FFloat16(Instance->HairGroupPublicData->VFInput.Strands.HairTipScale).Encoded << 16u;
+				Data3.Y |= FFloat16(Instance->HairGroupPublicData->VFInput.Strands.Common.RootScale).Encoded;
+				Data3.Y |= FFloat16(Instance->HairGroupPublicData->VFInput.Strands.Common.TipScale).Encoded << 16u;
 
-				Data3.Z |= FFloat16(Instance->HairGroupPublicData->VFInput.Strands.HairLength).Encoded;
-				Data3.Z |= FFloat16(Instance->HairGroupPublicData->VFInput.Strands.HairLengthScale).Encoded << 16u;
+				Data3.Z |= FFloat16(Instance->HairGroupPublicData->VFInput.Strands.Common.Length).Encoded;
+				Data3.Z |= FFloat16(Instance->HairGroupPublicData->VFInput.Strands.Common.LengthScale).Encoded << 16u;
 			}
 			break;
 		case EHairGeometryType::Cards:
diff --git a/Engine/Plugins/Runtime/HairStrands/Source/HairStrandsCore/Private/GroomManager.cpp b/Engine/Plugins/Runtime/HairStrands/Source/HairStrandsCore/Private/GroomManager.cpp
--- a/Engine/Plugins/Runtime/HairStrands/Source/HairStrandsCore/Private/GroomManager.cpp
+++ b/Engine/Plugins/Runtime/HairStrands/Source/HairStrandsCore/Private/GroomManager.cpp
@@ -501,7 +501,7 @@ static void RunHairStrandsGatherCluster(
 		FHairGroupInstance* Instance = static_cast<FHairGroupInstance*>(AbstractInstance);
 		if (Instance->GeometryType == EHairGeometryType::Strands && Instance->Strands.IsValid())
 		{
-			if (Instance->Strands.bIsCullingEnabled)
+			if (Instance->Strands.bCullingEnable)
 			{
 				RegisterClusterData(Instance, ClusterData);
 			}
@@ -1014,7 +1014,7 @@ static void RunHairLODSelection(
 
 				// Early initialization, so that when filtering MeshBatch block in SceneVisiblity, we can use this value to know 
 				// if the hair instance is visible or not (i.e., HairLengthScale > 0)
-				Instance->HairGroupPublicData->VFInput.Strands.HairLengthScale = Instance->Strands.Modifier.HairLengthScale;
+				Instance->HairGroupPublicData->VFInput.Strands.Common.LengthScale = Instance->Strands.Modifier.HairLengthScale;
 			}
 
 			// Only switch LOD if the data are ready to be used
@@ -1035,7 +1035,7 @@ static void RunHairLODSelection(
 				Instance->Guides.bHasGlobalInterpolation = Instance->HairGroupPublicData->IsGlobalInterpolationEnable(IntLODIndex);
 				Instance->Guides.bIsDeformationEnable = Instance->HairGroupPublicData->bIsDeformationEnable;
 				Instance->Guides.bIsSimulationCacheEnable = Instance->HairGroupPublicData->bIsSimulationCacheEnable;
-				Instance->Strands.bIsCullingEnabled = bCullingEnable;
+				Instance->Strands.bCullingEnable = bCullingEnable;
 			}
 
 			// Ensure that if the binding type is set to Skinning (or has RBF enabled), the skel. mesh is valid and support skin cache.
diff --git a/Engine/Plugins/Runtime/HairStrands/Source/HairStrandsCore/Private/HairStrandsRendering.cpp b/Engine/Plugins/Runtime/HairStrands/Source/HairStrandsCore/Private/HairStrandsRendering.cpp
--- a/Engine/Plugins/Runtime/HairStrands/Source/HairStrandsCore/Private/HairStrandsRendering.cpp
+++ b/Engine/Plugins/Runtime/HairStrands/Source/HairStrandsCore/Private/HairStrandsRendering.cpp
@@ -497,7 +497,7 @@ class FHairInterpolationCS : public FGlobalShader
 		SHADER_PARAMETER(float, HairLengthScale)
 		SHADER_PARAMETER(FVector3f, InRenHairPositionOffset)
 		SHADER_PARAMETER(FVector3f, InSimHairPositionOffset)
-		SHADER_PARAMETER(uint32,  HairStrandsVF_bIsCullingEnable)
+		SHADER_PARAMETER(uint32,  HairStrandsVFTODO_bCullingEnable)
 		SHADER_PARAMETER(FMatrix44f, LocalToWorldMatrix)
 
 		SHADER_PARAMETER_RDG_BUFFER_SRV(Buffer, RenRestPosePositionBuffer)
@@ -540,10 +540,10 @@ class FHairInterpolationCS : public FGlobalShader
 		SHADER_PARAMETER_RDG_BUFFER_SRV(Buffer<float4>, OutSimHairPositionOffsetBuffer)
 		SHADER_PARAMETER_RDG_BUFFER_SRV(Buffer<float4>, OutRenHairPositionOffsetBuffer)
 
-		SHADER_PARAMETER_RDG_BUFFER_SRV(Buffer<uint>, HairStrandsVF_CullingIndirectBuffer)
-		SHADER_PARAMETER_RDG_BUFFER_SRV(Buffer<uint>, HairStrandsVF_CullingIndexBuffer)
-		SHADER_PARAMETER_RDG_BUFFER_SRV(Buffer<float>, HairStrandsVF_CullingRadiusScaleBuffer)
-		RDG_BUFFER_ACCESS(HairStrandsVF_CullingIndirectBufferArgs, ERHIAccess::IndirectArgs)
+		SHADER_PARAMETER_RDG_BUFFER_SRV(Buffer<uint>, HairStrandsVFTODO_CullingIndirectBuffer)
+		SHADER_PARAMETER_RDG_BUFFER_SRV(Buffer<uint>, HairStrandsVFTODO_CullingIndexBuffer)
+		SHADER_PARAMETER_RDG_BUFFER_SRV(Buffer<float>, HairStrandsVFTODO_CullingRadiusScaleBuffer)
+		RDG_BUFFER_ACCESS(HairStrandsVFTODO_CullingIndirectBufferArgs, ERHIAccess::IndirectArgs)
 		
 		END_SHADER_PARAMETER_STRUCT()
 
@@ -696,7 +696,7 @@ static void AddHairStrandsInterpolationPass(
 
 	const bool bHasLocalDeformation = Instance->Guides.bIsSimulationEnable || bSupportGlobalInterpolation || Instance->Guides.bIsDeformationEnable || Instance->Guides.bIsSimulationCacheEnable;
 	const bool bCullingEnable = IsHairStrandContinuousDecimationReorderingEnabled() ? false : (InstanceGeometryType == EHairGeometryType::Strands && CullingData.bCullingResultAvailable); 	// TODO: improve reordering so that culling can be used effectively
-	Parameters->HairStrandsVF_bIsCullingEnable = bCullingEnable ? 1 : 0;
+	Parameters->HairStrandsVFTODO_bCullingEnable = bCullingEnable ? 1 : 0;
 
 	// Select dynamic geometry permutation, based on the Simulation/RBF/InterpolationType
 	int32 DynamicGeometryType = 0;
@@ -726,10 +726,10 @@ static void AddHairStrandsInterpolationPass(
 
 	if (bCullingEnable)
 	{
-		Parameters->HairStrandsVF_CullingIndirectBuffer = CullingData.HairStrandsVF_CullingIndirectBuffer.SRV;
-		Parameters->HairStrandsVF_CullingIndexBuffer = CullingData.HairStrandsVF_CullingIndexBuffer.SRV;
-		Parameters->HairStrandsVF_CullingRadiusScaleBuffer = CullingData.HairStrandsVF_CullingRadiusScaleBuffer.SRV;
-		Parameters->HairStrandsVF_CullingIndirectBufferArgs = CullingData.HairStrandsVF_CullingIndirectBuffer.Buffer;
+		Parameters->HairStrandsVFTODO_CullingIndirectBuffer = CullingData.HairStrandsVF_CullingIndirectBuffer.SRV;
+		Parameters->HairStrandsVFTODO_CullingIndexBuffer = CullingData.HairStrandsVF_CullingIndexBuffer.SRV;
+		Parameters->HairStrandsVFTODO_CullingRadiusScaleBuffer = CullingData.HairStrandsVF_CullingRadiusScaleBuffer.SRV;
+		Parameters->HairStrandsVFTODO_CullingIndirectBufferArgs = CullingData.HairStrandsVF_CullingIndirectBuffer.Buffer;
 
 		FComputeShaderUtils::AddPass(
 			GraphBuilder,
@@ -876,7 +876,7 @@ class FHairClusterAABBCS : public FGlobalShader
 		SHADER_PARAMETER_RDG_BUFFER_SRV(Buffer, ClusterIndexOffsetBuffer)
 		SHADER_PARAMETER_RDG_BUFFER_SRV(Buffer, ClusterIndexCountBuffer)
 		SHADER_PARAMETER_RDG_BUFFER_SRV(Buffer, CulledClusterCountBuffer)
-		SHADER_PARAMETER_RDG_BUFFER_SRV(Buffer<uint>, HairStrandsVF_CullingIndirectBuffer)
+		SHADER_PARAMETER_RDG_BUFFER_SRV(Buffer<uint>, HairStrandsVFTODO_CullingIndirectBuffer)
 		SHADER_PARAMETER_RDG_BUFFER_SRV(Buffer<uint>, IndirectBuffer)
 		SHADER_PARAMETER_RDG_BUFFER_UAV(RWBuffer, OutClusterAABBBuffer)
 		SHADER_PARAMETER_RDG_BUFFER_UAV(RWBuffer, OutGroupAABBBuffer)
@@ -968,7 +968,7 @@ static void AddHairClusterAABBPass(
 		// Currently disabled this sanity check as the RDG allocation can return a larger buffer (for reuse purpose)
 		//check(ClusterData->ClusterCount * 6 == ClusterAABBData.ClusterAABBBuffer.Buffer->Desc.NumElements);
 	}
-	Parameters->HairStrandsVF_CullingIndirectBuffer = DrawIndirectRasterComputeBuffer; // Used for checking max vertex count
+	Parameters->HairStrandsVFTODO_CullingIndirectBuffer = DrawIndirectRasterComputeBuffer; // Used for checking max vertex count
 	Parameters->OutClusterAABBBuffer = ClusterAABBData.ClusterAABBBuffer.UAV;
 	Parameters->OutGroupAABBBuffer = ClusterAABBData.GroupAABBBuffer.UAV;
 
@@ -1173,7 +1173,7 @@ class FHairTangentCS : public FGlobalShader
 
 	BEGIN_SHADER_PARAMETER_STRUCT(FParameters, )
 		SHADER_PARAMETER(uint32, VertexCount)
-		SHADER_PARAMETER(uint32, HairStrandsVF_bIsCullingEnable)
+		SHADER_PARAMETER(uint32, HairStrandsVF_bCullingEnable)
 		SHADER_PARAMETER_RDG_BUFFER_SRV(Buffer, PositionBuffer)
 		SHADER_PARAMETER_RDG_BUFFER_SRV(Buffer<uint>,	HairStrandsVF_CullingIndirectBuffer)
 		SHADER_PARAMETER_RDG_BUFFER_SRV(Buffer<uint>,	HairStrandsVF_CullingIndexBuffer)
@@ -1208,7 +1208,7 @@ void AddHairTangentPass(
 	Parameters->PositionBuffer = PositionBuffer;
 	Parameters->OutputTangentBuffer = OutTangentBuffer.UAV;
 	Parameters->VertexCount = VertexCount;
-	Parameters->HairStrandsVF_bIsCullingEnable = bCullingEnable ? 1 : 0;
+	Parameters->HairStrandsVF_bCullingEnable = bCullingEnable ? 1 : 0;
 
 	FHairTangentCS::FPermutationDomain PermutationVector;
 	PermutationVector.Set<FHairTangentCS::FCulling>(bCullingEnable);
@@ -1258,7 +1258,7 @@ class FHairRaytracingGeometryCS : public FGlobalShader
 		SHADER_PARAMETER(float, HairStrandsVF_HairRootScale)
 		SHADER_PARAMETER(float, HairStrandsVF_HairTipScale)
 
-		SHADER_PARAMETER(uint32, HairStrandsVF_bIsCullingEnable)
+		SHADER_PARAMETER(uint32, HairStrandsVF_bCullingEnable)
 		SHADER_PARAMETER_RDG_BUFFER_SRV(Buffer,	HairStrandsVF_CullingIndirectBuffer)
 		SHADER_PARAMETER_RDG_BUFFER_SRV(Buffer,	HairStrandsVF_CullingIndexBuffer)
 		SHADER_PARAMETER_RDG_BUFFER_SRV(Buffer,	HairStrandsVF_CullingRadiusScaleBuffer)
@@ -1330,7 +1330,7 @@ static void AddGenerateRaytracingGeometryPass(
 	Parameters->RaytracingProceduralSplits = ProceduralSplits;
 
 	const bool bCullingEnable = CullingData.bCullingResultAvailable;
-	Parameters->HairStrandsVF_bIsCullingEnable = bCullingEnable ? 1 : 0;
+	Parameters->HairStrandsVF_bCullingEnable = bCullingEnable ? 1 : 0;
 	if (bCullingEnable)
 	{
 		Parameters->HairStrandsVF_CullingIndirectBuffer = CullingData.HairStrandsVF_CullingIndirectBuffer.SRV;
@@ -1629,7 +1629,10 @@ static FHairGroupPublicData::FVertexFactoryInput InternalComputeHairStrandsVerte
 	FHairGroupPublicData::FVertexFactoryInput OutVFInput;
 	if (!Instance || Instance->GeometryType != EHairGeometryType::Strands)
 	{
-		OutVFInput.Strands.AttributeOffsets.Init(0xFFFFFFFF, HAIR_ATTRIBUTE_COUNT);
+		for (uint32 OffsetIt = 0; OffsetIt < HAIR_ATTRIBUTE_OFFSET_COUNT; ++OffsetIt)
+		{
+			OutVFInput.Strands.Common.AttributeOffsets[OffsetIt] = FUintVector4(HAIR_ATTRIBUTE_INVALID_OFFSET,HAIR_ATTRIBUTE_INVALID_OFFSET,HAIR_ATTRIBUTE_INVALID_OFFSET,HAIR_ATTRIBUTE_INVALID_OFFSET);
+		}
 		return OutVFInput;
 	}
 
@@ -1649,11 +1652,10 @@ static FHairGroupPublicData::FVertexFactoryInput InternalComputeHairStrandsVerte
 		CONVERT_HAIRSSTRANDS_VF_PARAMETERS(OutVFInput.Strands.PositionOffsetBuffer,		Instance->Guides.DeformedResource->GetPositionOffsetBuffer(FHairStrandsDeformedResource::EFrameType::Current));
 		CONVERT_HAIRSSTRANDS_VF_PARAMETERS(OutVFInput.Strands.PrevPositionOffsetBuffer,	Instance->Guides.DeformedResource->GetPositionOffsetBuffer(FHairStrandsDeformedResource::EFrameType::Previous));
 
-		OutVFInput.Strands.PositionOffset = Instance->Guides.DeformedResource->GetPositionOffset(FHairStrandsDeformedResource::EFrameType::Current);
-		OutVFInput.Strands.PrevPositionOffset = Instance->Guides.DeformedResource->GetPositionOffset(FHairStrandsDeformedResource::EFrameType::Previous);
-		OutVFInput.Strands.PointCount = Instance->Guides.RestResource->GetPointCount();
-		OutVFInput.Strands.CurveCount = Instance->Guides.RestResource->GetCurveCount();
-
+		OutVFInput.Strands.Common.PositionOffset = FVector3f(Instance->Guides.DeformedResource->GetPositionOffset(FHairStrandsDeformedResource::EFrameType::Current)); // LWC precision loss
+		OutVFInput.Strands.Common.PrevPositionOffset = FVector3f(Instance->Guides.DeformedResource->GetPositionOffset(FHairStrandsDeformedResource::EFrameType::Previous)); // LWC precision loss
+		OutVFInput.Strands.Common.PointCount = Instance->Guides.RestResource->GetPointCount();
+		OutVFInput.Strands.Common.CurveCount = Instance->Guides.RestResource->GetCurveCount();
 		bUseGuideAttributeOffsets = true;
 	}
 	// 2. Render Strands with deformation
@@ -1669,8 +1671,8 @@ static FHairGroupPublicData::FVertexFactoryInput InternalComputeHairStrandsVerte
 			CONVERT_HAIRSSTRANDS_VF_PARAMETERS(OutVFInput.Strands.PositionOffsetBuffer,		Instance->Strands.DeformedResource->GetPositionOffsetBuffer(FHairStrandsDeformedResource::EFrameType::Current));
 			CONVERT_HAIRSSTRANDS_VF_PARAMETERS(OutVFInput.Strands.PrevPositionOffsetBuffer, Instance->Strands.DeformedResource->GetPositionOffsetBuffer(FHairStrandsDeformedResource::EFrameType::Previous));
 
-			OutVFInput.Strands.PositionOffset	  = Instance->Strands.DeformedResource->GetPositionOffset(FHairStrandsDeformedResource::EFrameType::Current);
-			OutVFInput.Strands.PrevPositionOffset = Instance->Strands.DeformedResource->GetPositionOffset(FHairStrandsDeformedResource::EFrameType::Previous);
+			OutVFInput.Strands.Common.PositionOffset	 = FVector3f(Instance->Strands.DeformedResource->GetPositionOffset(FHairStrandsDeformedResource::EFrameType::Current)); // LWC precision loss
+			OutVFInput.Strands.Common.PrevPositionOffset = FVector3f(Instance->Strands.DeformedResource->GetPositionOffset(FHairStrandsDeformedResource::EFrameType::Previous)); // LWC precision loss
 		}
 		else
 		{
@@ -1679,16 +1681,16 @@ static FHairGroupPublicData::FVertexFactoryInput InternalComputeHairStrandsVerte
 			CONVERT_HAIRSSTRANDS_VF_PARAMETERS(OutVFInput.Strands.PositionOffsetBuffer,		Instance->Strands.DeformedResource->GetPositionOffsetBuffer(FHairStrandsDeformedResource::EFrameType::Current));
 			CONVERT_HAIRSSTRANDS_VF_PARAMETERS(OutVFInput.Strands.PrevPositionOffsetBuffer,	Instance->Strands.DeformedResource->GetPositionOffsetBuffer(FHairStrandsDeformedResource::EFrameType::Current));
 
-			OutVFInput.Strands.PositionOffset	  = Instance->Strands.DeformedResource->GetPositionOffset(FHairStrandsDeformedResource::EFrameType::Current);
-			OutVFInput.Strands.PrevPositionOffset = Instance->Strands.DeformedResource->GetPositionOffset(FHairStrandsDeformedResource::EFrameType::Current);
+			OutVFInput.Strands.Common.PositionOffset	 = FVector3f(Instance->Strands.DeformedResource->GetPositionOffset(FHairStrandsDeformedResource::EFrameType::Current)); // LWC precision loss
+			OutVFInput.Strands.Common.PrevPositionOffset = FVector3f(Instance->Strands.DeformedResource->GetPositionOffset(FHairStrandsDeformedResource::EFrameType::Current)); // LWC precision loss
 		}
 		CONVERT_HAIRSSTRANDS_VF_PARAMETERS(OutVFInput.Strands.TangentBuffer,			Instance->Strands.DeformedResource->TangentBuffer);
 		CONVERT_HAIRSSTRANDS_VF_PARAMETERS(OutVFInput.Strands.AttributeBuffer,			(ViewMode == EGroomViewMode::RenderHairStrands || ViewMode == EGroomViewMode::Cluster || ViewMode == EGroomViewMode::ClusterAABB) ? Instance->Strands.DebugAttributeBuffer : Instance->Strands.RestResource->AttributeBuffer);
 		CONVERT_HAIRSSTRANDS_VF_PARAMETERS(OutVFInput.Strands.PointToCurveBuffer,		Instance->Strands.RestResource->PointToCurveBuffer);
 		CONVERT_HAIRSSTRANDS_VF_PARAMETERS(OutVFInput.Strands.CurveBuffer,				Instance->Strands.RestResource->CurveBuffer);
 
-		OutVFInput.Strands.PointCount = Instance->HairGroupPublicData->GetActiveStrandsPointCount();
-		OutVFInput.Strands.CurveCount = Instance->HairGroupPublicData->GetActiveStrandsCurveCount();
+		OutVFInput.Strands.Common.PointCount = Instance->HairGroupPublicData->GetActiveStrandsPointCount();
+		OutVFInput.Strands.Common.CurveCount = Instance->HairGroupPublicData->GetActiveStrandsCurveCount();
 	}
 	// 3. Render Strands in rest position: used when there are no skinning (rigid binding or no binding), no simulation, and no RBF
 	else
@@ -1702,34 +1704,39 @@ static FHairGroupPublicData::FVertexFactoryInput InternalComputeHairStrandsVerte
 		CONVERT_HAIRSSTRANDS_VF_PARAMETERS(OutVFInput.Strands.PositionOffsetBuffer,		Instance->Strands.RestResource->PositionOffsetBuffer);
 		CONVERT_HAIRSSTRANDS_VF_PARAMETERS(OutVFInput.Strands.PrevPositionOffsetBuffer, Instance->Strands.RestResource->PositionOffsetBuffer);
 
-		OutVFInput.Strands.PositionOffset = Instance->Strands.RestResource->GetPositionOffset();
-		OutVFInput.Strands.PrevPositionOffset = Instance->Strands.RestResource->GetPositionOffset();
-		OutVFInput.Strands.PointCount = Instance->HairGroupPublicData->GetActiveStrandsPointCount();
-		OutVFInput.Strands.CurveCount = Instance->HairGroupPublicData->GetActiveStrandsCurveCount();
+		OutVFInput.Strands.Common.PositionOffset 	= FVector3f(Instance->Strands.RestResource->GetPositionOffset()); // LWC Precision loss
+		OutVFInput.Strands.Common.PrevPositionOffset= FVector3f(Instance->Strands.RestResource->GetPositionOffset()); // LWC Precision loss
+		OutVFInput.Strands.Common.PointCount 		= Instance->HairGroupPublicData->GetActiveStrandsPointCount();
+		OutVFInput.Strands.Common.CurveCount 		= Instance->HairGroupPublicData->GetActiveStrandsCurveCount();
 	}
 
-	OutVFInput.Strands.HairRadius = (GStrandHairWidth > 0 ? GStrandHairWidth : Instance->Strands.Modifier.HairWidth) * 0.5f;
-	OutVFInput.Strands.HairRootScale = Instance->Strands.Modifier.HairRootScale;
-	OutVFInput.Strands.HairTipScale = Instance->Strands.Modifier.HairTipScale;
-	OutVFInput.Strands.HairRaytracingRadiusScale = (GHairRaytracingRadiusScale > 0 ? GHairRaytracingRadiusScale : Instance->Strands.Modifier.HairRaytracingRadiusScale);
-	OutVFInput.Strands.HairLength = Instance->Strands.Modifier.HairLength;
-	OutVFInput.Strands.HairDensity = Instance->Strands.Modifier.HairShadowDensity;
-	OutVFInput.Strands.bScatterSceneLighting = Instance->Strands.Modifier.bScatterSceneLighting;
-	OutVFInput.Strands.bUseStableRasterization = Instance->Strands.Modifier.bUseStableRasterization;
-	OutVFInput.Strands.bUseRaytracingGeometry = false;
-
-	OutVFInput.Strands.AttributeOffsets.SetNum(0, false);
-	OutVFInput.Strands.AttributeOffsets.Reserve(HAIR_ATTRIBUTE_COUNT);
-	for (uint32 AttributeIt=0; AttributeIt< HAIR_ATTRIBUTE_COUNT; ++AttributeIt)
+	OutVFInput.Strands.Common.Radius = (GStrandHairWidth > 0 ? GStrandHairWidth : Instance->Strands.Modifier.HairWidth) * 0.5f;
+	OutVFInput.Strands.Common.RootScale = Instance->Strands.Modifier.HairRootScale;
+	OutVFInput.Strands.Common.TipScale = Instance->Strands.Modifier.HairTipScale;
+	OutVFInput.Strands.Common.RaytracingRadiusScale = (GHairRaytracingRadiusScale > 0 ? GHairRaytracingRadiusScale : Instance->Strands.Modifier.HairRaytracingRadiusScale);
+	OutVFInput.Strands.Common.Length = Instance->Strands.Modifier.HairLength;
+	OutVFInput.Strands.Common.Density = Instance->Strands.Modifier.HairShadowDensity;
+	OutVFInput.Strands.Common.bScatterSceneLighting = Instance->Strands.Modifier.bScatterSceneLighting;
+	OutVFInput.Strands.Common.bStableRasterization = Instance->Strands.Modifier.bUseStableRasterization;
+	OutVFInput.Strands.Common.bRaytracingGeometry = false;
+	OutVFInput.Strands.Common.RaytracingProceduralSplits = GetHairRaytracingProceduralSplits();
+
+	// Add local transform assignement here ...
+
+	if (bUseGuideAttributeOffsets)
+	{
+		PACK_HAIR_ATTRIBUTE_OFFSETS(OutVFInput.Strands.Common.AttributeOffsets, Instance->Guides.Data->AttributeOffsets);
+	}
+	else
 	{
-		OutVFInput.Strands.AttributeOffsets.Add(bUseGuideAttributeOffsets ? Instance->Guides.Data->AttributeOffsets[AttributeIt] : Instance->Strands.Data->AttributeOffsets[AttributeIt]);
+		PACK_HAIR_ATTRIBUTE_OFFSETS(OutVFInput.Strands.Common.AttributeOffsets, Instance->Strands.Data->AttributeOffsets);
 	}
 
 #if RHI_RAYTRACING
 	// Flag bUseRaytracingGeometry only if RT geometry has been allocated for RayTracing view (not for PathTracing view). 
 	// This flag is used later for selecting if voxelization needs to flags voxel has shadow caster or if shadow casting is handled 
 	// by the RT geometry.
-	OutVFInput.Strands.bUseRaytracingGeometry = Instance->Strands.RenRaytracingResource != nullptr && (Instance->Strands.ViewRayTracingMask & EHairViewRayTracingMask::RayTracing) != 0;
+	OutVFInput.Strands.Common.bRaytracingGeometry = Instance->Strands.RenRaytracingResource != nullptr && (Instance->Strands.ViewRayTracingMask & EHairViewRayTracingMask::RayTracing) != 0;
 #endif
 
 	return OutVFInput;
@@ -2154,7 +2161,7 @@ void ComputeHairStrandsInterpolation(
 				if (CullingData.bCullingResultAvailable)
 				{
 					// Sanity check
-					check(Instance->Strands.bIsCullingEnabled);
+					check(Instance->Strands.bCullingEnable);
 					HairGroupCluster = &InClusterData->HairGroups[Instance->HairGroupPublicData->ClusterDataIndex];
 					if (!HairGroupCluster->bVisible)
 					{
@@ -2196,9 +2203,9 @@ void ComputeHairStrandsInterpolation(
 			#if RHI_RAYTRACING
 			if (Instance->Strands.RenRaytracingResource && (Instance->Strands.ViewRayTracingMask & ViewRayTracingMask) != 0)
 			{
-				const float HairRadiusRT	= Instance->HairGroupPublicData->VFInput.Strands.HairRaytracingRadiusScale * Instance->HairGroupPublicData->VFInput.Strands.HairRadius;
-				const float HairRootScaleRT = Instance->HairGroupPublicData->VFInput.Strands.HairRootScale;
-				const float HairTipScaleRT	= Instance->HairGroupPublicData->VFInput.Strands.HairTipScale;
+				const float HairRadiusRT	= Instance->HairGroupPublicData->VFInput.Strands.Common.RaytracingRadiusScale * Instance->HairGroupPublicData->VFInput.Strands.Common.Radius;
+				const float HairRootScaleRT = Instance->HairGroupPublicData->VFInput.Strands.Common.RootScale;
+				const float HairTipScaleRT	= Instance->HairGroupPublicData->VFInput.Strands.Common.TipScale;
 				const int ProceduralSplits = GetHairRaytracingProceduralSplits();
 
 				const bool bNeedUpdate = Instance->Strands.DeformedResource != nullptr ||
diff --git a/Engine/Plugins/Runtime/HairStrands/Source/HairStrandsCore/Private/HairStrandsVertexFactory.cpp b/Engine/Plugins/Runtime/HairStrands/Source/HairStrandsCore/Private/HairStrandsVertexFactory.cpp
--- a/Engine/Plugins/Runtime/HairStrands/Source/HairStrandsCore/Private/HairStrandsVertexFactory.cpp
+++ b/Engine/Plugins/Runtime/HairStrands/Source/HairStrandsCore/Private/HairStrandsVertexFactory.cpp
@@ -94,46 +94,39 @@ FHairStrandsVertexFactoryUniformShaderParameters FHairGroupInstance::GetHairStan
 	const FHairGroupPublicData::FVertexFactoryInput VFInput = ComputeHairStrandsVertexInputData(this, ViewMode);
 
 	FHairStrandsVertexFactoryUniformShaderParameters Out = {};
-	Out.GroupIndex					= Debug.GroupIndex;
-	Out.Radius 						= VFInput.Strands.HairRadius;
-	Out.RootScale 					= VFInput.Strands.HairRootScale;
-	Out.TipScale 					= VFInput.Strands.HairTipScale;
-	Out.RaytracingRadiusScale		= VFInput.Strands.HairRaytracingRadiusScale;
-	Out.RaytracingProceduralSplits  = GetHairRaytracingProceduralSplits();
-	Out.Length 						= VFInput.Strands.HairLength;
-	Out.Density 					= VFInput.Strands.HairDensity;
-	Out.StableRasterization			= VFInput.Strands.bUseStableRasterization;
-	Out.ScatterSceneLighing			= VFInput.Strands.bScatterSceneLighting;
-	Out.PositionBuffer				= VFInput.Strands.PositionBufferRHISRV;
-	Out.PreviousPositionBuffer		= VFInput.Strands.PrevPositionBufferRHISRV;
-	Out.AttributeBuffer				= VFInput.Strands.AttributeBufferRHISRV;
-	Out.PointToCurveBuffer			= VFInput.Strands.PointToCurveBufferRHISRV;
-	Out.TangentBuffer 				= VFInput.Strands.TangentBufferRHISRV;
-	Out.PositionOffsetBuffer 		= VFInput.Strands.PositionOffsetBufferRHISRV;
-	Out.PreviousPositionOffsetBuffer= VFInput.Strands.PrevPositionOffsetBufferRHISRV;
+	Out.Common = VFInput.Strands.Common;
+
+	Out.Resources.PositionBuffer					= VFInput.Strands.PositionBufferRHISRV;
+	Out.Resources.PositionOffsetBuffer 				= VFInput.Strands.PositionOffsetBufferRHISRV;
+	Out.Resources.AttributeBuffer					= VFInput.Strands.AttributeBufferRHISRV;
+	Out.Resources.CurveBuffer						= VFInput.Strands.CurveBufferRHISRV;
+	Out.Resources.PointToCurveBuffer				= VFInput.Strands.PointToCurveBufferRHISRV;
+	Out.Resources.TangentBuffer 					= VFInput.Strands.TangentBufferRHISRV;
+	Out.PrevResources.PreviousPositionBuffer		= VFInput.Strands.PrevPositionBufferRHISRV;
+	Out.PrevResources.PreviousPositionOffsetBuffer	= VFInput.Strands.PrevPositionOffsetBufferRHISRV;
 
 	// swap in some default data for those buffers that are not valid yet
-	if (!Out.PositionBuffer) 				{ Out.PositionBuffer = GDummyCulledDispatchVertexIdsBuffer.SRVFloat; }
-	if (!Out.PreviousPositionBuffer) 		{ Out.PreviousPositionBuffer = GDummyCulledDispatchVertexIdsBuffer.SRVFloat; }
-	if (!Out.AttributeBuffer) 				{ Out.AttributeBuffer = GDummyCulledDispatchVertexIdsBuffer.SRVByteAddress; }
-	if (!Out.PointToCurveBuffer) 			{ Out.PointToCurveBuffer = GDummyCulledDispatchVertexIdsBuffer.SRVUint; }
-	if (!Out.TangentBuffer) 				{ Out.TangentBuffer = GDummyCulledDispatchVertexIdsBuffer.SRVFloat; }
-	if (!Out.PositionOffsetBuffer) 			{ Out.PositionOffsetBuffer = GDummyCulledDispatchVertexIdsBuffer.SRVFloat; }
-	if (!Out.PreviousPositionOffsetBuffer) 	{ Out.PreviousPositionOffsetBuffer = GDummyCulledDispatchVertexIdsBuffer.SRVFloat; }
-
-	Out.CullingEnable = HairGroupPublicData->GetCullingResultAvailable();
-	if (Out.CullingEnable)
+	if (!Out.Resources.PositionBuffer) 						{ Out.Resources.PositionBuffer = GDummyCulledDispatchVertexIdsBuffer.SRVFloat; }
+	if (!Out.Resources.PositionOffsetBuffer) 				{ Out.Resources.PositionOffsetBuffer = GDummyCulledDispatchVertexIdsBuffer.SRVFloat; }
+	if (!Out.Resources.AttributeBuffer) 					{ Out.Resources.AttributeBuffer = GDummyCulledDispatchVertexIdsBuffer.SRVByteAddress; }
+	if (!Out.Resources.CurveBuffer) 						{ Out.Resources.CurveBuffer = GDummyCulledDispatchVertexIdsBuffer.SRVByteAddress; }
+	if (!Out.Resources.PointToCurveBuffer) 					{ Out.Resources.PointToCurveBuffer = GDummyCulledDispatchVertexIdsBuffer.SRVUint; }
+	if (!Out.Resources.TangentBuffer) 						{ Out.Resources.TangentBuffer = GDummyCulledDispatchVertexIdsBuffer.SRVFloat; }
+
+	if (!Out.PrevResources.PreviousPositionBuffer) 			{ Out.PrevResources.PreviousPositionBuffer = GDummyCulledDispatchVertexIdsBuffer.SRVFloat; }
+	if (!Out.PrevResources.PreviousPositionOffsetBuffer) 	{ Out.PrevResources.PreviousPositionOffsetBuffer = GDummyCulledDispatchVertexIdsBuffer.SRVFloat; }
+
+	Out.Culling.bCullingEnable = HairGroupPublicData->GetCullingResultAvailable();
+	if (Out.Culling.bCullingEnable)
 	{
-		Out.CulledVertexIdsBuffer = HairGroupPublicData->GetCulledVertexIdBuffer().SRV;
-		Out.CulledVertexRadiusScaleBuffer = HairGroupPublicData->GetCulledVertexRadiusScaleBuffer().SRV;
+		Out.Culling.CullingIndexBuffer = HairGroupPublicData->GetCulledVertexIdBuffer().SRV;
+		Out.Culling.CullingRadiusScaleBuffer = HairGroupPublicData->GetCulledVertexRadiusScaleBuffer().SRV;
 	}
 	else
 	{
-		Out.CulledVertexIdsBuffer = GDummyCulledDispatchVertexIdsBuffer.SRVUint;
-		Out.CulledVertexRadiusScaleBuffer = GDummyCulledDispatchVertexIdsBuffer.SRVFloat;
+		Out.Culling.CullingIndexBuffer = GDummyCulledDispatchVertexIdsBuffer.SRVUint;
+		Out.Culling.CullingRadiusScaleBuffer = GDummyCulledDispatchVertexIdsBuffer.SRVFloat;
 	}
-
-	PACK_HAIR_ATTRIBUTE_OFFSETS(Out.AttributeOffsets, VFInput.Strands.AttributeOffsets);	
 	return Out;
 }
 
diff --git a/Engine/Plugins/Runtime/HairStrands/Source/HairStrandsCore/Public/GroomInstance.h b/Engine/Plugins/Runtime/HairStrands/Source/HairStrandsCore/Public/GroomInstance.h
--- a/Engine/Plugins/Runtime/HairStrands/Source/HairStrandsCore/Public/GroomInstance.h
+++ b/Engine/Plugins/Runtime/HairStrands/Source/HairStrandsCore/Public/GroomInstance.h
@@ -47,34 +47,11 @@ END_GLOBAL_SHADER_PARAMETER_STRUCT()
 
 typedef TUniformBufferRef<FHairCardsVertexFactoryUniformShaderParameters> FHairCardsUniformBuffer;
 
-
 BEGIN_GLOBAL_SHADER_PARAMETER_STRUCT(FHairStrandsVertexFactoryUniformShaderParameters, HAIRSTRANDSCORE_API)
-	SHADER_PARAMETER(float, Radius)
-	SHADER_PARAMETER(float, RootScale)
-	SHADER_PARAMETER(float, TipScale)
-	SHADER_PARAMETER(float, Length)
-	SHADER_PARAMETER(float, Density)
-	SHADER_PARAMETER(float, RaytracingRadiusScale)
-	SHADER_PARAMETER(uint32, CullingEnable)
-	SHADER_PARAMETER(uint32, HasMaterial)
-	SHADER_PARAMETER(uint32, StableRasterization)
-	SHADER_PARAMETER(uint32, ScatterSceneLighing)
-	SHADER_PARAMETER(uint32, RaytracingProceduralSplits)
-	SHADER_PARAMETER(float, GroupIndex)
-	SHADER_PARAMETER_ARRAY(FUintVector4, AttributeOffsets, [HAIR_ATTRIBUTE_OFFSET_COUNT])
-
-	SHADER_PARAMETER_SRV(Buffer<float4>, PositionOffsetBuffer)
-	SHADER_PARAMETER_SRV(Buffer<float4>, PreviousPositionOffsetBuffer)
-
-	SHADER_PARAMETER_SRV(Buffer<uint4>, PositionBuffer)
-	SHADER_PARAMETER_SRV(Buffer<uint4>, PreviousPositionBuffer)
-
-	SHADER_PARAMETER_SRV(ByteAddressBuffer, AttributeBuffer)
-	SHADER_PARAMETER_SRV(Buffer<uint>, PointToCurveBuffer)
-	SHADER_PARAMETER_SRV(Buffer<float4>, TangentBuffer)
-
-	SHADER_PARAMETER_SRV(Buffer<uint>, CulledVertexIdsBuffer)
-	SHADER_PARAMETER_SRV(Buffer<float>, CulledVertexRadiusScaleBuffer)
+	SHADER_PARAMETER_STRUCT_INCLUDE(FHairStrandsInstanceCommonParameters, Common)
+	SHADER_PARAMETER_STRUCT_INCLUDE(FHairStrandsInstanceResourceRawParameters, Resources)
+	SHADER_PARAMETER_STRUCT_INCLUDE(FHairStrandsInstancePrevResourceRawParameters, PrevResources)
+	SHADER_PARAMETER_STRUCT_INCLUDE(FHairStrandsInstanceCullingRawParameters, Culling)
 END_GLOBAL_SHADER_PARAMETER_STRUCT()
 
 typedef TUniformBufferRef<FHairStrandsVertexFactoryUniformShaderParameters> FHairStrandsUniformBuffer;
@@ -128,7 +105,7 @@ struct HAIRSTRANDSCORE_API FHairGroupInstance : public FHairStrandsInstance
 		EHairInterpolationType HairInterpolationType = EHairInterpolationType::NoneSkinning;
 
 		// Indicates if culling is enabled for this hair strands data.
-		bool bIsCullingEnabled = false;
+		bool bCullingEnable = false;
 	};
 
 	//////////////////////////////////////////////////////////////////////////////////////////
diff --git a/Engine/Shaders/Private/HairStrands/HairStrandsDebug.usf b/Engine/Shaders/Private/HairStrands/HairStrandsDebug.usf
--- a/Engine/Shaders/Private/HairStrands/HairStrandsDebug.usf
+++ b/Engine/Shaders/Private/HairStrands/HairStrandsDebug.usf
@@ -252,10 +252,10 @@ void MainCS(uint3 DispatchThreadId : SV_DispatchThreadID)
 
 	uint VertexIndex = VertexId;
 	float ScaleR = 1;
-	if (HairStrandsVF.CullingEnable)
+	if (HairStrandsVF.bCullingEnable)
 	{
-		ScaleR = HairStrandsVF.CulledVertexRadiusScaleBuffer[VertexIndex];
-		VertexIndex = HairStrandsVF.CulledVertexIdsBuffer[VertexIndex];
+		ScaleR = HairStrandsVF.CullingRadiusScaleBuffer[VertexIndex];
+		VertexIndex = HairStrandsVF.CullingIndexBuffer[VertexIndex];
 	}
 
 	const float3 P_Local = UnpackHairControlPoint(
diff --git a/Engine/Shaders/Private/HairStrands/HairStrandsInterpolation.usf b/Engine/Shaders/Private/HairStrands/HairStrandsInterpolation.usf
--- a/Engine/Shaders/Private/HairStrands/HairStrandsInterpolation.usf
+++ b/Engine/Shaders/Private/HairStrands/HairStrandsInterpolation.usf
@@ -97,12 +97,12 @@ Buffer<uint4>	Interpolation1Buffer;
 
 Buffer<uint>   SimRootPointIndexBuffer;
 
-uint HairStrandsVF_bIsCullingEnable;
+uint HairStrandsVFTODO_bCullingEnable;
 #if PERMUTATION_CULLING == 1
-Buffer<uint>	HairStrandsVF_CullingIndirectBuffer; 
-Buffer<uint>	HairStrandsVF_CullingIndexBuffer;
-Buffer<float>	HairStrandsVF_CullingRadiusScaleBuffer;
-Buffer<uint>	HairStrandsVF_CullingIndirectBufferArgs;
+Buffer<uint>	HairStrandsVFTODO_CullingIndirectBuffer; 
+Buffer<uint>	HairStrandsVFTODO_CullingIndexBuffer;
+Buffer<float>	HairStrandsVFTODO_CullingRadiusScaleBuffer;
+Buffer<uint>	HairStrandsVFTODO_CullingIndirectBufferArgs;
 #endif
 
 #if PERMUTATION_DYNAMIC_GEOMETRY >= 1 || PERMUTATION_DEBUG == 1 
@@ -462,15 +462,15 @@ void MainCS(uint2 DispatchThreadId : SV_DispatchThreadID)
 	// 1. Find vertex index
 	uint VertexIndex = DispatchThreadId;
 #if PERMUTATION_CULLING
-	if (HairStrandsVF_bIsCullingEnable)
+	if (HairStrandsVFTODO_bCullingEnable)
 	{
 		const uint FetchIndex = DispatchThreadId;
-		const uint VertexCountAfterCulling = HairStrandsVF_CullingIndirectBuffer[3];
+		const uint VertexCountAfterCulling = HairStrandsVFTODO_CullingIndirectBuffer[3];
 		if (FetchIndex >= VertexCountAfterCulling)
 		{
 			return;
 		}
-		VertexIndex = HairStrandsVF_CullingIndexBuffer[FetchIndex];
+		VertexIndex = HairStrandsVFTODO_CullingIndexBuffer[FetchIndex];
 	}
 #endif
 
@@ -662,7 +662,7 @@ Buffer<uint> ClusterIndexCountBuffer;
 Buffer<uint> CulledClusterCountBuffer;
 
 Buffer<uint> IndirectBuffer;
-Buffer<uint> HairStrandsVF_CullingIndirectBuffer;
+Buffer<uint> HairStrandsVFTODO_CullingIndirectBuffer;
 Buffer<float4> RenderDeformedOffsetBuffer;
 
 RWBuffer<int> OutClusterAABBBuffer; 	// Cluster data packed as {uint3 AABBMin, uint3 AABBMax}
@@ -699,7 +699,7 @@ void ClusterAABBEvaluationCS(uint3 DispatchThreadId : SV_DispatchThreadID, uint3
 	const uint IndexOffset = ClusterIndexOffsetBuffer[ClusterFetchIndex];
 	const uint IndexCount = min(ClusterIndexCountBuffer[ClusterFetchIndex], GROUP_SIZE);
 	
-	const uint VertexCountAfterCulling = HairStrandsVF_CullingIndirectBuffer[3];
+	const uint VertexCountAfterCulling = HairStrandsVFTODO_CullingIndirectBuffer[3];
 
 	int3 ClusterAABBMin = 9999999.0f;
 	int3 ClusterAABBMax =-9999999.0f;
diff --git a/Engine/Shaders/Private/HairStrands/HairStrandsRaytracingGeometry.usf b/Engine/Shaders/Private/HairStrands/HairStrandsRaytracingGeometry.usf
--- a/Engine/Shaders/Private/HairStrands/HairStrandsRaytracingGeometry.usf
+++ b/Engine/Shaders/Private/HairStrands/HairStrandsRaytracingGeometry.usf
@@ -24,7 +24,7 @@ uint             RaytracingProceduralSplits;
 RWBuffer<uint>   OutputIndexBuffer;
 #endif
 
-uint HairStrandsVF_bIsCullingEnable;
+uint HairStrandsVF_bCullingEnable;
 #if PERMUTATION_CULLING == 1
 Buffer<uint>	HairStrandsVF_CullingIndirectBuffer;
 Buffer<uint>	HairStrandsVF_CullingIndexBuffer;
@@ -172,7 +172,7 @@ void MainCS(uint DispatchThreadId : SV_DispatchThreadID)
 	uint SourceIndex0 = VertexIndex0;
 	uint SourceIndex1 = VertexIndex1;
 #if PERMUTATION_CULLING == 1
-	if (HairStrandsVF_bIsCullingEnable)
+	if (HairStrandsVF_bCullingEnable)
 	{
 		const uint VertexCountAfterCulling = HairStrandsVF_CullingIndirectBuffer[3];
 		if (VertexIndex1 >= VertexCountAfterCulling)
@@ -301,7 +301,7 @@ uint VertexCount;
 uint bDrawInvalidElement;
 float PositionThreshold2;
 
-uint HairStrandsVF_bIsCullingEnable;
+uint HairStrandsVF_bCullingEnable;
 Buffer<uint> 	 HairStrandsVF_CullingIndirectBuffer;
 Buffer<uint> 	 HairStrandsVF_CullingIndexBuffer;
 Buffer<uint4> 	 CurrPositionBuffer;
@@ -324,7 +324,7 @@ void MainCS(uint2 DispatchThreadId : SV_DispatchThreadID, uint LinearIndex : SV_
 		// Fetch vertex index if culling is enabled
 		uint VertexIndex0 = DispatchThreadId;
 		bool bIsValid = VertexIndex0 < VertexCount;
-		if (bIsValid && HairStrandsVF_bIsCullingEnable)
+		if (bIsValid && HairStrandsVF_bCullingEnable)
 		{
 			const uint VertexCountAfterCulling = HairStrandsVF_CullingIndirectBuffer[3];
 			const uint FetchIndex0 = DispatchThreadId;
diff --git a/Engine/Shaders/Private/HairStrands/HairStrandsTangent.usf b/Engine/Shaders/Private/HairStrands/HairStrandsTangent.usf
--- a/Engine/Shaders/Private/HairStrands/HairStrandsTangent.usf
+++ b/Engine/Shaders/Private/HairStrands/HairStrandsTangent.usf
@@ -10,7 +10,7 @@ uint VertexCount;
 Buffer<uint4>  	 PositionBuffer;
 RWBuffer<float4> OutputTangentBuffer;
 
-uint			HairStrandsVF_bIsCullingEnable;
+uint			HairStrandsVF_bCullingEnable;
 
 #if PERMUTATION_CULLING == 1
 Buffer<uint>	HairStrandsVF_CullingIndirectBuffer;
@@ -32,7 +32,7 @@ void MainCS(uint DispatchThreadId : SV_DispatchThreadID)
 	int IndexNext = clamp(IndexCurr+1, 0, MaxVertexCount-1);
 	bool bValid = IndexCurr < MaxVertexCount;
 #if PERMUTATION_CULLING
-	if (HairStrandsVF_bIsCullingEnable)
+	if (HairStrandsVF_bCullingEnable)
 	{
 		const int FetchIndex = DispatchThreadId;
 		const int VertexCountAfterCulling = clamp((int)HairStrandsVF_CullingIndirectBuffer[3], 0, MaxVertexCount);
diff --git a/Engine/Shaders/Private/HairStrands/HairStrandsVertexFactory.ush b/Engine/Shaders/Private/HairStrands/HairStrandsVertexFactory.ush
--- a/Engine/Shaders/Private/HairStrands/HairStrandsVertexFactory.ush
+++ b/Engine/Shaders/Private/HairStrands/HairStrandsVertexFactory.ush
@@ -74,12 +74,12 @@
  */
 bool UseStableRasterization()
 {
-	return HairStrandsVF.StableRasterization > 0;
+	return HairStrandsVF.bStableRasterization > 0;
 }
 
 bool UseScatterSceneLighting()
 {
-	return HairStrandsVF.ScatterSceneLighing > 0;
+	return HairStrandsVF.bScatterSceneLighting > 0;
 }
 
 struct FVertexFactoryInput
@@ -138,11 +138,11 @@ FVertexInfo GetVertexInfo(FVertexFactoryInput Input)
 #endif
 
 #if USE_CULLED_CLUSTER == 1
-	if (HairStrandsVF.CullingEnable)
+	if (HairStrandsVF.bCullingEnable)
 	{
-		VertexInfo.VertexIndex = HairStrandsVF.CulledVertexIdsBuffer[BaseIndex + VertexInfo.IsTip];
+		VertexInfo.VertexIndex = HairStrandsVF.CullingIndexBuffer[BaseIndex + VertexInfo.IsTip];
 #if !RAYHITGROUPSHADER
-		VertexInfo.LodRadiusScale = HairStrandsVF.CulledVertexRadiusScaleBuffer[BaseIndex + VertexInfo.IsTip];
+		VertexInfo.LodRadiusScale = HairStrandsVF.CullingRadiusScaleBuffer[BaseIndex + VertexInfo.IsTip];
 		VertexInfo.bForceInvalidQuad = VertexInfo.LodRadiusScale <= 0.0f;
 #endif
 		VertexInfo.HairPrimitiveId = VertexInfo.VertexIndex - VertexInfo.IsTip;
@@ -398,12 +398,12 @@ RAY_TRACING_ENTRY_INTERSECTION(MaterialIS)
 	uint VertexIndex0 = PrimitiveIndex() / HairStrandsVF.RaytracingProceduralSplits;
 	uint VertexIndex1 = VertexIndex0 + 1;
 	float ScaleR0 = 1, ScaleR1 = 1;
-	if (HairStrandsVF.CullingEnable)
+	if (HairStrandsVF.bCullingEnable)
 	{
-		ScaleR0 = HairStrandsVF.CulledVertexRadiusScaleBuffer[VertexIndex0];
-		ScaleR1 = HairStrandsVF.CulledVertexRadiusScaleBuffer[VertexIndex1];
-		VertexIndex0 = HairStrandsVF.CulledVertexIdsBuffer[VertexIndex0];
-		VertexIndex1 = HairStrandsVF.CulledVertexIdsBuffer[VertexIndex1];
+		ScaleR0 = HairStrandsVF.CullingRadiusScaleBuffer[VertexIndex0];
+		ScaleR1 = HairStrandsVF.CullingRadiusScaleBuffer[VertexIndex1];
+		VertexIndex0 = HairStrandsVF.CullingIndexBuffer[VertexIndex0];
+		VertexIndex1 = HairStrandsVF.CullingIndexBuffer[VertexIndex1];
 	}
 
 	const float EffectiveRadius = HairStrandsVF.Radius * HairStrandsVF.RaytracingRadiusScale;
@@ -482,10 +482,10 @@ FMaterialPixelParameters GetMaterialPixelParameters(float3 RayOrigin, float3 Ray
 	uint VertexIndex0 = HitPrimitiveIndex / HairStrandsVF.RaytracingProceduralSplits;
 	uint VertexIndex1 = VertexIndex0 + 1;
 
-	if (HairStrandsVF.CullingEnable)
+	if (HairStrandsVF.bCullingEnable)
 	{
-		VertexIndex0 = HairStrandsVF.CulledVertexIdsBuffer[VertexIndex0];
-		VertexIndex1 = HairStrandsVF.CulledVertexIdsBuffer[VertexIndex1];
+		VertexIndex0 = HairStrandsVF.CullingIndexBuffer[VertexIndex0];
+		VertexIndex1 = HairStrandsVF.CullingIndexBuffer[VertexIndex1];
 	}
 
 	const float EffectiveRadius = HairStrandsVF.Radius * HairStrandsVF.RaytracingRadiusScale;
@@ -533,10 +533,10 @@ FMaterialPixelParameters GetMaterialPixelParameters(float3 RayOrigin, float3 Ray
 	uint VertexIndex0 = Tri.Indices.z / 4; // last index always belongs to the "bottom" vertex of the hair segment
 	uint VertexIndex1 = VertexIndex0 + 1;
 
-	if (HairStrandsVF.CullingEnable)
+	if (HairStrandsVF.bCullingEnable)
 	{
-		VertexIndex0 = HairStrandsVF.CulledVertexIdsBuffer[VertexIndex0];
-		VertexIndex1 = HairStrandsVF.CulledVertexIdsBuffer[VertexIndex1];
+		VertexIndex0 = HairStrandsVF.CullingIndexBuffer[VertexIndex0];
+		VertexIndex1 = HairStrandsVF.CullingIndexBuffer[VertexIndex1];
 	}
 
 	float UCoord = dot(float3(1, float((HitPrimitiveIndex % 2) == 0), 0), Weights);
diff --git a/Engine/Shaders/Private/HairStrands/HairStrandsVertexFactoryCommon.ush b/Engine/Shaders/Private/HairStrands/HairStrandsVertexFactoryCommon.ush
--- a/Engine/Shaders/Private/HairStrands/HairStrandsVertexFactoryCommon.ush
+++ b/Engine/Shaders/Private/HairStrands/HairStrandsVertexFactoryCommon.ush
@@ -4,6 +4,10 @@
 
 #include "/Engine/Shared/HairStrandsDefinitions.h"
 
+#ifndef HAIR_STRANDS_PARAMETERS
+#define HAIR_STRANDS_PARAMETERS 0
+#endif
+
 ////////////////////////////////////////////////////////////////////////////////////////////////////////////////////////////////
 // Hair control points
 
@@ -143,14 +147,15 @@ FHairCurve UnpackHairCurve(uint In)
 // Hair strands instance parameters
 
 #if defined(HAIR_STRANDS_PARAMETERS) && HAIR_STRANDS_PARAMETERS == 1
-uint			HairStrandsVF_bIsCullingEnable;
-uint			HairStrandsVF_bHasRaytracedGeometry;
+uint			HairStrandsVF_bCullingEnable;
+uint			HairStrandsVF_bRaytracingGeometry;
+uint			HairStrandsVF_RaytracingProceduralSplits;
 float			HairStrandsVF_Radius;
 float			HairStrandsVF_RootScale;
 float			HairStrandsVF_TipScale;
 float			HairStrandsVF_Length;
 float			HairStrandsVF_Density;
-uint			HairStrandsVF_bUseStableRasterization;
+uint			HairStrandsVF_bStableRasterization;
 float3			HairStrandsVF_PositionOffset;
 uint			HairStrandsVF_PointCount;
 uint			HairStrandsVF_CurveCount;
diff --git a/Engine/Shaders/Private/HairStrands/HairStrandsVisibilityRasterClassification.usf b/Engine/Shaders/Private/HairStrands/HairStrandsVisibilityRasterClassification.usf
--- a/Engine/Shaders/Private/HairStrands/HairStrandsVisibilityRasterClassification.usf
+++ b/Engine/Shaders/Private/HairStrands/HairStrandsVisibilityRasterClassification.usf
@@ -62,7 +62,7 @@ void CSMain(uint DispatchThreadID : SV_DispatchThreadID, uint GroupThreadID : SV
 
 #if PERMUTATION_CULLING	
 	uint NumVerticesAfterCulling = ControlPointCount;
-	if (HairStrandsVF_bIsCullingEnable)
+	if (HairStrandsVF_bCullingEnable)
 	{
 		NumVerticesAfterCulling = HairStrandsVF_CullingIndirectBuffer[3];
 	}
@@ -81,7 +81,7 @@ void CSMain(uint DispatchThreadID : SV_DispatchThreadID, uint GroupThreadID : SV
 		uint PrimID = BatchIndex * 1024 + GroupThreadID;
 		bool bSegValid = (BatchIndex < NumBatches) && (PrimID < NumVerticesAfterCulling);
 
-		if (bSegValid && HairStrandsVF_bIsCullingEnable)
+		if (bSegValid && HairStrandsVF_bCullingEnable)
 		{
 			const uint FetchIndex0 = PrimID;
 			const uint FetchIndex1 = min(FetchIndex0 + 1, NumVerticesAfterCulling - 1);
diff --git a/Engine/Shaders/Private/HairStrands/HairStrandsVisibilityRasterCompute.usf b/Engine/Shaders/Private/HairStrands/HairStrandsVisibilityRasterCompute.usf
--- a/Engine/Shaders/Private/HairStrands/HairStrandsVisibilityRasterCompute.usf
+++ b/Engine/Shaders/Private/HairStrands/HairStrandsVisibilityRasterCompute.usf
@@ -400,7 +400,7 @@ void BinningCS(uint DispatchThreadID : SV_DispatchThreadID, uint GroupThreadID :
 		group_VerticesNum = IndirectPrimIDCount.Load(0);
 #else // PERMUTATION_INDIRECT_PRIM_IDS
 	#if PERMUTATION_CULLING	
-		group_VerticesNum = HairStrandsVF_bIsCullingEnable ? HairStrandsVF_CullingIndirectBuffer[3] : VertexCount;
+		group_VerticesNum = HairStrandsVF_bCullingEnable ? HairStrandsVF_CullingIndirectBuffer[3] : VertexCount;
 	#else // PERMUTATION_CULLING
 		group_VerticesNum = VertexCount;
 	#endif //PERMUTATION_CULLING
@@ -430,7 +430,7 @@ void BinningCS(uint DispatchThreadID : SV_DispatchThreadID, uint GroupThreadID :
 		uint PrimID = BatchIndex * 1024 + GroupThreadID;
 		bSegValid = bSegValid && (PrimID < group_VerticesNum);
 
-		if (bSegValid && HairStrandsVF_bIsCullingEnable)
+		if (bSegValid && HairStrandsVF_bCullingEnable)
 		{
 			const uint FetchIndex0 = PrimID;
 			const uint FetchIndex1 = min(FetchIndex0 + 1, group_VerticesNum - 1);
diff --git a/Engine/Shaders/Private/HairStrands/HairStrandsVisibilityRasterComputeNaive.usf b/Engine/Shaders/Private/HairStrands/HairStrandsVisibilityRasterComputeNaive.usf
--- a/Engine/Shaders/Private/HairStrands/HairStrandsVisibilityRasterComputeNaive.usf
+++ b/Engine/Shaders/Private/HairStrands/HairStrandsVisibilityRasterComputeNaive.usf
@@ -123,7 +123,7 @@ void CSMain(uint DispatchThreadID : SV_DispatchThreadID, uint GroupThreadID : SV
 	const uint NumControlPoints = IndirectPrimIDCount.Load(0);
 #else // PERMUTATION_INDIRECT_PRIM_IDS
 	#if PERMUTATION_CULLING	
-	const uint NumControlPoints = HairStrandsVF_bIsCullingEnable ? HairStrandsVF_CullingIndirectBuffer[3] : ControlPointCount;
+	const uint NumControlPoints = HairStrandsVF_bCullingEnable ? HairStrandsVF_CullingIndirectBuffer[3] : ControlPointCount;
 	#else // PERMUTATION_CULLING
 	const uint NumControlPoints = ControlPointCount;
 	#endif //PERMUTATION_CULLING
@@ -151,7 +151,7 @@ void CSMain(uint DispatchThreadID : SV_DispatchThreadID, uint GroupThreadID : SV
 		uint PrimID = BatchIndex * 64 + GroupThreadID;
 		bSegValid = bSegValid && (PrimID < NumControlPoints);
 
-		if (bSegValid && HairStrandsVF_bIsCullingEnable)
+		if (bSegValid && HairStrandsVF_bCullingEnable)
 		{
 			const uint FetchIndex0 = PrimID;
 			const uint FetchIndex1 = min(FetchIndex0 + 1, NumControlPoints - 1);
diff --git a/Engine/Shaders/Private/HairStrands/HairStrandsVisibilityRasterHW.usf b/Engine/Shaders/Private/HairStrands/HairStrandsVisibilityRasterHW.usf
--- a/Engine/Shaders/Private/HairStrands/HairStrandsVisibilityRasterHW.usf
+++ b/Engine/Shaders/Private/HairStrands/HairStrandsVisibilityRasterHW.usf
@@ -53,7 +53,7 @@ VSOutput VSMain(VSInput Input)
 	#if PERMUTATION_CULLING
 	uint PrimID = Input.InstanceID;
 
-	if (HairStrandsVF_bIsCullingEnable)
+	if (HairStrandsVF_bCullingEnable)
 	{
 		const uint NumVerticesAfterCulling = HairStrandsVF_CullingIndirectBuffer[3];
 		if (PrimID >= NumVerticesAfterCulling)
diff --git a/Engine/Shaders/Private/HairStrands/HairStrandsVoxelRasterCompute.usf b/Engine/Shaders/Private/HairStrands/HairStrandsVoxelRasterCompute.usf
--- a/Engine/Shaders/Private/HairStrands/HairStrandsVoxelRasterCompute.usf
+++ b/Engine/Shaders/Private/HairStrands/HairStrandsVoxelRasterCompute.usf
@@ -690,7 +690,7 @@ void MainCS(uint2 DispatchThreadID : SV_DispatchThreadID)
 	float LodRadiusScale0 = 1;
 	float LodRadiusScale1 = 1;
 #if PERMUTATION_CULLING == 1
-	if (HairStrandsVF_bIsCullingEnable)
+	if (HairStrandsVF_bCullingEnable)
 	{
 		const uint VertexCountAfterCulling = HairStrandsVF_CullingIndirectBuffer[3];
 		uint FetchIndex0 = VertexIndex0;
@@ -824,7 +824,7 @@ void MainCS(uint2 DispatchThreadID : SV_DispatchThreadID)
 				// Groom having raytraced geometry will cast-shadow on opaque geomtry with their RT-geometry, not with their 
 				// voxelization. To avoid for doubling/incorrect shadowing, we mark voxel with no-shadow casting flag. 
 				// This adds a significant cost when used.
-				if (HairStrandsVF_bHasRaytracedGeometry)
+				if (HairStrandsVF_bRaytracingGeometry)
 				{
 					InterlockedOr(OutPageTexture[VoxelPageCoord], VOXEL_CAST_NO_SHADOW_MASK);
 				}
diff --git a/Engine/Source/Runtime/Renderer/Private/HairStrands/HairStrandsCluster.cpp b/Engine/Source/Runtime/Renderer/Private/HairStrands/HairStrandsCluster.cpp
--- a/Engine/Source/Runtime/Renderer/Private/HairStrands/HairStrandsCluster.cpp
+++ b/Engine/Source/Runtime/Renderer/Private/HairStrands/HairStrandsCluster.cpp
@@ -164,7 +164,7 @@ static void InternalUpdateMacroGroup(FHairStrandsMacroGroupData& MacroGroup, int
 {
 	check(HairData);
 
-	if (HairData->VFInput.Strands.bScatterSceneLighting)
+	if (HairData->VFInput.Strands.Common.bScatterSceneLighting)
 	{
 		MacroGroup.bNeedScatterSceneLighting = true;
 	}
diff --git a/Engine/Source/Runtime/Renderer/Private/HairStrands/HairStrandsData.h b/Engine/Source/Runtime/Renderer/Private/HairStrands/HairStrandsData.h
--- a/Engine/Source/Runtime/Renderer/Private/HairStrands/HairStrandsData.h
+++ b/Engine/Source/Runtime/Renderer/Private/HairStrands/HairStrandsData.h
@@ -13,6 +13,7 @@
 #include "Shader.h"
 #include "ConvexVolume.h"
 #include "HairStrandsDefinitions.h"
+#include "HairStrandsInterface.h"
 
 class FLightSceneInfo;
 class FPrimitiveSceneProxy;
@@ -46,35 +47,6 @@ BEGIN_GLOBAL_SHADER_PARAMETER_STRUCT(FHairStrandsViewUniformParameters, )
 	SHADER_PARAMETER_RDG_BUFFER_SRV(Buffer<uint>,  HairTileCount)						// Tile total count (actual number of tiles)
 END_GLOBAL_SHADER_PARAMETER_STRUCT()
 
-////////////////////////////////////////////////////////////////////////////////////
-// Hair Instance data
-
-BEGIN_SHADER_PARAMETER_STRUCT(FHairStrandsInstanceParameters, )
-	SHADER_PARAMETER(uint32, HairStrandsVF_bIsCullingEnable)
-	SHADER_PARAMETER(uint32, HairStrandsVF_bHasRaytracedGeometry)
-	SHADER_PARAMETER(float, HairStrandsVF_Density)
-	SHADER_PARAMETER(float, HairStrandsVF_Radius)
-	SHADER_PARAMETER(float, HairStrandsVF_RootScale)
-	SHADER_PARAMETER(float, HairStrandsVF_TipScale)
-	SHADER_PARAMETER(float, HairStrandsVF_Length)
-	SHADER_PARAMETER(uint32, HairStrandsVF_bUseStableRasterization)
-	SHADER_PARAMETER(uint32, HairStrandsVF_PointCount)
-	SHADER_PARAMETER(uint32, HairStrandsVF_CurveCount)
-	SHADER_PARAMETER_ARRAY(FUintVector4, HairStrandsVF_AttributeOffsets, [HAIR_ATTRIBUTE_OFFSET_COUNT])
-	SHADER_PARAMETER(FVector3f, HairStrandsVF_PositionOffset)
-	SHADER_PARAMETER(FMatrix44f, HairStrandsVF_LocalToWorldPrimitiveTransform)
-	SHADER_PARAMETER(FMatrix44f, HairStrandsVF_LocalToTranslatedWorldPrimitiveTransform)
-	SHADER_PARAMETER_RDG_BUFFER_SRV(Buffer, HairStrandsVF_PositionBuffer)
-	SHADER_PARAMETER_RDG_BUFFER_SRV(Buffer, HairStrandsVF_PositionOffsetBuffer)
-	SHADER_PARAMETER_RDG_BUFFER_SRV(Buffer, HairStrandsVF_CurveBuffer)
-	SHADER_PARAMETER_RDG_BUFFER_SRV(Buffer, HairStrandsVF_PointToCurveBuffer)
-	SHADER_PARAMETER_RDG_BUFFER_SRV(ByteAddressBuffer, HairStrandsVF_AttributeBuffer)
-	SHADER_PARAMETER_RDG_BUFFER_SRV(Buffer, HairStrandsVF_CullingIndirectBuffer)
-	SHADER_PARAMETER_RDG_BUFFER_SRV(Buffer, HairStrandsVF_CullingIndexBuffer)
-	SHADER_PARAMETER_RDG_BUFFER_SRV(Buffer, HairStrandsVF_CullingRadiusScaleBuffer)
-	RDG_BUFFER_ACCESS(HairStrandsVF_CullingIndirectBufferArgs, ERHIAccess::IndirectArgs)
-END_SHADER_PARAMETER_STRUCT()
-
 FHairStrandsInstanceParameters GetHairStrandsInstanceParameters(FRDGBuilder& GraphBuilder, const FViewInfo& ViewInfo, const FHairGroupPublicData* HairGroupPublicData, bool bCullingEnable, bool bForceRegister);
 
 ////////////////////////////////////////////////////////////////////////////////////
diff --git a/Engine/Source/Runtime/Renderer/Private/HairStrands/HairStrandsForwardRaster.cpp b/Engine/Source/Runtime/Renderer/Private/HairStrands/HairStrandsForwardRaster.cpp
--- a/Engine/Source/Runtime/Renderer/Private/HairStrands/HairStrandsForwardRaster.cpp
+++ b/Engine/Source/Runtime/Renderer/Private/HairStrands/HairStrandsForwardRaster.cpp
@@ -305,7 +305,7 @@ FRasterForwardCullingOutput AddHairStrandsForwardCullingPass(
 
 			const uint32 PointCount = HairGroupPublicData->GetActiveStrandsPointCount();
 			// Sanity check
-			check(PointCount == HairGroupPublicData->VFInput.Strands.PointCount);
+			check(PointCount == HairGroupPublicData->VFInput.Strands.Common.PointCount);
 
 			MaxNumPrimIDs += PointCount;
 		}
@@ -356,9 +356,6 @@ FRasterForwardCullingOutput AddHairStrandsForwardCullingPass(
 			const FHairGroupPublicData* HairGroupPublicData = reinterpret_cast<const FHairGroupPublicData*>(PrimitiveInfo.Mesh->Elements[0].VertexFactoryUserData);
 			check(HairGroupPublicData);
 
-			const FHairGroupPublicData::FVertexFactoryInput& VFInput = HairGroupPublicData->VFInput;
-			const FMatrix44d LocalToWorldPrimitiveTransform = VFInput.LocalToWorldTransform.ToMatrixWithScale(); 
-
 			const bool bCullingEnable = false; //bSupportCulling && GHairVisibilityComputeRaster_Culling ? HairGroupPublicData->GetCullingResultAvailable() : false;
 			
 			// calculate current view screen size - which can result in fewer strands rasterized in current view
diff --git a/Engine/Source/Runtime/Renderer/Private/HairStrands/HairStrandsInterface.cpp b/Engine/Source/Runtime/Renderer/Private/HairStrands/HairStrandsInterface.cpp
--- a/Engine/Source/Runtime/Renderer/Private/HairStrands/HairStrandsInterface.cpp
+++ b/Engine/Source/Runtime/Renderer/Private/HairStrands/HairStrandsInterface.cpp
@@ -676,7 +676,7 @@ bool IsHairVisible(const FMeshBatchAndRelevance& MeshBatch)
 		const FHairGroupPublicData* Data = HairStrands::GetHairData(MeshBatch.Mesh);
 		switch (Data->VFInput.GeometryType)
 		{
-		case EHairGeometryType::Strands: return Data->VFInput.Strands.HairLengthScale > 0;
+		case EHairGeometryType::Strands: return Data->VFInput.Strands.Common.LengthScale > 0;
 		case EHairGeometryType::Cards: return true;
 		case EHairGeometryType::Meshes: return true;
 		}
diff --git a/Engine/Source/Runtime/Renderer/Private/HairStrands/HairStrandsVisibility.cpp b/Engine/Source/Runtime/Renderer/Private/HairStrands/HairStrandsVisibility.cpp
--- a/Engine/Source/Runtime/Renderer/Private/HairStrands/HairStrandsVisibility.cpp
+++ b/Engine/Source/Runtime/Renderer/Private/HairStrands/HairStrandsVisibility.cpp
@@ -2835,57 +2835,47 @@ FHairStrandsInstanceParameters GetHairStrandsInstanceParameters(FRDGBuilder& Gra
 
 	const FHairGroupPublicData::FVertexFactoryInput& VFInput = HairGroupPublicData->VFInput;
 
-	FHairStrandsInstanceParameters Out;
+	FHairStrandsInstanceParameters Output;
+	FHairStrandsInstanceIntermediateParameters& Out = Output.HairStrandsVF;
 	if (bForceRegister)
 	{
-		Out.HairStrandsVF_PositionBuffer		= Register(GraphBuilder, VFInput.Strands.PositionBufferExternal, ERDGImportedBufferFlags::CreateSRV).SRV;;
-		Out.HairStrandsVF_PositionOffsetBuffer	= Register(GraphBuilder, VFInput.Strands.PositionOffsetBufferExternal, ERDGImportedBufferFlags::CreateSRV).SRV;;
-		Out.HairStrandsVF_CurveBuffer			= Register(GraphBuilder, VFInput.Strands.CurveBufferExternal, ERDGImportedBufferFlags::CreateSRV).SRV;
-		Out.HairStrandsVF_PointToCurveBuffer	= Register(GraphBuilder, VFInput.Strands.PointToCurveBufferExternal, ERDGImportedBufferFlags::CreateSRV).SRV;
-		Out.HairStrandsVF_AttributeBuffer		= Register(GraphBuilder, VFInput.Strands.AttributeBufferExternal, ERDGImportedBufferFlags::CreateSRV).SRV;
+		Out.Resources.PositionBuffer		= Register(GraphBuilder, VFInput.Strands.PositionBufferExternal, ERDGImportedBufferFlags::CreateSRV).SRV;
+		Out.Resources.PositionOffsetBuffer	= Register(GraphBuilder, VFInput.Strands.PositionOffsetBufferExternal, ERDGImportedBufferFlags::CreateSRV).SRV;
+		Out.Resources.CurveBuffer			= Register(GraphBuilder, VFInput.Strands.CurveBufferExternal, ERDGImportedBufferFlags::CreateSRV).SRV;
+		Out.Resources.PointToCurveBuffer	= Register(GraphBuilder, VFInput.Strands.PointToCurveBufferExternal, ERDGImportedBufferFlags::CreateSRV).SRV;
+		Out.Resources.AttributeBuffer		= Register(GraphBuilder, VFInput.Strands.AttributeBufferExternal, ERDGImportedBufferFlags::CreateSRV).SRV;
 	}
 	else
 	{
-		Out.HairStrandsVF_PositionBuffer		= VFInput.Strands.PositionBuffer.SRV;
-		Out.HairStrandsVF_PositionOffsetBuffer	= VFInput.Strands.PositionOffsetBuffer.SRV;
-		Out.HairStrandsVF_CurveBuffer			= VFInput.Strands.CurveBuffer.SRV;
-		Out.HairStrandsVF_PointToCurveBuffer	= VFInput.Strands.PointToCurveBuffer.SRV;
-		Out.HairStrandsVF_AttributeBuffer		= VFInput.Strands.AttributeBuffer.SRV;
+		Out.Resources.PositionBuffer		= VFInput.Strands.PositionBuffer.SRV;
+		Out.Resources.PositionOffsetBuffer	= VFInput.Strands.PositionOffsetBuffer.SRV;
+		Out.Resources.CurveBuffer			= VFInput.Strands.CurveBuffer.SRV;
+		Out.Resources.PointToCurveBuffer	= VFInput.Strands.PointToCurveBuffer.SRV;
+		Out.Resources.AttributeBuffer		= VFInput.Strands.AttributeBuffer.SRV;
 	}
 
-	Out.HairStrandsVF_PointCount				= VFInput.Strands.PointCount;
-	Out.HairStrandsVF_CurveCount				= VFInput.Strands.CurveCount;
-	Out.HairStrandsVF_Radius					= VFInput.Strands.HairRadius;
-	Out.HairStrandsVF_RootScale					= VFInput.Strands.HairRootScale;
-	Out.HairStrandsVF_TipScale					= VFInput.Strands.HairTipScale;
-	Out.HairStrandsVF_Length					= VFInput.Strands.HairLength;
-	Out.HairStrandsVF_bUseStableRasterization	= VFInput.Strands.bUseStableRasterization ? 1 : 0;
-	Out.HairStrandsVF_Density					= VFInput.Strands.HairDensity;
-	Out.HairStrandsVF_bIsCullingEnable			= 0;
-	Out.HairStrandsVF_bHasRaytracedGeometry		= VFInput.Strands.bUseRaytracingGeometry ? 1u : 0u;
-	Out.HairStrandsVF_PositionOffset			= (FVector3f)VFInput.Strands.PositionOffset;
-
-	PACK_HAIR_ATTRIBUTE_OFFSETS(Out.HairStrandsVF_AttributeOffsets, VFInput.Strands.AttributeOffsets);
+	Out.Common = VFInput.Strands.Common;
+	Out.Culling.bCullingEnable = 0;
 
 	// Absolute local to world
-	Out.HairStrandsVF_LocalToWorldPrimitiveTransform	= FMatrix44f(VFInput.LocalToWorldTransform.ToMatrixWithScale()); // LWC_TODO: Precision loss
+	Out.Common.LocalToWorldPrimitiveTransform	= FMatrix44f(VFInput.LocalToWorldTransform.ToMatrixWithScale()); // LWC_TODO: Precision loss, remove usage for this
 
 	// Translated local to world
 	const FVector& TranslatedWorldOffset = ViewInfo.ViewMatrices.GetPreViewTranslation();
 	FTransform LocalToTranslatedWorldTransform = VFInput.LocalToWorldTransform;
 	LocalToTranslatedWorldTransform.AddToTranslation(TranslatedWorldOffset);
-	Out.HairStrandsVF_LocalToTranslatedWorldPrimitiveTransform = FMatrix44f(LocalToTranslatedWorldTransform.ToMatrixWithScale());
+	Out.Common.LocalToTranslatedWorldPrimitiveTransform = FMatrix44f(LocalToTranslatedWorldTransform.ToMatrixWithScale());
 
 	if (bCullingEnable)
 	{
 		FRDGImportedBuffer CullingIndirectBuffer = Register(GraphBuilder, HairGroupPublicData->GetDrawIndirectRasterComputeBuffer(), ERDGImportedBufferFlags::CreateSRV);
-		Out.HairStrandsVF_CullingIndirectBuffer = CullingIndirectBuffer.SRV;
-		Out.HairStrandsVF_bIsCullingEnable = 1;
-		Out.HairStrandsVF_CullingIndexBuffer = RegisterAsSRV(GraphBuilder, HairGroupPublicData->GetCulledVertexIdBuffer());
-		Out.HairStrandsVF_CullingRadiusScaleBuffer = RegisterAsSRV(GraphBuilder, HairGroupPublicData->GetCulledVertexRadiusScaleBuffer());
-		Out.HairStrandsVF_CullingIndirectBufferArgs = CullingIndirectBuffer.Buffer;
+		Out.Culling.bCullingEnable = 1;
+		Out.Culling.CullingIndirectBuffer = CullingIndirectBuffer.SRV;
+		Out.Culling.CullingIndexBuffer = RegisterAsSRV(GraphBuilder, HairGroupPublicData->GetCulledVertexIdBuffer());
+		Out.Culling.CullingRadiusScaleBuffer = RegisterAsSRV(GraphBuilder, HairGroupPublicData->GetCulledVertexRadiusScaleBuffer());
+		Out.Culling.CullingIndirectBufferArgs = CullingIndirectBuffer.Buffer;
 	}
-	return Out;
+	return Output;
 }
 
 class FVisiblityRasterClassificationCS : public FGlobalShader
@@ -3265,7 +3255,7 @@ static FRasterComputeOutput AddVisibilityComputeRasterPass(
 
 			const uint32 PointCount = HairGroupPublicData->GetActiveStrandsPointCount();
 			// Sanity check
-			check(HairGroupPublicData->VFInput.Strands.PointCount == PointCount);
+			check(HairGroupPublicData->VFInput.Strands.Common.PointCount == PointCount);
 
 			MaxNumPrimIDs = FMath::Max(MaxNumPrimIDs, PointCount);
 		}
@@ -3440,14 +3430,11 @@ static FRasterComputeOutput AddVisibilityComputeRasterPass(
 			const FHairGroupPublicData* HairGroupPublicData = reinterpret_cast<const FHairGroupPublicData*>(PrimitiveInfo.Mesh->Elements[0].VertexFactoryUserData);
 			check(HairGroupPublicData);
 
-			const FHairGroupPublicData::FVertexFactoryInput& VFInput = HairGroupPublicData->VFInput;
-			const FMatrix44d LocalToWorldPrimitiveTransform = VFInput.LocalToWorldTransform.ToMatrixWithScale(); 
-
 			const bool bCullingEnable = bSupportCulling && GHairVisibilityComputeRaster_Culling ? HairGroupPublicData->GetCullingResultAvailable() : false;
 			const bool bCullingEnableInRasterizers = bCullingEnable && !bClassification;
 
 			const uint32 PointCount = HairGroupPublicData->GetActiveStrandsPointCount();
-			check(PointCount == HairGroupPublicData->VFInput.Strands.PointCount);
+			check(PointCount == HairGroupPublicData->VFInput.Strands.Common.PointCount);
 			const float CoverageScale = HairGroupPublicData->GetActiveStrandsCoverageScale();
 
 			// HW/SW classification
@@ -3882,7 +3869,7 @@ class FHairStrandsPositionChangedCS : public FGlobalShader
 	BEGIN_SHADER_PARAMETER_STRUCT(FParameters, )
 		SHADER_PARAMETER(uint32, VertexCount)
 		SHADER_PARAMETER(float, PositionThreshold2)
-		SHADER_PARAMETER(uint32, HairStrandsVF_bIsCullingEnable)
+		SHADER_PARAMETER(uint32, HairStrandsVF_bCullingEnable)
 		SHADER_PARAMETER(uint32, bDrawInvalidElement)
 		SHADER_PARAMETER_RDG_BUFFER_SRV(Buffer<uint>, HairStrandsVF_CullingIndexBuffer)
 		SHADER_PARAMETER_RDG_BUFFER_SRV(Buffer<uint>, HairStrandsVF_CullingIndirectBuffer)
@@ -3921,7 +3908,7 @@ static void AddHairStrandsHasPositionChangedPass(
 	Parameters->VertexCount = PointCount;
 	Parameters->PositionThreshold2 = FMath::Square(GHairStrands_InvalidationPosition_Threshold);
 	Parameters->bDrawInvalidElement = GHairStrands_InvalidationPosition_Debug > 0 ? 1u : 0u;
-	Parameters->HairStrandsVF_bIsCullingEnable = 0u;
+	Parameters->HairStrandsVF_bCullingEnable = 0u;
 	Parameters->HairStrandsVF_CullingIndexBuffer = GraphBuilder.CreateSRV(GSystemTextures.GetDefaultBuffer(GraphBuilder, 4, 0u), PF_R32_UINT);
 	Parameters->HairStrandsVF_CullingIndirectBuffer = Parameters->HairStrandsVF_CullingIndexBuffer;
 	Parameters->CurrPositionBuffer = HairGroupPublicData->VFInput.Strands.PositionBuffer.SRV;
@@ -3934,7 +3921,7 @@ static void AddHairStrandsHasPositionChangedPass(
 	{
 		Parameters->HairStrandsVF_CullingIndexBuffer = Register(GraphBuilder, HairGroupPublicData->GetCulledVertexIdBuffer(), ERDGImportedBufferFlags::CreateSRV).SRV;
 		Parameters->HairStrandsVF_CullingIndirectBuffer = Register(GraphBuilder, HairGroupPublicData->GetDrawIndirectRasterComputeBuffer(), ERDGImportedBufferFlags::CreateSRV).SRV;
-		Parameters->HairStrandsVF_bIsCullingEnable = 1;
+		Parameters->HairStrandsVF_bCullingEnable = 1;
 	}
 
 	const FIntVector DispatchCount(FMath::DivideAndRoundUp(PointCount, FHairStrandsPositionChangedCS::GetGroupSize()), 1, 1);
diff --git a/Engine/Source/Runtime/Renderer/Private/HairStrands/HairStrandsVoxelization.cpp b/Engine/Source/Runtime/Renderer/Private/HairStrands/HairStrandsVoxelization.cpp
--- a/Engine/Source/Runtime/Renderer/Private/HairStrands/HairStrandsVoxelization.cpp
+++ b/Engine/Source/Runtime/Renderer/Private/HairStrands/HairStrandsVoxelization.cpp
@@ -1161,9 +1161,6 @@ static void AddVirtualVoxelizationComputeRasterPass(
 				const uint32 PointCount = HairGroupPublicData->GetActiveStrandsPointCount();
 				const float CoverageScale = HairGroupPublicData->GetActiveStrandsCoverageScale();
 
-				FTransform LocalToTranslatedWorldTransform = VFInput.LocalToWorldTransform;
-				LocalToTranslatedWorldTransform.AddToTranslation(TranslatedWorldOffset);
-
 				FVoxelRasterComputeCS::FParameters* PassParameters = GraphBuilder.AllocParameters<FVoxelRasterComputeCS::FParameters>();
 				PassParameters->MaxRasterCount = FMath::Clamp(GHairStrandsVoxelComputeRasterMaxVoxelCount, 1, 256);
 				PassParameters->VirtualVoxelParams = VoxelResources.Parameters.Common;
@@ -1180,7 +1177,7 @@ static void AddVirtualVoxelizationComputeRasterPass(
 
 				if (bCullingEnable)
 				{
-					FComputeShaderUtils::AddPass(GraphBuilder, RDG_EVENT_NAME("HairStrands::VoxelComputeRaster(culling=on)"), ComputeShader_CullingOn, PassParameters, PassParameters->HairInstance.HairStrandsVF_CullingIndirectBufferArgs, 0);
+					FComputeShaderUtils::AddPass(GraphBuilder, RDG_EVENT_NAME("HairStrands::VoxelComputeRaster(culling=on)"), ComputeShader_CullingOn, PassParameters, PassParameters->HairInstance.HairStrandsVF.Culling.CullingIndirectBufferArgs, 0);
 				}
 				else
 				{
diff --git a/Engine/Source/Runtime/Renderer/Public/HairStrandsInterface.h b/Engine/Source/Runtime/Renderer/Public/HairStrandsInterface.h
--- a/Engine/Source/Runtime/Renderer/Public/HairStrandsInterface.h
+++ b/Engine/Source/Runtime/Renderer/Public/HairStrandsInterface.h
@@ -20,6 +20,92 @@
 class UTexture2D;
 class FSceneInterface;
 
+////////////////////////////////////////////////////////////////////////////////////////////////////
+// Shader parameters
+
+// Instance common parameters
+BEGIN_SHADER_PARAMETER_STRUCT(FHairStrandsInstanceCommonParameters, RENDERER_API)
+	SHADER_PARAMETER(float,  Density)
+	SHADER_PARAMETER(float,  Radius)
+	SHADER_PARAMETER(float,  RootScale)
+	SHADER_PARAMETER(float,  TipScale)
+	SHADER_PARAMETER(float,  Length)
+	SHADER_PARAMETER(float,  LengthScale)
+	SHADER_PARAMETER(float,  RaytracingRadiusScale)
+	SHADER_PARAMETER(uint32, GroupIndex)
+	SHADER_PARAMETER(uint32, PointCount)
+	SHADER_PARAMETER(uint32, CurveCount)
+	SHADER_PARAMETER(uint32, RaytracingProceduralSplits)
+	SHADER_PARAMETER(uint32, bRaytracingGeometry)
+	SHADER_PARAMETER(uint32, bStableRasterization)
+	SHADER_PARAMETER(uint32, bScatterSceneLighting)
+	SHADER_PARAMETER(FVector3f, PositionOffset)
+	SHADER_PARAMETER(FVector3f, PrevPositionOffset)
+	SHADER_PARAMETER(FMatrix44f, LocalToWorldPrimitiveTransform)
+	SHADER_PARAMETER(FMatrix44f, LocalToTranslatedWorldPrimitiveTransform)
+	SHADER_PARAMETER_ARRAY(FUintVector4, AttributeOffsets, [HAIR_ATTRIBUTE_OFFSET_COUNT])
+END_SHADER_PARAMETER_STRUCT()
+
+// Instance resources (RDG)
+BEGIN_SHADER_PARAMETER_STRUCT(FHairStrandsInstanceResourceParameters, RENDERER_API)
+	SHADER_PARAMETER_RDG_BUFFER_SRV(Buffer, PositionBuffer)
+	SHADER_PARAMETER_RDG_BUFFER_SRV(Buffer, PositionOffsetBuffer)
+	SHADER_PARAMETER_RDG_BUFFER_SRV(Buffer, CurveBuffer)
+	SHADER_PARAMETER_RDG_BUFFER_SRV(Buffer, PointToCurveBuffer)
+	SHADER_PARAMETER_RDG_BUFFER_SRV(ByteAddressBuffer, AttributeBuffer)
+END_SHADER_PARAMETER_STRUCT()
+
+// Instance prev. resources (RDG)
+BEGIN_SHADER_PARAMETER_STRUCT(FHairStrandsInstancePrevResourceParameters, RENDERER_API)
+	SHADER_PARAMETER_RDG_BUFFER_SRV(Buffer, PrevPositionBuffer)
+	SHADER_PARAMETER_RDG_BUFFER_SRV(Buffer, PrevPositionOffsetBuffer)
+END_SHADER_PARAMETER_STRUCT()
+
+// Instance culling resources (RDG)
+BEGIN_SHADER_PARAMETER_STRUCT(FHairStrandsInstanceCullingParameters, RENDERER_API)
+	SHADER_PARAMETER(uint32, bCullingEnable)
+	SHADER_PARAMETER_RDG_BUFFER_SRV(Buffer, CullingIndirectBuffer)
+	SHADER_PARAMETER_RDG_BUFFER_SRV(Buffer, CullingIndexBuffer)
+	SHADER_PARAMETER_RDG_BUFFER_SRV(Buffer, CullingRadiusScaleBuffer)
+	RDG_BUFFER_ACCESS(CullingIndirectBufferArgs, ERHIAccess::IndirectArgs)
+END_SHADER_PARAMETER_STRUCT()
+
+// Instance resources (Raw)
+BEGIN_SHADER_PARAMETER_STRUCT(FHairStrandsInstanceResourceRawParameters, RENDERER_API)
+	SHADER_PARAMETER_SRV(Buffer<uint4>, PositionBuffer)
+	SHADER_PARAMETER_SRV(Buffer<float4>, PositionOffsetBuffer)
+	SHADER_PARAMETER_SRV(Buffer<uint>, CurveBuffer)
+	SHADER_PARAMETER_SRV(Buffer<uint>, PointToCurveBuffer)
+	SHADER_PARAMETER_SRV(Buffer<float4>, TangentBuffer)
+	SHADER_PARAMETER_SRV(ByteAddressBuffer, AttributeBuffer)
+END_SHADER_PARAMETER_STRUCT()
+
+// Instance prev. resources (Raw)
+BEGIN_SHADER_PARAMETER_STRUCT(FHairStrandsInstancePrevResourceRawParameters, RENDERER_API)
+	SHADER_PARAMETER_SRV(Buffer<uint4>, PreviousPositionBuffer)
+	SHADER_PARAMETER_SRV(Buffer<float4>, PreviousPositionOffsetBuffer)
+END_SHADER_PARAMETER_STRUCT()
+
+// Instance culling resources (Raw)
+BEGIN_SHADER_PARAMETER_STRUCT(FHairStrandsInstanceCullingRawParameters, RENDERER_API)
+	SHADER_PARAMETER(uint32, bCullingEnable)
+//	SHADER_PARAMETER_SRV(Buffer<uint>, CullingIndirectBuffer)
+	SHADER_PARAMETER_SRV(Buffer<uint>, CullingIndexBuffer)
+	SHADER_PARAMETER_SRV(Buffer<float>, CullingRadiusScaleBuffer)
+END_SHADER_PARAMETER_STRUCT()
+
+// Intermediate struct which can be referenced by FHairStrandsInstanceParameters for getting 
+// the HairStrandsVF_ decoration in shader
+BEGIN_SHADER_PARAMETER_STRUCT(FHairStrandsInstanceIntermediateParameters, )
+	SHADER_PARAMETER_STRUCT_INCLUDE(FHairStrandsInstanceCommonParameters, Common)
+	SHADER_PARAMETER_STRUCT_INCLUDE(FHairStrandsInstanceResourceParameters, Resources)
+	SHADER_PARAMETER_STRUCT_INCLUDE(FHairStrandsInstanceCullingParameters, Culling)
+END_SHADER_PARAMETER_STRUCT()
+
+BEGIN_SHADER_PARAMETER_STRUCT(FHairStrandsInstanceParameters, )
+	SHADER_PARAMETER_STRUCT(FHairStrandsInstanceIntermediateParameters, HairStrandsVF)
+END_SHADER_PARAMETER_STRUCT()
+
 ////////////////////////////////////////////////////////////////////////////////////////////////////
 // Utils buffers for importing/exporting hair resources
 
@@ -229,22 +315,7 @@ public:
 			FShaderResourceViewRHIRef PrevPositionOffsetBufferRHISRV	= nullptr;
 			FShaderResourceViewRHIRef CurveBufferRHISRV					= nullptr;
 
-			FVector PositionOffset = FVector::ZeroVector;
-			FVector PrevPositionOffset = FVector::ZeroVector;
-			TArray<uint32> AttributeOffsets;
-
-			uint32 CurveCount = 0;
-			uint32 PointCount = 0;
-			float HairRadius = 0;
-			float HairRootScale = 0;
-			float HairTipScale = 0;
-			float HairRaytracingRadiusScale = 0;
-			float HairLengthScale = 1.f;
-			float HairLength = 0;
-			float HairDensity = 0;
-			bool bUseStableRasterization = false;
-			bool bScatterSceneLighting = false;
-			bool bUseRaytracingGeometry = false;
+			FHairStrandsInstanceCommonParameters Common; 
 		} Strands;
 
 		struct FCards
