Convert groom texture and groom deformer to unified hair instance paramters.

diff --git a/Engine/Plugins/Runtime/HairStrands/Shaders/Private/DeformerDataInterfaceGroom.ush b/Engine/Plugins/Runtime/HairStrands/Shaders/Private/DeformerDataInterfaceGroom.ush
--- a/Engine/Plugins/Runtime/HairStrands/Shaders/Private/DeformerDataInterfaceGroom.ush
+++ b/Engine/Plugins/Runtime/HairStrands/Shaders/Private/DeformerDataInterfaceGroom.ush
@@ -3,29 +3,18 @@
 #include "/Plugin/ComputeFramework/Private/ComputeKernelCommon.ush"
 #include "/Engine/Private/HairStrands/HairStrandsVertexFactoryCommon.ush"
 #include "/Engine/Private/HairStrands/HairStrandsAttributeCommon.ush"
+
 ///////////////////////////////////////////////////////////////////////////////////////////////////
 // Bindings
 
-uint				{DataInterfaceName}_NumControlPoints;
-uint				{DataInterfaceName}_NumCurves;
-float				{DataInterfaceName}_VF_HairRadius;
-float				{DataInterfaceName}_VF_HairLength;
-float				{DataInterfaceName}_VF_RootScale;
-float				{DataInterfaceName}_VF_TipScale;
-uint				{DataInterfaceName}_VF_GroupIndex;
-uint4				{DataInterfaceName}_VF_AttributeOffsets[HAIR_ATTRIBUTE_OFFSET_COUNT];
-Buffer<float4>		{DataInterfaceName}_PositionOffsetBuffer;
-Buffer<uint4>		{DataInterfaceName}_PositionBuffer;
-ByteAddressBuffer	{DataInterfaceName}_AttributeBuffer;
-Buffer<uint>		{DataInterfaceName}_VertexToCurveBuffer;
-Buffer<uint>		{DataInterfaceName}_CurveBuffer;
+HAIR_STRANDS_INSTANCE_PARAMETERS({DataInterfaceName})
 
 ///////////////////////////////////////////////////////////////////////////////////////////////////
 // Helpers
 
 FHairControlPoint ReadControlPointData(uint ControlPointIndex)
 {
-	const uint MaxControlPoints = {DataInterfaceName}_NumControlPoints;
+	const uint MaxControlPoints = {DataInterfaceName}_PointCount;
 	if (ControlPointIndex >= MaxControlPoints)
 	{
 		return (FHairControlPoint)0;
@@ -35,14 +24,14 @@ FHairControlPoint ReadControlPointData(uint ControlPointIndex)
 	return UnpackHairControlPoint(
 		{DataInterfaceName}_PositionBuffer[ControlPointIndex], 
 		{DataInterfaceName}_PositionOffsetBuffer[0],
-		{DataInterfaceName}_VF_HairRadius, 
+		{DataInterfaceName}_Radius, 
 		1.0f /*RootScale*/, 
 		1.0f /*TipScale*/);
 }
 
 FHairCurve ReadCurveData(uint CurveIndex)
 {
-	const uint MaxCurve = {DataInterfaceName}_NumCurves;
+	const uint MaxCurve = {DataInterfaceName}_CurveCount;
 	if (CurveIndex >= MaxCurve)
 	{
 		return (FHairCurve)0;
@@ -56,12 +45,12 @@ FHairCurve ReadCurveData(uint CurveIndex)
 
 uint ReadNumControlPoints_{DataInterfaceName}()
 {
-	return {DataInterfaceName}_NumControlPoints;
+	return {DataInterfaceName}_PointCount;
 }
 
 uint ReadNumCurves_{DataInterfaceName}()
 {
-	return {DataInterfaceName}_NumCurves;
+	return {DataInterfaceName}_CurveCount;
 }
 
 float3 ReadPosition_{DataInterfaceName}(uint ControlPointIndex)
@@ -82,18 +71,18 @@ float ReadCoordU_{DataInterfaceName}(uint ControlPointIndex)
 float ReadLength_{DataInterfaceName}(uint ControlPointIndex)
 {
 	const float CoordU      = ReadControlPointData(ControlPointIndex).UCoord;
-	const float CurveLength = UnpackHairLength(InternalGetHairStrandsAttributesPerCurve(ControlPointIndex, HAIR_ATTRIBUTE_OFFSET_LENGTH({DataInterfaceName}_VF_AttributeOffsets), HAIR_ATTRIBUTE_STRIDE_LENGTH, {DataInterfaceName}_AttributeBuffer, {DataInterfaceName}_VertexToCurveBuffer));
+	const float CurveLength = UnpackHairLength(InternalGetHairStrandsAttributesPerCurve(ControlPointIndex, HAIR_ATTRIBUTE_OFFSET_LENGTH({DataInterfaceName}_AttributeOffsets), HAIR_ATTRIBUTE_STRIDE_LENGTH, {DataInterfaceName}_AttributeBuffer, {DataInterfaceName}_PointToCurveBuffer));
 	return CurveLength * CoordU;
 }
 
 float2 ReadRootUV_{DataInterfaceName}(uint ControlPointIndex)
 {
-	return UnpackHairRootUV(InternalGetHairStrandsAttributesPerCurve(ControlPointIndex, HAIR_ATTRIBUTE_OFFSET_ROOTUV({DataInterfaceName}_VF_AttributeOffsets), HAIR_ATTRIBUTE_STRIDE_ROOTUV, {DataInterfaceName}_AttributeBuffer, {DataInterfaceName}_VertexToCurveBuffer));
+	return UnpackHairRootUV(InternalGetHairStrandsAttributesPerCurve(ControlPointIndex, HAIR_ATTRIBUTE_OFFSET_ROOTUV({DataInterfaceName}_AttributeOffsets), HAIR_ATTRIBUTE_STRIDE_ROOTUV, {DataInterfaceName}_AttributeBuffer, {DataInterfaceName}_PointToCurveBuffer));
 }
 
 float ReadSeed_{DataInterfaceName}(uint ControlPointIndex)
 {
-	return UnpackHairSeed(InternalGetHairStrandsAttributesPerCurve(ControlPointIndex, HAIR_ATTRIBUTE_OFFSET_SEED({DataInterfaceName}_VF_AttributeOffsets), HAIR_ATTRIBUTE_STRIDE_SEED, {DataInterfaceName}_AttributeBuffer, {DataInterfaceName}_VertexToCurveBuffer));
+	return UnpackHairSeed(InternalGetHairStrandsAttributesPerCurve(ControlPointIndex, HAIR_ATTRIBUTE_OFFSET_SEED({DataInterfaceName}_AttributeOffsets), HAIR_ATTRIBUTE_STRIDE_SEED, {DataInterfaceName}_AttributeBuffer, {DataInterfaceName}_PointToCurveBuffer));
 }
 
 float ReadCurveOffsetPoint_{DataInterfaceName}(uint CurveIndex)
diff --git a/Engine/Plugins/Runtime/HairStrands/Source/HairStrandsCore/Private/Deformer/DeformerDataInterfaceGroom.cpp b/Engine/Plugins/Runtime/HairStrands/Source/HairStrandsCore/Private/Deformer/DeformerDataInterfaceGroom.cpp
--- a/Engine/Plugins/Runtime/HairStrands/Source/HairStrandsCore/Private/Deformer/DeformerDataInterfaceGroom.cpp
+++ b/Engine/Plugins/Runtime/HairStrands/Source/HairStrandsCore/Private/Deformer/DeformerDataInterfaceGroom.cpp
@@ -95,19 +95,8 @@ void UOptimusGroomDataInterface::GetSupportedInputs(TArray<FShaderFunctionDefini
 }
 
 BEGIN_SHADER_PARAMETER_STRUCT(FGroomDataInterfaceParameters, )
-	SHADER_PARAMETER(uint32, NumControlPoints)
-	SHADER_PARAMETER(uint32, NumCurves)
-	SHADER_PARAMETER(float, VF_HairRadius)
-	SHADER_PARAMETER(float, VF_HairLength)
-	SHADER_PARAMETER(float, VF_RootScale)
-	SHADER_PARAMETER(float, VF_TipScale)
-	SHADER_PARAMETER(uint32, VF_GroupIndex)
-	SHADER_PARAMETER_ARRAY(FUintVector4, VF_AttributeOffsets, [HAIR_ATTRIBUTE_OFFSET_COUNT])
-	SHADER_PARAMETER_RDG_BUFFER_SRV(Buffer<float4>,    PositionOffsetBuffer)
-	SHADER_PARAMETER_RDG_BUFFER_SRV(Buffer<uint4>,     PositionBuffer)
-	SHADER_PARAMETER_RDG_BUFFER_SRV(ByteAddressBuffer, AttributeBuffer)
-	SHADER_PARAMETER_RDG_BUFFER_SRV(Buffer<uint>,      VertexToCurveBuffer)
-	SHADER_PARAMETER_RDG_BUFFER_SRV(Buffer<uint>,      CurveBuffer)
+	SHADER_PARAMETER_STRUCT_INCLUDE(FHairStrandsInstanceCommonParameters, Common)
+	SHADER_PARAMETER_STRUCT_INCLUDE(FHairStrandsInstanceResourceParameters, Resources)
 END_SHADER_PARAMETER_STRUCT()
 
 void UOptimusGroomDataInterface::GetShaderParameters(TCHAR const* UID, FShaderParametersMetadataBuilder& InOutBuilder, FShaderParametersMetadataAllocations& InOutAllocations) const
@@ -179,18 +168,23 @@ bool FOptimusGroomDataProviderProxy::IsValid(FValidationData const& InValidation
 void FOptimusGroomDataProviderProxy::AllocateResources(FRDGBuilder& GraphBuilder)
 {
 	Resources.Empty();
+	FallbackSRV = nullptr;
 	const uint32 InstanceCount = GroomComponent ? GroomComponent->GetGroupCount() : 0;
 	for (uint32 Index =0; Index <InstanceCount;++Index)
 	{
 		if (FHairGroupInstance* Instance = GroomComponent->GetGroupInstance(Index))
 		{
-			FResources& R = Resources.AddDefaulted_GetRef();
-			R.PositionOffsetSRV = Register(GraphBuilder, Instance->Strands.RestResource->PositionOffsetBuffer, ERDGImportedBufferFlags::CreateSRV).SRV;
-			R.PositionSRV		= Register(GraphBuilder, Instance->Strands.RestResource->PositionBuffer, ERDGImportedBufferFlags::CreateSRV).SRV;
-			R.AttributeSRV		= Register(GraphBuilder, Instance->Strands.RestResource->AttributeBuffer, ERDGImportedBufferFlags::CreateSRV).SRV;
-			R.VertexToCurveSRV	= Register(GraphBuilder, Instance->Strands.RestResource->PointToCurveBuffer, ERDGImportedBufferFlags::CreateSRV).SRV;
-			R.CurveSRV			= Register(GraphBuilder, Instance->Strands.RestResource->CurveBuffer, ERDGImportedBufferFlags::CreateSRV).SRV;
-			R.FallbackSRV		= GraphBuilder.CreateSRV(GraphBuilder.RegisterExternalBuffer(GWhiteVertexBufferWithRDG->Buffer), PF_R16G16B16A16_UINT);
+			FHairStrandsInstanceResourceParameters& R = Resources.AddDefaulted_GetRef();
+			R.PositionBuffer		= RegisterAsSRV(GraphBuilder, Instance->Strands.RestResource->PositionBuffer);
+			R.PositionOffsetBuffer 	= RegisterAsSRV(GraphBuilder, Instance->Strands.RestResource->PositionOffsetBuffer);
+			R.CurveBuffer			= RegisterAsSRV(GraphBuilder, Instance->Strands.RestResource->CurveBuffer);
+			R.PointToCurveBuffer	= RegisterAsSRV(GraphBuilder, Instance->Strands.RestResource->PointToCurveBuffer);
+			R.AttributeBuffer		= RegisterAsSRV(GraphBuilder, Instance->Strands.RestResource->AttributeBuffer);
+
+			if (FallbackSRV)
+			{
+				FallbackSRV = GraphBuilder.CreateSRV(GraphBuilder.RegisterExternalBuffer(GWhiteVertexBufferWithRDG->Buffer), PF_R16G16B16A16_UINT);
+			}
 		}
 	}
 }
@@ -207,38 +201,26 @@ void FOptimusGroomDataProviderProxy::GatherDispatchData(FDispatchData const& InD
 	{
 		if (FHairGroupInstance* Instance = GroomComponent->GetGroupInstance(InvocationIndex))
 		{
-			const bool bIsSRVValid = Resources[InvocationIndex].PositionSRV != nullptr;
+			const bool bIsSRVValid = Resources[InvocationIndex].PositionBuffer != nullptr;
 			const int32 NumControlPoints = bIsSRVValid ? Instance->Strands.Data->GetNumPoints() : 0;
 			const int32 NumCurves = bIsSRVValid ? Instance->Strands.Data->GetNumCurves() : 0;
 
 			const FHairGroupPublicData::FVertexFactoryInput VFInput = ComputeHairStrandsVertexInputData(Instance, EGroomViewMode::None);
 			
 			FParameters& Parameters = ParameterArray[InvocationIndex];
-			Parameters.NumControlPoints = NumControlPoints;
-			Parameters.NumCurves = NumCurves;
-			Parameters.VF_HairRadius = VFInput.Strands.Common.Radius;
-			Parameters.VF_HairLength = VFInput.Strands.Common.Length;
-			Parameters.VF_RootScale  = VFInput.Strands.Common.RootScale;
-			Parameters.VF_TipScale   = VFInput.Strands.Common.TipScale;
-			Parameters.VF_GroupIndex = VFInput.Strands.Common.GroupIndex;// Instance->Debug.GroupIndex;
-			Parameters.VF_AttributeOffsets = VFInput.Strands.Common.AttributeOffsets;
-			//PACK_HAIR_ATTRIBUTE_OFFSETS(Parameters.VF_AttributeOffsets, VFInput.Strands.AttributeOffsets);
+			Parameters.Common = VFInput.Strands.Common;
 
 			if (bIsSRVValid)
 			{
-				Parameters.PositionOffsetBuffer = Resources[InvocationIndex].PositionOffsetSRV;
-				Parameters.PositionBuffer		= Resources[InvocationIndex].PositionSRV;
-				Parameters.AttributeBuffer		= Resources[InvocationIndex].AttributeSRV;
-				Parameters.VertexToCurveBuffer	= Resources[InvocationIndex].VertexToCurveSRV;
-				Parameters.CurveBuffer			= Resources[InvocationIndex].CurveSRV;
+				Parameters.Resources = Resources[InvocationIndex];
 			}
 			else
 			{
-				Parameters.PositionOffsetBuffer = Resources[InvocationIndex].FallbackSRV;
-				Parameters.PositionBuffer		= Resources[InvocationIndex].FallbackSRV;
-				Parameters.AttributeBuffer		= Resources[InvocationIndex].FallbackSRV;
-				Parameters.VertexToCurveBuffer	= Resources[InvocationIndex].FallbackSRV;
-				Parameters.CurveBuffer			= Resources[InvocationIndex].FallbackSRV;
+				Parameters.Resources.PositionBuffer			= FallbackSRV;
+				Parameters.Resources.PositionOffsetBuffer 	= FallbackSRV;
+				Parameters.Resources.AttributeBuffer		= FallbackSRV;
+				Parameters.Resources.PointToCurveBuffer		= FallbackSRV;
+				Parameters.Resources.CurveBuffer			= FallbackSRV;
 			}
 		}
 	}
diff --git a/Engine/Plugins/Runtime/HairStrands/Source/HairStrandsCore/Private/Deformer/DeformerDataInterfaceGroom.h b/Engine/Plugins/Runtime/HairStrands/Source/HairStrandsCore/Private/Deformer/DeformerDataInterfaceGroom.h
--- a/Engine/Plugins/Runtime/HairStrands/Source/HairStrandsCore/Private/Deformer/DeformerDataInterfaceGroom.h
+++ b/Engine/Plugins/Runtime/HairStrands/Source/HairStrandsCore/Private/Deformer/DeformerDataInterfaceGroom.h
@@ -4,6 +4,7 @@
 
 #include "RenderGraphDefinitions.h"
 #include "OptimusComputeDataInterface.h"
+#include "HairStrandsInterface.h"
 #include "ComputeFramework/ComputeDataProvider.h"
 #include "DeformerDataInterfaceGroom.generated.h"
 
@@ -67,14 +68,7 @@ private:
 	using FParameters = FGroomDataInterfaceParameters;
 
 	UGroomComponent* GroomComponent = nullptr; // Should it be HairInstance instead?
-	struct FResources
-	{
-		FRDGBufferSRVRef PositionOffsetSRV = nullptr;
-		FRDGBufferSRVRef PositionSRV = nullptr;
-		FRDGBufferSRVRef AttributeSRV = nullptr;
-		FRDGBufferSRVRef VertexToCurveSRV = nullptr;
-		FRDGBufferSRVRef FallbackSRV = nullptr;
-		FRDGBufferSRVRef CurveSRV = nullptr;
-	};
-	TArray<FResources> Resources;
+	TArray<FHairStrandsInstanceResourceParameters> Resources;
+	FRDGBufferSRVRef FallbackSRV;
+
 };
diff --git a/Engine/Plugins/Runtime/HairStrands/Source/HairStrandsCore/Private/GroomTextureBuilder.cpp b/Engine/Plugins/Runtime/HairStrands/Source/HairStrandsCore/Private/GroomTextureBuilder.cpp
--- a/Engine/Plugins/Runtime/HairStrands/Source/HairStrandsCore/Private/GroomTextureBuilder.cpp
+++ b/Engine/Plugins/Runtime/HairStrands/Source/HairStrandsCore/Private/GroomTextureBuilder.cpp
@@ -617,6 +617,12 @@ class FHairStrandsTextureVS : public FGlobalShader
 	}
 };
 
+BEGIN_SHADER_PARAMETER_STRUCT(FHairStrandsInstanceRawParameters, )
+	SHADER_PARAMETER_STRUCT_INCLUDE(FHairStrandsInstanceCommonParameters, Common)
+	SHADER_PARAMETER_STRUCT_INCLUDE(FHairStrandsInstanceResourceRawParameters, Resources)
+	SHADER_PARAMETER_STRUCT_INCLUDE(FHairStrandsInstanceCullingRawParameters, Culling)
+END_SHADER_PARAMETER_STRUCT()
+
 class FHairStrandsTexturePS : public FGlobalShader
 {
 	DECLARE_GLOBAL_SHADER(FHairStrandsTexturePS);
@@ -628,18 +634,11 @@ class FHairStrandsTexturePS : public FGlobalShader
 		SHADER_PARAMETER(uint32, VertexCount)
 		SHADER_PARAMETER(float, MaxDistance)
 		SHADER_PARAMETER(int32, TracingDirection)
-		SHADER_PARAMETER_ARRAY(FUintVector4, InVF_AttributeOffsets, [HAIR_ATTRIBUTE_OFFSET_COUNT])
+		SHADER_PARAMETER_STRUCT(FHairStrandsInstanceRawParameters, InVF)
 
 		SHADER_PARAMETER(uint32, UVsChannelIndex)
 		SHADER_PARAMETER(uint32, UVsChannelCount)
 
-		SHADER_PARAMETER(float, InVF_Radius)
-		SHADER_PARAMETER(FVector3f, InVF_PositionOffset)
-		SHADER_PARAMETER_SRV(Buffer, InVF_PositionBuffer)
-		SHADER_PARAMETER_SRV(ByteAddressBuffer, InVF_AttributeBuffer)
-		SHADER_PARAMETER_SRV(Buffer, InVF_VertexToCurveBuffer)
-		SHADER_PARAMETER(uint32, InVF_GroupIndex)
-
 		SHADER_PARAMETER(FVector3f, Voxel_MinBound)
 		SHADER_PARAMETER(FVector3f, Voxel_MaxBound)
 		SHADER_PARAMETER(FIntVector, Voxel_Resolution)
@@ -763,15 +762,7 @@ static void InternalGenerateHairStrandsTextures(
 	FRDGBufferRef VoxelOffsetAndCount,
 	FRDGBufferRef VoxelData,
 	
-	FRHIShaderResourceView* InHairStrands_PositionBuffer,
-	FRHIShaderResourceView* InHairStrands_AttributeBuffer,
-	FRHIShaderResourceView* InHairStrands_VertexToCurveBuffer,
-	const FVector& InHairStrands_PositionOffset,
-	float InHairStrands_Radius,
-	float InHairStrands_Length,
-	uint32 InHairStrands_ControlPointCount,
-	uint32 InHairStrands_GroupIndex,
-	const uint32* InHairStrands_AttributeOffsets,
+	FHairStrandsInstanceRawParameters& Instance,
 	FHairStrandsRDGTextures& Out)
 {
 	const FIntPoint OutputResolution = Out.DepthTexture->Desc.Extent;
@@ -789,14 +780,7 @@ static void InternalGenerateHairStrandsTextures(
 	ParametersPS->UVsChannelIndex = UVsChannelIndex;
 	ParametersPS->UVsChannelCount = UVsChannelCount;
 
-	ParametersPS->InVF_PositionBuffer = InHairStrands_PositionBuffer;
-	ParametersPS->InVF_PositionOffset = (FVector3f)InHairStrands_PositionOffset;
-	ParametersPS->InVF_AttributeBuffer = InHairStrands_AttributeBuffer;
-	ParametersPS->InVF_VertexToCurveBuffer = InHairStrands_VertexToCurveBuffer;
-	PACK_HAIR_ATTRIBUTE_OFFSETS(ParametersPS->InVF_AttributeOffsets, InHairStrands_AttributeOffsets);
-
-	ParametersPS->InVF_Radius = InHairStrands_Radius;
-	ParametersPS->InVF_GroupIndex = InHairStrands_GroupIndex;
+	ParametersPS->InVF = Instance;
 	
 	ParametersPS->Voxel_MinBound = VoxelMinBound;
 	ParametersPS->Voxel_MaxBound = VoxelMaxBound;
@@ -1063,6 +1047,23 @@ static bool TraceTextures(
 
 			// Ensure the rest resources are loaded when rendering the strands textures
 			GroupData.Strands.RestResource->Allocate(GraphBuilder, EHairResourceLoadingType::Sync);
+			
+			// Manually setting up parameters, since we don't have a real instance
+			FHairStrandsInstanceRawParameters Instance;
+			Instance.Common.GroupIndex = GroupIndex;
+			Instance.Common.PointCount = GroupData.Strands.RestResource->BulkData.GetNumPoints();
+			Instance.Common.CurveCount = GroupData.Strands.RestResource->BulkData.GetNumCurves();
+			Instance.Common.Radius = RenderingData.GeometrySettings.HairWidth * 0.5f;
+			Instance.Common.RootScale = 1.0f;
+			Instance.Common.TipScale = 1.0f;
+			Instance.Common.LengthScale = 1.0f;
+			Instance.Common.Length = GroupData.Strands.RestResource->BulkData.MaxLength;
+			Instance.Common.PositionOffset = FVector3f(GroupData.Strands.RestResource->GetPositionOffset());
+			PACK_HAIR_ATTRIBUTE_OFFSETS(Instance.Common.AttributeOffsets, GroupData.Strands.RestResource->BulkData.AttributeOffsets);
+
+			Instance.Resources.PositionBuffer = GroupData.Strands.RestResource->PositionBuffer.SRV;
+			Instance.Resources.AttributeBuffer = GroupData.Strands.RestResource->AttributeBuffer.SRV;
+			Instance.Resources.PointToCurveBuffer = GroupData.Strands.RestResource->PointToCurveBuffer.SRV;
 
 			InternalGenerateHairStrandsTextures(
 				GraphBuilder,
@@ -1092,17 +1093,7 @@ static bool TraceTextures(
 				GroupData.Debug.Resource->VoxelDescription.MaxSegmentPerVoxel,
 				VoxelOffsetAndCount,
 				VoxelData,
-
-				GroupData.Strands.RestResource->PositionBuffer.SRV,
-				GroupData.Strands.RestResource->AttributeBuffer.SRV,
-				GroupData.Strands.RestResource->PointToCurveBuffer.SRV,
-				GroupData.Strands.RestResource->GetPositionOffset(),
-				RenderingData.GeometrySettings.HairWidth * 0.5f,
-				GroupData.Strands.RestResource->BulkData.MaxLength,
-				GroupData.Strands.RestResource->BulkData.GetNumPoints(),
-				GroupIndex,
-				GroupData.Strands.RestResource->BulkData.AttributeOffsets,
-
+				Instance,
 				Output);
 
 
diff --git a/Engine/Shaders/Private/HairStrands/HairStrandsTexturesGeneration.usf b/Engine/Shaders/Private/HairStrands/HairStrandsTexturesGeneration.usf
--- a/Engine/Shaders/Private/HairStrands/HairStrandsTexturesGeneration.usf
+++ b/Engine/Shaders/Private/HairStrands/HairStrandsTexturesGeneration.usf
@@ -59,10 +59,11 @@ void MainVS(
 
 #if SHADER_PIXEL
 
-#define HAIR_STRANDS_PARAMETERS 0
 #include "HairStrandsVertexFactoryCommon.ush"
 #include "HairStrandsAttributeCommon.ush"
 
+HAIR_STRANDS_INSTANCE_PARAMETERS(InVF)
+
 float3 ComputeTangent(
 	const FHairControlPoint CP0,
 	const FHairControlPoint CP1,
@@ -85,15 +86,6 @@ uint VertexCount;
 uint PrimitiveCount;
 uint bHasMaterialData;
 
-Buffer<uint4>		InVF_PositionBuffer;
-ByteAddressBuffer	InVF_AttributeBuffer;
-Buffer<uint>		InVF_VertexToCurveBuffer;
-uint4				InVF_AttributeOffsets[HAIR_ATTRIBUTE_OFFSET_COUNT];
-
-float3 InVF_PositionOffset;
-float InVF_Radius;
-uint InVF_GroupIndex;
-
 Buffer<float> VertexBuffer;
 Buffer<float4> NormalsBuffer;
 Buffer<float2> UVsBuffer;
@@ -145,9 +137,6 @@ bool Trace(
 {
 	bool bIsValid = false;
 
-	const float InVF_RootScale = 1.0f;
-	const float InVF_TipScale = 1.0f;
-
 	const float3 WP0 = InRayO;
 	const float3 WP1 = InRayO + InRayD * InMaxDistance;
 
@@ -298,8 +287,8 @@ bool Trace(
 							OutTangent = Tangent;
 							OutDepth = Distance;
 
-							OutSeed         = UnpackHairSeed(InternalGetHairStrandsAttributesPerCurve(P0, HAIR_ATTRIBUTE_OFFSET_SEED(InVF_AttributeOffsets), HAIR_ATTRIBUTE_STRIDE_SEED, InVF_AttributeBuffer, InVF_VertexToCurveBuffer));
-							OutRootUV       = UnpackHairRootUV(InternalGetHairStrandsAttributesPerCurve(P0, HAIR_ATTRIBUTE_OFFSET_ROOTUV(InVF_AttributeOffsets), HAIR_ATTRIBUTE_STRIDE_ROOTUV, InVF_AttributeBuffer, InVF_VertexToCurveBuffer));
+							OutSeed         = UnpackHairSeed(InternalGetHairStrandsAttributesPerCurve(P0, HAIR_ATTRIBUTE_OFFSET_SEED(InVF_AttributeOffsets), HAIR_ATTRIBUTE_STRIDE_SEED, InVF_AttributeBuffer, InVF_PointToCurveBuffer));
+							OutRootUV       = UnpackHairRootUV(InternalGetHairStrandsAttributesPerCurve(P0, HAIR_ATTRIBUTE_OFFSET_ROOTUV(InVF_AttributeOffsets), HAIR_ATTRIBUTE_STRIDE_ROOTUV, InVF_AttributeBuffer, InVF_PointToCurveBuffer));
 							OutU			= lerp(CP_Curr.UCoord, CP_Next.UCoord, VertexU);
 							OutMaterial.xyz = 0;
 							OutMaterial.w   = 0;
diff --git a/Engine/Shaders/Private/HairStrands/HairStrandsVertexFactoryCommon.ush b/Engine/Shaders/Private/HairStrands/HairStrandsVertexFactoryCommon.ush
--- a/Engine/Shaders/Private/HairStrands/HairStrandsVertexFactoryCommon.ush
+++ b/Engine/Shaders/Private/HairStrands/HairStrandsVertexFactoryCommon.ush
@@ -146,31 +146,33 @@ FHairCurve UnpackHairCurve(uint In)
 ///////////////////////////////////////////////////////////////////////////////////////////////////
 // Hair strands instance parameters
 
-#if defined(HAIR_STRANDS_PARAMETERS) && HAIR_STRANDS_PARAMETERS == 1
-uint			HairStrandsVF_bCullingEnable;
-uint			HairStrandsVF_bRaytracingGeometry;
-uint			HairStrandsVF_RaytracingProceduralSplits;
-float			HairStrandsVF_Radius;
-float			HairStrandsVF_RootScale;
-float			HairStrandsVF_TipScale;
-float			HairStrandsVF_Length;
-float			HairStrandsVF_Density;
-uint			HairStrandsVF_bStableRasterization;
-float3			HairStrandsVF_PositionOffset;
-uint			HairStrandsVF_PointCount;
-uint			HairStrandsVF_CurveCount;
-float4x4		HairStrandsVF_LocalToWorldPrimitiveTransform;
-float4x4		HairStrandsVF_LocalToTranslatedWorldPrimitiveTransform;
-Buffer<uint4>	HairStrandsVF_PositionBuffer;
-Buffer<float4>	HairStrandsVF_PositionOffsetBuffer;
-Buffer<uint>	HairStrandsVF_CurveBuffer;
-uint4			HairStrandsVF_AttributeOffsets[HAIR_ATTRIBUTE_OFFSET_COUNT];
-
-#if defined(PERMUTATION_CULLING) && PERMUTATION_CULLING == 1
-Buffer<uint>	HairStrandsVF_CullingIndirectBuffer;
-Buffer<uint>	HairStrandsVF_CullingIndexBuffer;
-Buffer<float>	HairStrandsVF_CullingRadiusScaleBuffer;
-Buffer<uint>	HairStrandsVF_CullingIndirectBufferArgs;
-#endif
+#define HAIR_STRANDS_INSTANCE_PARAMETERS(Prefix) \
+	uint				Prefix##_bCullingEnable; \
+	uint				Prefix##_bRaytracingGeometry; \
+	uint				Prefix##_RaytracingProceduralSplits; \
+	uint				Prefix##_GroupIndex; \
+	float				Prefix##_Radius; \
+	float				Prefix##_RootScale; \
+	float				Prefix##_TipScale; \
+	float				Prefix##_Length; \
+	float				Prefix##_Density; \
+	uint				Prefix##_bStableRasterization; \
+	float3				Prefix##_PositionOffset; \
+	uint				Prefix##_PointCount; \
+	uint				Prefix##_CurveCount; \
+	float4x4			Prefix##_LocalToWorldPrimitiveTransform; \
+	float4x4			Prefix##_LocalToTranslatedWorldPrimitiveTransform; \
+	uint4				Prefix##_AttributeOffsets[HAIR_ATTRIBUTE_OFFSET_COUNT]; \
+	Buffer<uint4>		Prefix##_PositionBuffer; \
+	Buffer<float4>		Prefix##_PositionOffsetBuffer; \
+	Buffer<uint>		Prefix##_CurveBuffer; \
+	Buffer<uint>		Prefix##_PointToCurveBuffer; \
+	ByteAddressBuffer 	Prefix##_AttributeBuffer; \
+	Buffer<uint>		Prefix##_CullingIndirectBuffer; \
+	Buffer<uint>		Prefix##_CullingIndexBuffer; \
+	Buffer<float>		Prefix##_CullingRadiusScaleBuffer; \
+	Buffer<uint>		Prefix##_CullingIndirectBufferArgs;
 
-#endif // HAIR_STRANDS_PARAMETERS
\ No newline at end of file
+#if defined(HAIR_STRANDS_PARAMETERS) && HAIR_STRANDS_PARAMETERS == 1
+	HAIR_STRANDS_INSTANCE_PARAMETERS(HairStrandsVF)
+#endif
\ No newline at end of file
