Rename VertexToCurve into PointToCurve for better semantic

diff --git a/Engine/Plugins/Runtime/HairStrands/Shaders/Private/NiagaraInterpolateGroomGuides.usf b/Engine/Plugins/Runtime/HairStrands/Shaders/Private/NiagaraInterpolateGroomGuides.usf
--- a/Engine/Plugins/Runtime/HairStrands/Shaders/Private/NiagaraInterpolateGroomGuides.usf
+++ b/Engine/Plugins/Runtime/HairStrands/Shaders/Private/NiagaraInterpolateGroomGuides.usf
@@ -26,7 +26,7 @@ Buffer<uint>		RootToUniqueTriangleIndexBuffer;
 
 #endif
 
-Buffer<uint>		VertexToCurveIndexBuffer;
+Buffer<uint>		PointToCurveIndexBuffer;
 RWBuffer<uint4>		DeformedPositionBuffer;
 Buffer<uint4>		RestPositionBuffer;
 Buffer<uint>		CurvesOffsetsBuffer;
@@ -74,7 +74,7 @@ void MainCS(uint3 DispatchThreadId : SV_DispatchThreadID)
 {
 	if (DispatchThreadId.x < uint(NumPoints))
 	{
-		const int StrandIndex = VertexToCurveIndexBuffer[DispatchThreadId.x];
+		const int StrandIndex = PointToCurveIndexBuffer[DispatchThreadId.x];
 		
 		const int PointCount = CurvesOffsetsBuffer[StrandIndex + 1] - CurvesOffsetsBuffer[StrandIndex];
 		const int PointIndex = DispatchThreadId.x - CurvesOffsetsBuffer[StrandIndex];
diff --git a/Engine/Plugins/Runtime/HairStrands/Source/HairStrandsCore/Private/GroomBuilder.cpp b/Engine/Plugins/Runtime/HairStrands/Source/HairStrandsCore/Private/GroomBuilder.cpp
--- a/Engine/Plugins/Runtime/HairStrands/Source/HairStrandsCore/Private/GroomBuilder.cpp
+++ b/Engine/Plugins/Runtime/HairStrands/Source/HairStrandsCore/Private/GroomBuilder.cpp
@@ -167,15 +167,15 @@ namespace HairStrandsBuilder
 		OutPackedCurves.SetNum(NumCurves * FHairStrandsCurveFormat::ComponentCount);
 
 		const bool bUse16bitsCurveIndex = NumCurves < 65536;
-		TArray<FHairStrandsVertexToCurveFormat16::Type> OutVertexToCurve16;
-		TArray<FHairStrandsVertexToCurveFormat32::Type> OutVertexToCurve32;
+		TArray<FHairStrandsPointToCurveFormat16::Type> OutPointToCurve16;
+		TArray<FHairStrandsPointToCurveFormat32::Type> OutPointToCurve32;
 		if (bUse16bitsCurveIndex)
 		{
-			OutVertexToCurve16.SetNum(NumPoints);
+			OutPointToCurve16.SetNum(NumPoints);
 		}
 		else
 		{
-			OutVertexToCurve32.SetNum(NumPoints);
+			OutPointToCurve32.SetNum(NumPoints);
 		}
 
 		const uint32 Attributes = HairStrands.GetAttributes();
@@ -246,11 +246,11 @@ namespace HairStrandsBuilder
 				// Vertex to Curve
 				if (bUse16bitsCurveIndex)
 				{
-					OutVertexToCurve16[PointIndex + IndexOffset] = CurveIndex;
+					OutPointToCurve16[PointIndex + IndexOffset] = CurveIndex;
 				}
 				else
 				{
-					OutVertexToCurve32[PointIndex + IndexOffset] = CurveIndex;
+					OutPointToCurve32[PointIndex + IndexOffset] = CurveIndex;
 				}
 
 				// Per-Vertex Color
@@ -429,11 +429,11 @@ namespace HairStrandsBuilder
 		if (bUse16bitsCurveIndex)
 		{
 			OutBulkData.Flags |= FHairStrandsBulkData::DataFlags_Has16bitsCurveIndex;
-			CopyToBulkData<FHairStrandsVertexToCurveFormat16>(OutBulkData.VertexToCurve, OutVertexToCurve16);
+			CopyToBulkData<FHairStrandsPointToCurveFormat16>(OutBulkData.PointToCurve, OutPointToCurve16);
 		}
 		else
 		{
-			CopyToBulkData<FHairStrandsVertexToCurveFormat32>(OutBulkData.VertexToCurve, OutVertexToCurve32);
+			CopyToBulkData<FHairStrandsPointToCurveFormat32>(OutBulkData.PointToCurve, OutPointToCurve32);
 		}
 	}
 
diff --git a/Engine/Plugins/Runtime/HairStrands/Source/HairStrandsCore/Private/GroomRBFDeformer.cpp b/Engine/Plugins/Runtime/HairStrands/Source/HairStrandsCore/Private/GroomRBFDeformer.cpp
--- a/Engine/Plugins/Runtime/HairStrands/Source/HairStrandsCore/Private/GroomRBFDeformer.cpp
+++ b/Engine/Plugins/Runtime/HairStrands/Source/HairStrandsCore/Private/GroomRBFDeformer.cpp
@@ -81,7 +81,7 @@ FVector3f DisplacePosition(
 
 void DeformStrands(	
 	const FHairStrandsDatas& HairStandsData,
-	const TArray<FHairStrandsIndexFormat::Type>& VertexToCurveIndexBuffer,
+	const TArray<FHairStrandsIndexFormat::Type>& PointToCurveBuffer,
 	const TArray<FHairStrandsIndexFormat::Type>& RootToUniqueTriangleBuffer,
 	const TArray<FHairStrandsRootBarycentricFormat::Type>& RootBarycentricBuffer,
 
@@ -117,7 +117,7 @@ void DeformStrands(
 	ParallelFor(CurveCount, [&](uint32 CurveIndex)
 	{
 		const uint32 VertexOffset = HairStandsData.StrandsCurves.CurvesOffset[CurveIndex];
-		const uint32 RootIndex = VertexToCurveIndexBuffer[VertexOffset];
+		const uint32 RootIndex = PointToCurveBuffer[VertexOffset];
 		const uint32 TriangleIndex = RootToUniqueTriangleBuffer[RootIndex];
 
 		// Sanity check
@@ -152,7 +152,7 @@ void DeformStrands(
 	// Apply correction offset to each control points
 	ParallelFor(VertexCount, [&](uint32 VertexIndex)
 	{
-		const uint32 CurveIndex = VertexToCurveIndexBuffer[VertexIndex];
+		const uint32 CurveIndex = PointToCurveBuffer[VertexIndex];
 		const FVector4f CorrectionOffset = CorrectionOffsets[CurveIndex];
 		OutDeformedPositionBuffer[VertexIndex] += CorrectionOffset;
 	});
@@ -206,7 +206,7 @@ TArray<FVector3f> GetDeformedHairStrandsPositions(
 	const FHairStrandsDatas& HairStrandsData,
 	const uint32 MeshLODIndex,
 	const FSkeletalMeshRenderData* InMeshRenderData,
-	const TArray<FHairStrandsIndexFormat::Type>& VertexToCurveIndexBuffer,
+	const TArray<FHairStrandsIndexFormat::Type>& PointToCurveBuffer,
 	const FHairStrandsRootData::FMeshProjectionLOD& RestLODData)
 {
 	// Init the mesh samples with the target mesh vertices
@@ -258,7 +258,7 @@ TArray<FVector3f> GetDeformedHairStrandsPositions(
 
 	DeformStrands(
 		HairStrandsData,
-		VertexToCurveIndexBuffer,
+		PointToCurveBuffer,
 		RootToUniqueTriangleBuffer,
 		RestLODData.RootBarycentricBuffer,
 
@@ -521,7 +521,7 @@ void FGroomRBFDeformer::GetRBFDeformedGroomAsset(const UGroomAsset* InGroomAsset
 		}
 
 		// Build Curve index for every vertices
-		auto BuildVertexToCurveMapping = [](const FHairStrandsDatas& In, TArray<uint32>& Out)
+		auto BuildPointToCurveMapping = [](const FHairStrandsDatas& In, TArray<uint32>& Out)
 		{
 			const uint32 CurveCount = In.GetNumCurves();
 			Out.SetNum(In.GetNumPoints());
@@ -551,10 +551,10 @@ void FGroomRBFDeformer::GetRBFDeformedGroomAsset(const UGroomAsset* InGroomAsset
 			FHairGroupInfo DummyInfo;
 			FGroomBuilder::BuildData(Group, InGroomAsset->HairGroupsInterpolation[GroupIndex], DummyInfo, StrandsData, GuidesData);
 
-			TArray<uint32> SimRootDataVertexToCurveIndexBuffer;
-			TArray<uint32> RenRootDataVertexToCurveIndexBuffer;
-			BuildVertexToCurveMapping(GuidesData, SimRootDataVertexToCurveIndexBuffer);
-			BuildVertexToCurveMapping(StrandsData, RenRootDataVertexToCurveIndexBuffer);
+			TArray<uint32> SimRootDataPointToCurveBuffer;
+			TArray<uint32> RenRootDataPointToCurveBuffer;
+			BuildPointToCurveMapping(GuidesData, SimRootDataPointToCurveBuffer);
+			BuildPointToCurveMapping(StrandsData, RenRootDataPointToCurveBuffer);
 
 			// Get deformed guides
 			// If the groom override the value, we output dummy value for the guides, since they won't be used
@@ -573,7 +573,7 @@ void FGroomRBFDeformer::GetRBFDeformedGroomAsset(const UGroomAsset* InGroomAsset
 					GuidesData,
 					MeshLODIndex,
 					SkeletalMeshData_Target,
-					SimRootDataVertexToCurveIndexBuffer,
+					SimRootDataPointToCurveBuffer,
 					SimRootData.MeshProjectionLODs[MeshLODIndex]);
 			}
 
@@ -587,7 +587,7 @@ void FGroomRBFDeformer::GetRBFDeformedGroomAsset(const UGroomAsset* InGroomAsset
 					StrandsData,
 					MeshLODIndex,
 					SkeletalMeshData_Target,
-					RenRootDataVertexToCurveIndexBuffer,
+					RenRootDataPointToCurveBuffer,
 					RenRootData.MeshProjectionLODs[MeshLODIndex]);
 			}
 		}
diff --git a/Engine/Plugins/Runtime/HairStrands/Source/HairStrandsCore/Private/GroomResources.cpp b/Engine/Plugins/Runtime/HairStrands/Source/HairStrandsCore/Private/GroomResources.cpp
--- a/Engine/Plugins/Runtime/HairStrands/Source/HairStrandsCore/Private/GroomResources.cpp
+++ b/Engine/Plugins/Runtime/HairStrands/Source/HairStrandsCore/Private/GroomResources.cpp
@@ -875,7 +875,7 @@ bool FHairStrandsRestResource::InternalIsDataLoaded()
 		FBulkDataBatchRequest::FBatchBuilder Batch = FBulkDataBatchRequest::NewBatch(5);
 		Batch.Read(BulkData.Positions);
 		Batch.Read(BulkData.Attributes);
-		Batch.Read(BulkData.VertexToCurve);
+		Batch.Read(BulkData.PointToCurve);
 		Batch.Read(BulkData.Curves);
 
 		Batch.Issue(BulkDataRequest);
@@ -899,11 +899,11 @@ void FHairStrandsRestResource::InternalAllocate(FRDGBuilder& GraphBuilder)
 	InternalCreateVertexBufferRDG_FromBulkData<FHairStrandsCurveFormat>(GraphBuilder, BulkData.Curves, CurveCount, CurveBuffer, ToHairResourceDebugName(HAIRSTRANDS_RESOUCE_NAME(CurveType, Hair.StrandsRest_CurveBuffer), ResourceName), OwnerName, EHairResourceUsageType::Static);
 	if (!!(BulkData.Flags & BulkData.DataFlags_Has16bitsCurveIndex))
 	{
-		InternalCreateVertexBufferRDG_FromBulkData<FHairStrandsVertexToCurveFormat16>(GraphBuilder, BulkData.VertexToCurve, PointCount, PointToCurveBuffer, ToHairResourceDebugName(HAIRSTRANDS_RESOUCE_NAME(CurveType, Hair.StrandsRest_PointToCurveBuffer), ResourceName), OwnerName, EHairResourceUsageType::Static);
+		InternalCreateVertexBufferRDG_FromBulkData<FHairStrandsPointToCurveFormat16>(GraphBuilder, BulkData.PointToCurve, PointCount, PointToCurveBuffer, ToHairResourceDebugName(HAIRSTRANDS_RESOUCE_NAME(CurveType, Hair.StrandsRest_PointToCurveBuffer), ResourceName), OwnerName, EHairResourceUsageType::Static);
 	}
 	else
 	{
-		InternalCreateVertexBufferRDG_FromBulkData<FHairStrandsVertexToCurveFormat32>(GraphBuilder, BulkData.VertexToCurve, PointCount, PointToCurveBuffer, ToHairResourceDebugName(HAIRSTRANDS_RESOUCE_NAME(CurveType, Hair.StrandsRest_PointToCurveBuffer), ResourceName), OwnerName, EHairResourceUsageType::Static);
+		InternalCreateVertexBufferRDG_FromBulkData<FHairStrandsPointToCurveFormat32>(GraphBuilder, BulkData.PointToCurve, PointCount, PointToCurveBuffer, ToHairResourceDebugName(HAIRSTRANDS_RESOUCE_NAME(CurveType, Hair.StrandsRest_PointToCurveBuffer), ResourceName), OwnerName, EHairResourceUsageType::Static);
 	}
 
 	TArray<FVector4f> RestOffset;
diff --git a/Engine/Plugins/Runtime/HairStrands/Source/HairStrandsCore/Private/HairStrandsDatas.cpp b/Engine/Plugins/Runtime/HairStrands/Source/HairStrandsCore/Private/HairStrandsDatas.cpp
--- a/Engine/Plugins/Runtime/HairStrands/Source/HairStrandsCore/Private/HairStrandsDatas.cpp
+++ b/Engine/Plugins/Runtime/HairStrands/Source/HairStrandsCore/Private/HairStrandsDatas.cpp
@@ -229,8 +229,8 @@ void FHairStrandsBulkData::Serialize(FArchive& Ar, UObject* Owner)
 {
 	static_assert(sizeof(FHairStrandsPositionFormat::BulkType) == sizeof(FHairStrandsPositionFormat::Type));
 	static_assert(sizeof(FHairStrandsAttributeFormat::BulkType) == sizeof(FHairStrandsAttributeFormat::Type));
-	static_assert(sizeof(FHairStrandsVertexToCurveFormat16::BulkType) == sizeof(FHairStrandsVertexToCurveFormat16::Type));
-	static_assert(sizeof(FHairStrandsVertexToCurveFormat32::BulkType) == sizeof(FHairStrandsVertexToCurveFormat32::Type));
+	static_assert(sizeof(FHairStrandsPointToCurveFormat16::BulkType) == sizeof(FHairStrandsPointToCurveFormat16::Type));
+	static_assert(sizeof(FHairStrandsPointToCurveFormat32::BulkType) == sizeof(FHairStrandsPointToCurveFormat32::Type));
 	static_assert(sizeof(FHairStrandsRootIndexFormat::BulkType) == sizeof(FHairStrandsRootIndexFormat::Type)); 
 
 	Ar.UsingCustomVersion(FReleaseObjectVersion::GUID);
@@ -256,7 +256,7 @@ void FHairStrandsBulkData::Serialize(FArchive& Ar, UObject* Owner)
 		const uint32 BulkFlags = BULKDATA_Force_NOT_InlinePayload;
 		Positions.SetBulkDataFlags(BulkFlags);
 		Attributes.SetBulkDataFlags(BulkFlags);
-		VertexToCurve.SetBulkDataFlags(BulkFlags);
+		PointToCurve.SetBulkDataFlags(BulkFlags);
 		Curves.SetBulkDataFlags(BulkFlags);
 	}
 
@@ -267,7 +267,7 @@ void FHairStrandsBulkData::Serialize(FArchive& Ar, UObject* Owner)
 
 		Positions.Serialize(Ar, Owner, ChunkIndex, bAttemptFileMapping);
 		Attributes.Serialize(Ar, Owner, ChunkIndex, bAttemptFileMapping);
-		VertexToCurve.Serialize(Ar, Owner, ChunkIndex, bAttemptFileMapping);
+		PointToCurve.Serialize(Ar, Owner, ChunkIndex, bAttemptFileMapping);
 		Curves.Serialize(Ar, Owner, ChunkIndex, bAttemptFileMapping);
 	}
 }
@@ -287,13 +287,13 @@ void FHairStrandsBulkData::Reset()
 	// Deallocate memory if needed
 	Positions.RemoveBulkData();
 	Attributes.RemoveBulkData();
-	VertexToCurve.RemoveBulkData();
+	PointToCurve.RemoveBulkData();
 	Curves.RemoveBulkData();
 
 	// Reset the bulk byte buffer to ensure the (serialize) data size is reset to 0
 	Positions 		= FByteBulkData();
 	Attributes		= FByteBulkData();
-	VertexToCurve	= FByteBulkData();
+	PointToCurve	= FByteBulkData();
 	Curves			= FByteBulkData();
 }
 
diff --git a/Engine/Plugins/Runtime/HairStrands/Source/HairStrandsCore/Private/HairStrandsRendering.cpp b/Engine/Plugins/Runtime/HairStrands/Source/HairStrandsCore/Private/HairStrandsRendering.cpp
--- a/Engine/Plugins/Runtime/HairStrands/Source/HairStrandsCore/Private/HairStrandsRendering.cpp
+++ b/Engine/Plugins/Runtime/HairStrands/Source/HairStrandsCore/Private/HairStrandsRendering.cpp
@@ -304,7 +304,7 @@ class FDeformGuideCS : public FGlobalShader
 		SHADER_PARAMETER_RDG_BUFFER_SRV(Buffer<float4>, SimDeformedPosition2Buffer)
 
 		SHADER_PARAMETER_RDG_BUFFER_SRV(Buffer<uint>, SimRootBarycentricBuffer)
-		SHADER_PARAMETER_RDG_BUFFER_SRV(Buffer<uint>, SimVertexToCurveBuffer)
+		SHADER_PARAMETER_RDG_BUFFER_SRV(Buffer<uint>, SimPointToCurveBuffer)
 		SHADER_PARAMETER_RDG_BUFFER_SRV(Buffer<uint>, SimRootToUniqueTriangleIndexBuffer)
 
 		SHADER_PARAMETER_RDG_BUFFER_SRV(Buffer, SimDeformedOffsetBuffer)
@@ -340,7 +340,7 @@ static void AddDeformSimHairStrandsPass(
 	FHairStrandsRestRootResource* SimRestRootResources,
 	FHairStrandsDeformedRootResource* SimDeformedRootResources,
 	FRDGBufferSRVRef SimRestPosePositionBuffer,
-	FRDGBufferSRVRef SimVertexToCurveBuffer,
+	FRDGBufferSRVRef SimPointToCurveBuffer,
 	FRDGImportedBuffer OutSimDeformedPositionBuffer,
 	const FVector& SimRestOffset,
 	FRDGBufferSRVRef SimDeformedOffsetBuffer,
@@ -375,10 +375,10 @@ static void AddDeformSimHairStrandsPass(
 
 	if (DeformationType == EDeformationType::OffsetGuide)
 	{
-		const bool bIsVertexToCurveBuffersValid = SimVertexToCurveBuffer != nullptr;
-		if (bIsVertexToCurveBuffersValid)
+		const bool bIsPointToCurveBuffersValid = SimPointToCurveBuffer != nullptr;
+		if (bIsPointToCurveBuffersValid)
 		{
-			Parameters->SimVertexToCurveBuffer = SimVertexToCurveBuffer;
+			Parameters->SimPointToCurveBuffer = SimPointToCurveBuffer;
 		}
 
 		const uint32 RootCount = SimRestRootResources ? SimRestRootResources->BulkData.RootCount : 0;
@@ -397,7 +397,7 @@ static void AddDeformSimHairStrandsPass(
 				MeshLODIndex < uint32(SimDeformedRootResources->LODs.Num()) &&
 				SimRestRootResources->LODs[MeshLODIndex].IsValid() &&
 				SimDeformedRootResources->LODs[MeshLODIndex].IsValid() &&
-				bIsVertexToCurveBuffersValid;
+				bIsPointToCurveBuffersValid;
 			
 			bool bSupportGlobalInterpolation = false;
 			if (bSupportDynamicMesh)
@@ -520,7 +520,7 @@ class FHairInterpolationCS : public FGlobalShader
 		SHADER_PARAMETER_RDG_BUFFER_SRV(Buffer<float4>, RenDeformedPosition2Buffer)
 
 		SHADER_PARAMETER_RDG_BUFFER_SRV(Buffer<uint>, RenRootBarycentricBuffer)
-		SHADER_PARAMETER_RDG_BUFFER_SRV(Buffer<uint>, RenVertexToCurveBuffer)
+		SHADER_PARAMETER_RDG_BUFFER_SRV(Buffer<uint>, RenPointToCurveBuffer)
 		SHADER_PARAMETER_RDG_BUFFER_SRV(Buffer<uint>, RenRootToUniqueTriangleIndexBuffer)
 
 		SHADER_PARAMETER_RDG_BUFFER_SRV(Buffer<float4>, SimRestPosition0Buffer)
@@ -532,7 +532,7 @@ class FHairInterpolationCS : public FGlobalShader
 		SHADER_PARAMETER_RDG_BUFFER_SRV(Buffer<float4>, SimDeformedPosition2Buffer)
 
 		SHADER_PARAMETER_RDG_BUFFER_SRV(Buffer<uint>, SimRootBarycentricBuffer)
-		SHADER_PARAMETER_RDG_BUFFER_SRV(Buffer<uint>, SimVertexToCurveBuffer)
+		SHADER_PARAMETER_RDG_BUFFER_SRV(Buffer<uint>, SimPointToCurveBuffer)
 		SHADER_PARAMETER_RDG_BUFFER_SRV(Buffer<uint>, SimRootToUniqueTriangleIndexBuffer)
 
 		SHADER_PARAMETER_RDG_BUFFER_SRV(Buffer<uint>, SimRootPointIndexBuffer)
@@ -588,7 +588,7 @@ static void AddHairStrandsInterpolationPass(
 	const FHairStrandsDeformedRootResource* RenDeformedRootResources,
 	const FHairStrandsDeformedRootResource* SimDeformedRootResources,
 	const FRDGBufferSRVRef& RenRestPosePositionBuffer,
-	const FRDGBufferSRVRef& RenVertexToCurveBuffer,
+	const FRDGBufferSRVRef& RenPointToCurveBuffer,
 	const bool bUseSingleGuide,
 	const FRDGBufferSRVRef& InterpolationBuffer,
 	const FRDGBufferSRVRef& Interpolation0Buffer,
@@ -596,7 +596,7 @@ static void AddHairStrandsInterpolationPass(
 	const FRDGBufferSRVRef& SimRestPosePositionBuffer,
 	const FRDGBufferSRVRef& SimDeformedPositionBuffer,
 	const FRDGBufferSRVRef& SimRootPointIndexBuffer,
-	const FRDGBufferSRVRef& SimVertexToCurveBuffer,
+	const FRDGBufferSRVRef& SimPointToCurveBuffer,
 	const FRDGBufferSRVRef& RenDeformerPositionBuffer,
 	FRDGImportedBuffer& OutRenPositionBuffer,
 	const FHairStrandsDeformedRootResource::FLOD::EFrameType DeformedFrame)
@@ -624,8 +624,8 @@ static void AddHairStrandsInterpolationPass(
 	Parameters->OutRenHairPositionOffsetBuffer = OutRenHairPositionOffsetBuffer;
 
 	Parameters->SimRootPointIndexBuffer = SimRootPointIndexBuffer;
-	Parameters->SimVertexToCurveBuffer = SimVertexToCurveBuffer;
-	Parameters->RenVertexToCurveBuffer = RenVertexToCurveBuffer;
+	Parameters->SimPointToCurveBuffer = SimPointToCurveBuffer;
+	Parameters->RenPointToCurveBuffer = RenPointToCurveBuffer;
 	
 	Parameters->LocalToWorldMatrix = FMatrix44f(Instance->LocalToWorld.ToMatrixWithScale());		// LWC_TODO: Precision loss
 	Parameters->HairLengthScale = HairLengthScale;
diff --git a/Engine/Plugins/Runtime/HairStrands/Source/HairStrandsCore/Private/Niagara/NiagaraDataInterfaceHairStrands.cpp b/Engine/Plugins/Runtime/HairStrands/Source/HairStrandsCore/Private/Niagara/NiagaraDataInterfaceHairStrands.cpp
--- a/Engine/Plugins/Runtime/HairStrands/Source/HairStrandsCore/Private/Niagara/NiagaraDataInterfaceHairStrands.cpp
+++ b/Engine/Plugins/Runtime/HairStrands/Source/HairStrandsCore/Private/Niagara/NiagaraDataInterfaceHairStrands.cpp
@@ -832,7 +832,7 @@ class FInterpolateGroomGuidesCS : public FGlobalShader
 		SHADER_PARAMETER_RDG_BUFFER_SRV(Buffer<uint>,	RootBarycentricCoordinatesBuffer)
 		SHADER_PARAMETER_RDG_BUFFER_SRV(Buffer<uint>,	RootToUniqueTriangleIndexBuffer)
 
-		SHADER_PARAMETER_RDG_BUFFER_SRV(Buffer<uint>,		VertexToCurveIndexBuffer)
+		SHADER_PARAMETER_RDG_BUFFER_SRV(Buffer<uint>,		PointToCurveIndexBuffer)
 		SHADER_PARAMETER_RDG_BUFFER_UAV(RWBuffer<uint4>,	DeformedPositionBuffer)
 		SHADER_PARAMETER_RDG_BUFFER_SRV(Buffer<uint4>,		RestPositionBuffer)
 		SHADER_PARAMETER_RDG_BUFFER_SRV(Buffer<uint>,		CurvesOffsetsBuffer)
@@ -933,7 +933,7 @@ static void InterpolateGroomGuides(FRDGBuilder& GraphBuilder, FNiagaraDataBuffer
 		PassParameters->DeformedPositionBuffer = RegisterAsUAV(GraphBuilder, HairStrandsBuffer->SourceDeformedResources->DeformedPositionBuffer[BufferIndex]);
 		PassParameters->RestPositionBuffer = RegisterAsSRV(GraphBuilder, HairStrandsBuffer->SourceRestResources->PositionBuffer);
 		
-		PassParameters->VertexToCurveIndexBuffer = RegisterAsSRV(GraphBuilder, HairStrandsBuffer->SourceRestResources->PointToCurveBuffer);
+		PassParameters->PointToCurveIndexBuffer = RegisterAsSRV(GraphBuilder, HairStrandsBuffer->SourceRestResources->PointToCurveBuffer);
 		PassParameters->CurvesOffsetsBuffer = RegisterAsSRV(GraphBuilder, HairStrandsBuffer->SourceRestResources->CurveBuffer);
 		if(bIsRootValid)
 		{
diff --git a/Engine/Plugins/Runtime/HairStrands/Source/HairStrandsCore/Public/GroomAsset.h b/Engine/Plugins/Runtime/HairStrands/Source/HairStrandsCore/Public/GroomAsset.h
--- a/Engine/Plugins/Runtime/HairStrands/Source/HairStrandsCore/Public/GroomAsset.h
+++ b/Engine/Plugins/Runtime/HairStrands/Source/HairStrandsCore/Public/GroomAsset.h
@@ -117,7 +117,7 @@ struct HAIRSTRANDSCORE_API FHairGroupData
 			Total += BulkData.Positions.IsBulkDataLoaded() 		? BulkData.Positions.GetBulkDataSize()   : 0;
 			Total += BulkData.Attributes.IsBulkDataLoaded()		? BulkData.Attributes.GetBulkDataSize() : 0;
 			Total += BulkData.Curves.IsBulkDataLoaded() 		? BulkData.Curves.GetBulkDataSize(): 0;
-			Total += BulkData.VertexToCurve.IsBulkDataLoaded()	? BulkData.VertexToCurve.GetBulkDataSize() : 0;
+			Total += BulkData.PointToCurve.IsBulkDataLoaded()	? BulkData.PointToCurve.GetBulkDataSize() : 0;
 			return Total;
 		}
 
diff --git a/Engine/Plugins/Runtime/HairStrands/Source/HairStrandsCore/Public/HairStrandsDatas.h b/Engine/Plugins/Runtime/HairStrands/Source/HairStrandsCore/Public/HairStrandsDatas.h
--- a/Engine/Plugins/Runtime/HairStrands/Source/HairStrandsCore/Public/HairStrandsDatas.h
+++ b/Engine/Plugins/Runtime/HairStrands/Source/HairStrandsCore/Public/HairStrandsDatas.h
@@ -115,7 +115,7 @@ struct FHairStrandsAttributeFormat
 	static const EPixelFormat Format = PF_R32_UINT;
 };
 
-struct FHairStrandsVertexToCurveFormat16
+struct FHairStrandsPointToCurveFormat16
 {
 	typedef uint16 Type;
 	typedef uint16 BulkType;
@@ -125,7 +125,7 @@ struct FHairStrandsVertexToCurveFormat16
 	static const EPixelFormat Format = PF_R16_UINT;
 };
 
-struct FHairStrandsVertexToCurveFormat32
+struct FHairStrandsPointToCurveFormat32
 {
 	typedef uint32 Type;
 	typedef uint32 BulkType;
@@ -491,7 +491,7 @@ struct HAIRSTRANDSCORE_API FHairStrandsBulkData
 
 	FByteBulkData Positions;	// Size = PointCount
 	FByteBulkData Attributes;	// Size = x*PointCount + y*CurveCount (depends on the stored attributes, which could be per-vertex(x) or per-curve(y))
-	FByteBulkData VertexToCurve;// Size = PointCount
+	FByteBulkData PointToCurve; // Size = PointCount
 	FByteBulkData Curves;		// Size = CurveCount
 };
 
diff --git a/Engine/Shaders/Private/HairStrands/HairStrandsGuideDeform.usf b/Engine/Shaders/Private/HairStrands/HairStrandsGuideDeform.usf
--- a/Engine/Shaders/Private/HairStrands/HairStrandsGuideDeform.usf
+++ b/Engine/Shaders/Private/HairStrands/HairStrandsGuideDeform.usf
@@ -31,7 +31,7 @@ Buffer<float4>  SimDeformedPosition1Buffer;
 Buffer<float4>  SimDeformedPosition2Buffer;
 
 Buffer<uint>	SimRootBarycentricBuffer;
-Buffer<uint>    SimVertexToCurveBuffer;
+Buffer<uint>    SimPointToCurveBuffer;
 Buffer<uint>	SimRootToUniqueTriangleIndexBuffer;
 
 // Return world position of a guide vertex deformed by a skinned triangle
@@ -77,7 +77,7 @@ float3 DisplacePosition(uint VertexIndex, float3 Pos)
 #endif
 
 #if PERMUTATION_DEFORMATION == DEFORMATION_OFFSET_GUIDE_DYNAMIC
-	const uint RootIndex = SimVertexToCurveBuffer[VertexIndex];
+	const uint RootIndex = SimPointToCurveBuffer[VertexIndex];
 	const float3 RootBarycentric = DecodeBarycentrics(SimRootBarycentricBuffer[RootIndex]);
 
 	const uint TriangleIndex = SimRootToUniqueTriangleIndexBuffer[RootIndex];
diff --git a/Engine/Shaders/Private/HairStrands/HairStrandsInterpolation.usf b/Engine/Shaders/Private/HairStrands/HairStrandsInterpolation.usf
--- a/Engine/Shaders/Private/HairStrands/HairStrandsInterpolation.usf
+++ b/Engine/Shaders/Private/HairStrands/HairStrandsInterpolation.usf
@@ -106,8 +106,8 @@ Buffer<uint>	HairStrandsVFTODO_CullingIndirectBufferArgs;
 #endif
 
 #if PERMUTATION_DYNAMIC_GEOMETRY >= 1 || PERMUTATION_DEBUG == 1 
-Buffer<uint>   RenVertexToCurveBuffer;
-Buffer<uint>   SimVertexToCurveBuffer;
+Buffer<uint>   RenPointToCurveBuffer;
+Buffer<uint>   SimPointToCurveBuffer;
 #endif
 
 #if PERMUTATION_DEFORMER
@@ -245,7 +245,7 @@ FRootTriangleData GetRootTriangleData(uint VertexIndex, bool bSim)
 	#if PERMUTATION_DYNAMIC_GEOMETRY >= 1 
 	if (bSim)
 	{
-		Out.RootIndex		 = SimVertexToCurveBuffer[VertexIndex];
+		Out.RootIndex		 = SimPointToCurveBuffer[VertexIndex];
 		Out.TriangleIndex	 = SimRootToUniqueTriangleIndexBuffer[Out.RootIndex];
 		Out.RootBarycentric  = DecodeBarycentrics(SimRootBarycentricBuffer[Out.RootIndex]);
 		Out.RestTriangle	 = GetTriangleTransformation(Out.TriangleIndex, SimRestPosition0Buffer, SimRestPosition1Buffer, SimRestPosition2Buffer);
@@ -253,7 +253,7 @@ FRootTriangleData GetRootTriangleData(uint VertexIndex, bool bSim)
 	}
 	else
 	{
-		Out.RootIndex		 = RenVertexToCurveBuffer[VertexIndex];
+		Out.RootIndex		 = RenPointToCurveBuffer[VertexIndex];
 		Out.TriangleIndex	 = RenRootToUniqueTriangleIndexBuffer[Out.RootIndex];
 		Out.RootBarycentric  = DecodeBarycentrics(RenRootBarycentricBuffer[Out.RootIndex]);
 		Out.RestTriangle	 = GetTriangleTransformation(Out.TriangleIndex, RenRestPosition0Buffer, RenRestPosition1Buffer, RenRestPosition2Buffer);
