Split hair attribute resources per frequency (curve/strands) to support streaming (with variable amount of attributes).

diff --git a/Engine/Plugins/Runtime/HairStrands/Source/HairStrandsCore/Private/Deformer/DeformerDataInterfaceGroom.cpp b/Engine/Plugins/Runtime/HairStrands/Source/HairStrandsCore/Private/Deformer/DeformerDataInterfaceGroom.cpp
--- a/Engine/Plugins/Runtime/HairStrands/Source/HairStrandsCore/Private/Deformer/DeformerDataInterfaceGroom.cpp
+++ b/Engine/Plugins/Runtime/HairStrands/Source/HairStrandsCore/Private/Deformer/DeformerDataInterfaceGroom.cpp
@@ -179,7 +179,8 @@ void FOptimusGroomDataProviderProxy::AllocateResources(FRDGBuilder& GraphBuilder
 			R.PositionOffsetBuffer 	= RegisterAsSRV(GraphBuilder, Instance->Strands.RestResource->PositionOffsetBuffer);
 			R.CurveBuffer			= RegisterAsSRV(GraphBuilder, Instance->Strands.RestResource->CurveBuffer);
 			R.PointToCurveBuffer	= RegisterAsSRV(GraphBuilder, Instance->Strands.RestResource->PointToCurveBuffer);
-			R.AttributeBuffer		= RegisterAsSRV(GraphBuilder, Instance->Strands.RestResource->AttributeBuffer);
+			R.CurveAttributeBuffer	= RegisterAsSRV(GraphBuilder, Instance->Strands.RestResource->CurveAttributeBuffer);
+			R.PointAttributeBuffer	= RegisterAsSRV(GraphBuilder, Instance->Strands.RestResource->PointAttributeBuffer);
 
 			if (FallbackSRV)
 			{
@@ -218,7 +219,8 @@ void FOptimusGroomDataProviderProxy::GatherDispatchData(FDispatchData const& InD
 			{
 				Parameters.Resources.PositionBuffer			= FallbackSRV;
 				Parameters.Resources.PositionOffsetBuffer 	= FallbackSRV;
-				Parameters.Resources.AttributeBuffer		= FallbackSRV;
+				Parameters.Resources.CurveAttributeBuffer	= FallbackSRV;
+				Parameters.Resources.PointAttributeBuffer	= FallbackSRV;
 				Parameters.Resources.PointToCurveBuffer		= FallbackSRV;
 				Parameters.Resources.CurveBuffer			= FallbackSRV;
 			}
diff --git a/Engine/Plugins/Runtime/HairStrands/Source/HairStrandsCore/Private/GroomBuilder.cpp b/Engine/Plugins/Runtime/HairStrands/Source/HairStrandsCore/Private/GroomBuilder.cpp
--- a/Engine/Plugins/Runtime/HairStrands/Source/HairStrandsCore/Private/GroomBuilder.cpp
+++ b/Engine/Plugins/Runtime/HairStrands/Source/HairStrandsCore/Private/GroomBuilder.cpp
@@ -40,7 +40,7 @@ static FAutoConsoleVariableRef CVarHairGroupIndexBuilder_MaxVoxelResolution(TEXT
 
 FString FGroomBuilder::GetVersion()
 {
-	return TEXT("v8h");
+	return TEXT("v8m");
 }
 
 namespace FHairStrandsDecimation
@@ -185,7 +185,7 @@ namespace HairStrandsBuilder
 		TArray<uint32> AttributeSeed;
 		TArray<uint32> AttributeLength;
 		TArray<uint32> AttributeClumpIDs;
-		TArray<uint32> AttributeBaseColor;
+		TArray<uint32> AttributeColor;
 		TArray<uint32> AttributeRoughness;
 		TArray<uint32> AttributeAO;
 
@@ -193,7 +193,7 @@ namespace HairStrandsBuilder
 		AttributeLength.SetNum(FMath::DivideAndRoundUp(NumCurves, 2u));		// 2 bytes encoding - Per-Curve
 		if (HasHairAttribute(Attributes, EHairAttribute::RootUV)) 	{ AttributeRootUV.SetNum(NumCurves); }																		// 4 bytes encoding - Per-Curve
 		if (HasHairAttribute(Attributes, EHairAttribute::ClumpID))	{ AttributeClumpIDs.SetNum(bHasMultipleClumpIDs ? NumCurves * 2u : FMath::DivideAndRoundUp(NumCurves, 2u));}// 8 or 2 bytes encoding - Per-Curve
-		if (HasHairAttribute(Attributes, EHairAttribute::Color))	{ AttributeBaseColor.SetNum(NumPoints); }																	// 4 bytes encoding - Per-Vertex
+		if (HasHairAttribute(Attributes, EHairAttribute::Color))	{ AttributeColor.SetNum(NumPoints); }																		// 4 bytes encoding - Per-Vertex
 		if (HasHairAttribute(Attributes, EHairAttribute::Roughness)){ AttributeRoughness.SetNum(FMath::DivideAndRoundUp(NumPoints, 2u)); }										// 1 bytes encoding - Per-Vertex
 		if (HasHairAttribute(Attributes, EHairAttribute::AO))		{ AttributeAO.SetNum(FMath::DivideAndRoundUp(NumPoints, 4u)); }												// 1 bytes encoding - Per-Vertex
 
@@ -256,7 +256,7 @@ namespace HairStrandsBuilder
 				// Per-Vertex Color
 				if (HasHairAttribute(Attributes, EHairAttribute::Color))
 				{
-					uint32& Packed = AttributeBaseColor[PointIndex + IndexOffset];
+					uint32& Packed = AttributeColor[PointIndex + IndexOffset];
 					Packed |= FMath::Clamp(uint32(FMath::Sqrt(Points.PointsBaseColor[PointIndex + IndexOffset].R) * 0x7FF), 0u, 0x7FFu);
 					Packed |= FMath::Clamp(uint32(FMath::Sqrt(Points.PointsBaseColor[PointIndex + IndexOffset].G) * 0x7FF), 0u, 0x7FFu)<<11;
 					Packed |= FMath::Clamp(uint32(FMath::Sqrt(Points.PointsBaseColor[PointIndex + IndexOffset].B) * 0x3FF), 0u, 0x3FFu)<<22;
@@ -355,38 +355,35 @@ namespace HairStrandsBuilder
 		OutBulkData.ImportedAttributes = HairStrands.GetAttributes();
 		OutBulkData.ImportedAttributeFlags = HairStrands.GetAttributeFlags();
 
-		// Concatenate all attributes
-		TArray<FHairStrandsAttributeFormat::Type> OutPackedAttributes;
+		const uint32 UintToByte = 4;
+
+		// Concatenate all curve-attributes
 		{
+			TArray<FHairStrandsAttributeFormat::Type> OutPackedAttributes;
 			OutPackedAttributes.Reserve(
 				AttributeRootUV.Num() + 
 				AttributeSeed.Num() +
 				AttributeLength.Num() +
-				AttributeClumpIDs.Num() +
-				AttributeBaseColor.Num() +
-				AttributeRoughness.Num() +
-				AttributeAO.Num());
-
-			const uint32 UintToByte = 4;
+				AttributeClumpIDs.Num());
 
-			for (uint32 AttributeIt=0; AttributeIt< HAIR_ATTRIBUTE_COUNT; ++AttributeIt)
+			for (uint32 AttributeIt=0; AttributeIt< HAIR_CURVE_ATTRIBUTE_COUNT; ++AttributeIt)
 			{
-				OutBulkData.AttributeOffsets[AttributeIt] = HAIR_ATTRIBUTE_INVALID_OFFSET;
+				OutBulkData.CurveAttributeOffsets[AttributeIt] = HAIR_ATTRIBUTE_INVALID_OFFSET;
 			}
 
 			if (HasHairAttribute(Attributes, EHairAttribute::RootUV))
 			{
-				OutBulkData.AttributeOffsets[HAIR_ATTRIBUTE_ROOTUV] = OutPackedAttributes.Num() * UintToByte;
+				OutBulkData.CurveAttributeOffsets[HAIR_CURVE_ATTRIBUTE_ROOTUV] = OutPackedAttributes.Num() * UintToByte;
 				OutPackedAttributes.Append(AttributeRootUV);
 			}
 
 			{
-				OutBulkData.AttributeOffsets[HAIR_ATTRIBUTE_SEED]   = OutPackedAttributes.Num() * UintToByte;
+				OutBulkData.CurveAttributeOffsets[HAIR_CURVE_ATTRIBUTE_SEED]   = OutPackedAttributes.Num() * UintToByte;
 				OutPackedAttributes.Append(AttributeSeed);
 			}
 
 			{
-				OutBulkData.AttributeOffsets[HAIR_ATTRIBUTE_LENGTH] = OutPackedAttributes.Num() * UintToByte;
+				OutBulkData.CurveAttributeOffsets[HAIR_CURVE_ATTRIBUTE_LENGTH] = OutPackedAttributes.Num() * UintToByte;
 				OutPackedAttributes.Append(AttributeLength);
 			}
 
@@ -394,37 +391,53 @@ namespace HairStrandsBuilder
 			{
 				if (HasHairAttributeFlags(OutBulkData.ImportedAttributeFlags, EHairAttributeFlags::HasMultipleClumpIDs))
 				{
-					OutBulkData.AttributeOffsets[HAIR_ATTRIBUTE_CLUMPID3] = OutPackedAttributes.Num() * UintToByte;
+					OutBulkData.CurveAttributeOffsets[HAIR_CURVE_ATTRIBUTE_CLUMPID3] = OutPackedAttributes.Num() * UintToByte;
 					OutPackedAttributes.Append(AttributeClumpIDs);
 				}
 				else
 				{
-					OutBulkData.AttributeOffsets[HAIR_ATTRIBUTE_CLUMPID] = OutPackedAttributes.Num() * UintToByte;
+					OutBulkData.CurveAttributeOffsets[HAIR_CURVE_ATTRIBUTE_CLUMPID] = OutPackedAttributes.Num() * UintToByte;
 					OutPackedAttributes.Append(AttributeClumpIDs);
 				}
 			}
 
+			CopyToBulkData<FHairStrandsAttributeFormat>(OutBulkData.CurveAttributes, OutPackedAttributes);
+		}
+
+		// Concatenate all point-attributes
+		{
+			TArray<FHairStrandsAttributeFormat::Type> OutPackedAttributes;
+			OutPackedAttributes.Reserve(
+				AttributeColor.Num() +
+				AttributeRoughness.Num() +
+				AttributeAO.Num());
+
+			for (uint32 AttributeIt=0; AttributeIt< HAIR_POINT_ATTRIBUTE_COUNT; ++AttributeIt)
+			{
+				OutBulkData.PointAttributeOffsets[AttributeIt] = HAIR_ATTRIBUTE_INVALID_OFFSET;
+			}
+			
 			if (HasHairAttribute(Attributes, EHairAttribute::Color))
 			{
-				OutBulkData.AttributeOffsets[HAIR_ATTRIBUTE_BASECOLOR] = OutPackedAttributes.Num() * UintToByte;
-				OutPackedAttributes.Append(AttributeBaseColor);
+				OutBulkData.PointAttributeOffsets[HAIR_POINT_ATTRIBUTE_COLOR] = OutPackedAttributes.Num() * UintToByte;
+				OutPackedAttributes.Append(AttributeColor);
 			}
 
 			if (HasHairAttribute(Attributes, EHairAttribute::Roughness))
 			{
-				OutBulkData.AttributeOffsets[HAIR_ATTRIBUTE_ROUGHNESS] = OutPackedAttributes.Num() * UintToByte;
+				OutBulkData.PointAttributeOffsets[HAIR_POINT_ATTRIBUTE_ROUGHNESS] = OutPackedAttributes.Num() * UintToByte;
 				OutPackedAttributes.Append(AttributeRoughness);
 			}
 
 			if (HasHairAttribute(Attributes, EHairAttribute::AO))
 			{
-				OutBulkData.AttributeOffsets[HAIR_ATTRIBUTE_AO] = OutPackedAttributes.Num() * UintToByte;
+				OutBulkData.PointAttributeOffsets[HAIR_POINT_ATTRIBUTE_AO] = OutPackedAttributes.Num() * UintToByte;
 				OutPackedAttributes.Append(AttributeAO);
 			}
+			CopyToBulkData<FHairStrandsAttributeFormat>(OutBulkData.PointAttributes, OutPackedAttributes);
 		}
 
 		CopyToBulkData<FHairStrandsPositionFormat>(OutBulkData.Positions, OutPackedPositions);
-		CopyToBulkData<FHairStrandsAttributeFormat>(OutBulkData.Attributes, OutPackedAttributes);
 		CopyToBulkData<FHairStrandsCurveFormat>(OutBulkData.Curves, OutPackedCurves);
 		if (bUse16bitsCurveIndex)
 		{
diff --git a/Engine/Plugins/Runtime/HairStrands/Source/HairStrandsCore/Private/GroomComponent.cpp b/Engine/Plugins/Runtime/HairStrands/Source/HairStrandsCore/Private/GroomComponent.cpp
--- a/Engine/Plugins/Runtime/HairStrands/Source/HairStrandsCore/Private/GroomComponent.cpp
+++ b/Engine/Plugins/Runtime/HairStrands/Source/HairStrandsCore/Private/GroomComponent.cpp
@@ -2937,7 +2937,7 @@ void UGroomComponent::DeleteDeferredHairGroupInstances()
 				#endif
 
 				#if WITH_EDITOR
-				LocalInstance->Strands.DebugAttributeBuffer.Release();
+				LocalInstance->Strands.DebugCurveAttributeBuffer.Release();
 				#endif
 
 				InternalResourceRelease(LocalInstance->Strands.VertexFactory);
diff --git a/Engine/Plugins/Runtime/HairStrands/Source/HairStrandsCore/Private/GroomDebug.cpp b/Engine/Plugins/Runtime/HairStrands/Source/HairStrandsCore/Private/GroomDebug.cpp
--- a/Engine/Plugins/Runtime/HairStrands/Source/HairStrandsCore/Private/GroomDebug.cpp
+++ b/Engine/Plugins/Runtime/HairStrands/Source/HairStrandsCore/Private/GroomDebug.cpp
@@ -967,15 +967,7 @@ static void AddHairDebugPrintInstancePass(
 					Data2.Y = Instance->HairGroupPublicData->GetActiveStrandsPointCount();
 					Data2.Z = Instance->HairGroupPublicData->GetActiveStrandsCurveCount();
 
-					for (uint32 AttributeIt=0; AttributeIt<uint32(EHairAttribute::Count); ++AttributeIt)
-					{
-						const EHairAttribute Attribute = (EHairAttribute)AttributeIt;
-						const uint32 Index = GetHairAttributeIndex(Attribute);
-						if (Index < HAIR_ATTRIBUTE_COUNT)
-						{
-							Data2.W |= Instance->Strands.Data->AttributeOffsets[Index] != 0xFFFFFFFF ? (1u<<AttributeIt) : 0u;
-						}
-					}
+					Data2.W = Instance->Strands.Data->ImportedAttributes;
 				}
 				
 				Data3 = FUintVector4(0);
diff --git a/Engine/Plugins/Runtime/HairStrands/Source/HairStrandsCore/Private/GroomResources.cpp b/Engine/Plugins/Runtime/HairStrands/Source/HairStrandsCore/Private/GroomResources.cpp
--- a/Engine/Plugins/Runtime/HairStrands/Source/HairStrandsCore/Private/GroomResources.cpp
+++ b/Engine/Plugins/Runtime/HairStrands/Source/HairStrandsCore/Private/GroomResources.cpp
@@ -862,7 +862,7 @@ void FHairMeshesDeformedResource::InternalRelease()
 
 FHairStrandsRestResource::FHairStrandsRestResource(FHairStrandsBulkData& InBulkData, EHairStrandsResourcesType InCurveType, const FHairResourceName& InResourceName, const FName& InOwnerName) :
 	FHairCommonResource(EHairStrandsAllocationType::Deferred, InResourceName, InOwnerName),
-	PositionBuffer(), AttributeBuffer(), PointToCurveBuffer(), BulkData(InBulkData), CurveType(InCurveType)
+	PositionBuffer(), PointAttributeBuffer(), CurveAttributeBuffer(), PointToCurveBuffer(), BulkData(InBulkData), CurveType(InCurveType)
 {
 	// Sanity check
 	check(!!(BulkData.Flags & FHairStrandsBulkData::DataFlags_HasData));
@@ -874,7 +874,8 @@ bool FHairStrandsRestResource::InternalIsDataLoaded()
 	{
 		FBulkDataBatchRequest::FBatchBuilder Batch = FBulkDataBatchRequest::NewBatch(5);
 		Batch.Read(BulkData.Positions);
-		Batch.Read(BulkData.Attributes);
+		Batch.Read(BulkData.CurveAttributes);
+		Batch.Read(BulkData.PointAttributes);
 		Batch.Read(BulkData.PointToCurve);
 		Batch.Read(BulkData.Curves);
 
@@ -895,7 +896,8 @@ void FHairStrandsRestResource::InternalAllocate(FRDGBuilder& GraphBuilder)
 	// 1. Lock data, which force the loading data from files (on non-editor build/cooked data). These data are then uploaded to the GPU
 	// 2. A local copy is done by the buffer uploader. This copy is discarded once the uploading is done.
 	InternalCreateVertexBufferRDG_FromBulkData<FHairStrandsPositionFormat>(GraphBuilder, BulkData.Positions, PointCount, PositionBuffer, ToHairResourceDebugName(HAIRSTRANDS_RESOUCE_NAME(CurveType, Hair.StrandsRest_PositionBuffer), ResourceName), OwnerName, EHairResourceUsageType::Static);
-	InternalCreateByteAddressBufferRDG_FromBulkData(GraphBuilder, BulkData.Attributes, AttributeBuffer, ToHairResourceDebugName(HAIRSTRANDS_RESOUCE_NAME(CurveType, Hair.StrandsRest_AttributeBuffer), ResourceName), OwnerName, EHairResourceUsageType::Static);
+	InternalCreateByteAddressBufferRDG_FromBulkData(GraphBuilder, BulkData.CurveAttributes, CurveAttributeBuffer, ToHairResourceDebugName(HAIRSTRANDS_RESOUCE_NAME(CurveType, Hair.StrandsRest_CurveAttributeBuffer), ResourceName), OwnerName, EHairResourceUsageType::Static);
+	InternalCreateByteAddressBufferRDG_FromBulkData(GraphBuilder, BulkData.PointAttributes, PointAttributeBuffer, ToHairResourceDebugName(HAIRSTRANDS_RESOUCE_NAME(CurveType, Hair.StrandsRest_PointAttributeBuffer), ResourceName), OwnerName, EHairResourceUsageType::Static);
 	InternalCreateVertexBufferRDG_FromBulkData<FHairStrandsCurveFormat>(GraphBuilder, BulkData.Curves, CurveCount, CurveBuffer, ToHairResourceDebugName(HAIRSTRANDS_RESOUCE_NAME(CurveType, Hair.StrandsRest_CurveBuffer), ResourceName), OwnerName, EHairResourceUsageType::Static);
 	if (!!(BulkData.Flags & BulkData.DataFlags_Has16bitsCurveIndex))
 	{
@@ -954,7 +956,8 @@ void FHairStrandsRestResource::InternalRelease()
 {
 	PositionBuffer.Release();
 	PositionOffsetBuffer.Release();
-	AttributeBuffer.Release();
+	CurveAttributeBuffer.Release();
+	PointAttributeBuffer.Release();
 	PointToCurveBuffer.Release();
 	TangentBuffer.Release();
 	CurveBuffer.Release();
diff --git a/Engine/Plugins/Runtime/HairStrands/Source/HairStrandsCore/Private/GroomTextureBuilder.cpp b/Engine/Plugins/Runtime/HairStrands/Source/HairStrandsCore/Private/GroomTextureBuilder.cpp
--- a/Engine/Plugins/Runtime/HairStrands/Source/HairStrandsCore/Private/GroomTextureBuilder.cpp
+++ b/Engine/Plugins/Runtime/HairStrands/Source/HairStrandsCore/Private/GroomTextureBuilder.cpp
@@ -1059,10 +1059,12 @@ static bool TraceTextures(
 			Instance.Common.LengthScale = 1.0f;
 			Instance.Common.Length = GroupData.Strands.RestResource->BulkData.MaxLength;
 			Instance.Common.PositionOffset = FVector3f(GroupData.Strands.RestResource->GetPositionOffset());
-			PACK_HAIR_ATTRIBUTE_OFFSETS(Instance.Common.AttributeOffsets, GroupData.Strands.RestResource->BulkData.AttributeOffsets);
+			PACK_CURVE_HAIR_ATTRIBUTE_OFFSETS(Instance.Common.CurveAttributeOffsets, GroupData.Strands.RestResource->BulkData.CurveAttributeOffsets);
+			PACK_POINT_HAIR_ATTRIBUTE_OFFSETS(Instance.Common.PointAttributeOffsets, GroupData.Strands.RestResource->BulkData.PointAttributeOffsets);
 
 			Instance.Resources.PositionBuffer = GroupData.Strands.RestResource->PositionBuffer.SRV;
-			Instance.Resources.AttributeBuffer = GroupData.Strands.RestResource->AttributeBuffer.SRV;
+			Instance.Resources.CurveAttributeBuffer = GroupData.Strands.RestResource->CurveAttributeBuffer.SRV;
+			Instance.Resources.PointAttributeBuffer = GroupData.Strands.RestResource->PointAttributeBuffer.SRV;
 			Instance.Resources.PointToCurveBuffer = GroupData.Strands.RestResource->PointToCurveBuffer.SRV;
 
 			InternalGenerateHairStrandsTextures(
diff --git a/Engine/Plugins/Runtime/HairStrands/Source/HairStrandsCore/Private/HairAttributes.cpp b/Engine/Plugins/Runtime/HairStrands/Source/HairStrandsCore/Private/HairAttributes.cpp
--- a/Engine/Plugins/Runtime/HairStrands/Source/HairStrandsCore/Private/HairAttributes.cpp
+++ b/Engine/Plugins/Runtime/HairStrands/Source/HairStrandsCore/Private/HairAttributes.cpp
@@ -77,21 +77,6 @@ const TCHAR* GetHairAttributeText(EHairAttribute In, uint32 InFlags)
 	return TEXT("UNKNOWN");
 }
 
-uint32 GetHairAttributeIndex(EHairAttribute In)
-{
-	switch (In)
-	{
-	case EHairAttribute::RootUV:					return HAIR_ATTRIBUTE_ROOTUV;
-	case EHairAttribute::ClumpID:					return HAIR_ATTRIBUTE_CLUMPID;
-	case EHairAttribute::StrandID:					return HAIR_ATTRIBUTE_COUNT;	// We don't store this information (yet)
-	case EHairAttribute::PrecomputedGuideWeights:	return HAIR_ATTRIBUTE_COUNT;	// We don't store this information
-	case EHairAttribute::Color:						return HAIR_ATTRIBUTE_BASECOLOR;
-	case EHairAttribute::Roughness:					return HAIR_ATTRIBUTE_ROUGHNESS;
-	case EHairAttribute::AO:						return HAIR_ATTRIBUTE_AO;
-	}
-	return HAIR_ATTRIBUTE_COUNT;
-}
-
 bool FHairStrandsPoints::HasAttribute(EHairAttribute In) const
 {
 	if (In == EHairAttribute::Color)		{ return PointsBaseColor.Num() > 0; }
diff --git a/Engine/Plugins/Runtime/HairStrands/Source/HairStrandsCore/Private/HairStrandsDatas.cpp b/Engine/Plugins/Runtime/HairStrands/Source/HairStrandsCore/Private/HairStrandsDatas.cpp
--- a/Engine/Plugins/Runtime/HairStrands/Source/HairStrandsCore/Private/HairStrandsDatas.cpp
+++ b/Engine/Plugins/Runtime/HairStrands/Source/HairStrandsCore/Private/HairStrandsDatas.cpp
@@ -242,9 +242,13 @@ void FHairStrandsBulkData::Serialize(FArchive& Ar, UObject* Owner)
 	Ar << MaxRadius;
 	Ar << BoundingBox;
 	Ar << Flags;
-	for (uint8 AttributeIt = 0; AttributeIt < HAIR_ATTRIBUTE_COUNT; ++AttributeIt)
+	for (uint8 AttributeIt = 0; AttributeIt < HAIR_CURVE_ATTRIBUTE_COUNT; ++AttributeIt)
 	{
-		Ar << AttributeOffsets[AttributeIt];
+		Ar << CurveAttributeOffsets[AttributeIt];
+	}
+	for (uint8 AttributeIt = 0; AttributeIt < HAIR_POINT_ATTRIBUTE_COUNT; ++AttributeIt)
+	{
+		Ar << PointAttributeOffsets[AttributeIt];
 	}
 	Ar << ImportedAttributes;
 	Ar << ImportedAttributeFlags;
@@ -255,7 +259,8 @@ void FHairStrandsBulkData::Serialize(FArchive& Ar, UObject* Owner)
 	{
 		const uint32 BulkFlags = BULKDATA_Force_NOT_InlinePayload;
 		Positions.SetBulkDataFlags(BulkFlags);
-		Attributes.SetBulkDataFlags(BulkFlags);
+		CurveAttributes.SetBulkDataFlags(BulkFlags);
+		PointAttributes.SetBulkDataFlags(BulkFlags);
 		PointToCurve.SetBulkDataFlags(BulkFlags);
 		Curves.SetBulkDataFlags(BulkFlags);
 	}
@@ -266,7 +271,8 @@ void FHairStrandsBulkData::Serialize(FArchive& Ar, UObject* Owner)
 		bool bAttemptFileMapping = false;
 
 		Positions.Serialize(Ar, Owner, ChunkIndex, bAttemptFileMapping);
-		Attributes.Serialize(Ar, Owner, ChunkIndex, bAttemptFileMapping);
+		PointAttributes.Serialize(Ar, Owner, ChunkIndex, bAttemptFileMapping);
+		CurveAttributes.Serialize(Ar, Owner, ChunkIndex, bAttemptFileMapping);
 		PointToCurve.Serialize(Ar, Owner, ChunkIndex, bAttemptFileMapping);
 		Curves.Serialize(Ar, Owner, ChunkIndex, bAttemptFileMapping);
 	}
@@ -280,19 +286,25 @@ void FHairStrandsBulkData::Reset()
 	MaxRadius = 0;
 	BoundingBox = FBox(EForceInit::ForceInit);
 	Flags = 0;
-	for (uint8 AttributeIt = 0; AttributeIt < HAIR_ATTRIBUTE_COUNT; ++AttributeIt)
+	for (uint8 AttributeIt = 0; AttributeIt < HAIR_CURVE_ATTRIBUTE_COUNT; ++AttributeIt)
+	{
+		CurveAttributeOffsets[AttributeIt] = 0xFFFFFFFF;
+	}
+	for (uint8 AttributeIt = 0; AttributeIt < HAIR_POINT_ATTRIBUTE_COUNT; ++AttributeIt)
 	{
-		AttributeOffsets[AttributeIt] = 0xFFFFFFFF;
+		PointAttributeOffsets[AttributeIt] = 0xFFFFFFFF;
 	}
 	// Deallocate memory if needed
 	Positions.RemoveBulkData();
-	Attributes.RemoveBulkData();
+	CurveAttributes.RemoveBulkData();
+	PointAttributes.RemoveBulkData();
 	PointToCurve.RemoveBulkData();
 	Curves.RemoveBulkData();
 
 	// Reset the bulk byte buffer to ensure the (serialize) data size is reset to 0
 	Positions 		= FByteBulkData();
-	Attributes		= FByteBulkData();
+	CurveAttributes	= FByteBulkData();
+	PointAttributes	= FByteBulkData();
 	PointToCurve	= FByteBulkData();
 	Curves			= FByteBulkData();
 }
diff --git a/Engine/Plugins/Runtime/HairStrands/Source/HairStrandsCore/Private/HairStrandsRendering.cpp b/Engine/Plugins/Runtime/HairStrands/Source/HairStrandsCore/Private/HairStrandsRendering.cpp
--- a/Engine/Plugins/Runtime/HairStrands/Source/HairStrandsCore/Private/HairStrandsRendering.cpp
+++ b/Engine/Plugins/Runtime/HairStrands/Source/HairStrandsCore/Private/HairStrandsRendering.cpp
@@ -766,7 +766,7 @@ class FHairPatchAttributeCS : public FGlobalShader
 
 	BEGIN_SHADER_PARAMETER_STRUCT(FParameters, )
 		SHADER_PARAMETER(uint32, CurveCount)
-		SHADER_PARAMETER_ARRAY(FUintVector4, RenAttributeOffsets, [HAIR_ATTRIBUTE_OFFSET_COUNT])
+		SHADER_PARAMETER_ARRAY(FUintVector4, RenCurveAttributeOffsets, [HAIR_CURVE_ATTRIBUTE_OFFSET_COUNT])
 		SHADER_PARAMETER_RDG_BUFFER_SRV(Buffer, RenCurveBuffer)
 
 		SHADER_PARAMETER_RDG_BUFFER_SRV(Buffer, InterpolationBuffer)
@@ -802,7 +802,7 @@ static void AddPatchAttributePass(
 	const EHairPatchAttribute Mode,
 	const bool bSimulation,
 	const bool bUseSingleGuide,
-	const uint32* RenAttributeOffests,
+	const uint32* RenCurveAttributeOffests,
 	const FRDGBufferRef& RenAttributeBuffer,
 	const FRDGBufferSRVRef& RenCurveBuffer,
 	const FRDGBufferSRVRef& RenVertexToClusterIdBuffer,
@@ -824,7 +824,7 @@ static void AddPatchAttributePass(
 	Parameters->CurveCount = CurveCount;
 	Parameters->RenCurveBuffer = RenCurveBuffer;
 	Parameters->RenVertexToClusterIdBuffer = RenVertexToClusterIdBuffer;
-	PACK_HAIR_ATTRIBUTE_OFFSETS(Parameters->RenAttributeOffsets, RenAttributeOffests)
+	PACK_CURVE_HAIR_ATTRIBUTE_OFFSETS(Parameters->RenCurveAttributeOffsets, RenCurveAttributeOffests)
 
 	if (bSimulation)
 	{
@@ -1629,9 +1629,14 @@ static FHairGroupPublicData::FVertexFactoryInput InternalComputeHairStrandsVerte
 	FHairGroupPublicData::FVertexFactoryInput OutVFInput;
 	if (!Instance || Instance->GeometryType != EHairGeometryType::Strands)
 	{
-		for (uint32 OffsetIt = 0; OffsetIt < HAIR_ATTRIBUTE_OFFSET_COUNT; ++OffsetIt)
+		const FUintVector4 InvalidOffset = FUintVector4(HAIR_ATTRIBUTE_INVALID_OFFSET,HAIR_ATTRIBUTE_INVALID_OFFSET,HAIR_ATTRIBUTE_INVALID_OFFSET,HAIR_ATTRIBUTE_INVALID_OFFSET);
+		for (uint32 OffsetIt = 0; OffsetIt < HAIR_CURVE_ATTRIBUTE_OFFSET_COUNT; ++OffsetIt)
 		{
-			OutVFInput.Strands.Common.AttributeOffsets[OffsetIt] = FUintVector4(HAIR_ATTRIBUTE_INVALID_OFFSET,HAIR_ATTRIBUTE_INVALID_OFFSET,HAIR_ATTRIBUTE_INVALID_OFFSET,HAIR_ATTRIBUTE_INVALID_OFFSET);
+			OutVFInput.Strands.Common.CurveAttributeOffsets[OffsetIt] = InvalidOffset;
+		}
+		for (uint32 OffsetIt = 0; OffsetIt < HAIR_POINT_ATTRIBUTE_OFFSET_COUNT; ++OffsetIt)
+		{
+			OutVFInput.Strands.Common.PointAttributeOffsets[OffsetIt] = InvalidOffset;
 		}
 		return OutVFInput;
 	}
@@ -1646,7 +1651,8 @@ static FHairGroupPublicData::FVertexFactoryInput InternalComputeHairStrandsVerte
 		CONVERT_HAIRSSTRANDS_VF_PARAMETERS(OutVFInput.Strands.PositionBuffer,			Instance->Guides.DeformedResource->GetBuffer(FHairStrandsDeformedResource::EFrameType::Current));
 		CONVERT_HAIRSSTRANDS_VF_PARAMETERS(OutVFInput.Strands.PrevPositionBuffer,		Instance->Guides.DeformedResource->GetBuffer(FHairStrandsDeformedResource::EFrameType::Previous));
 		CONVERT_HAIRSSTRANDS_VF_PARAMETERS(OutVFInput.Strands.TangentBuffer,			Instance->Guides.DeformedResource->TangentBuffer);
-		CONVERT_HAIRSSTRANDS_VF_PARAMETERS(OutVFInput.Strands.AttributeBuffer,			Instance->Guides.RestResource->AttributeBuffer);
+		CONVERT_HAIRSSTRANDS_VF_PARAMETERS(OutVFInput.Strands.CurveAttributeBuffer,		Instance->Guides.RestResource->CurveAttributeBuffer);
+		CONVERT_HAIRSSTRANDS_VF_PARAMETERS(OutVFInput.Strands.PointAttributeBuffer,		Instance->Guides.RestResource->PointAttributeBuffer);
 		CONVERT_HAIRSSTRANDS_VF_PARAMETERS(OutVFInput.Strands.CurveBuffer,				Instance->Guides.RestResource->CurveBuffer);
 		CONVERT_HAIRSSTRANDS_VF_PARAMETERS(OutVFInput.Strands.PointToCurveBuffer,		Instance->Guides.RestResource->PointToCurveBuffer);
 		CONVERT_HAIRSSTRANDS_VF_PARAMETERS(OutVFInput.Strands.PositionOffsetBuffer,		Instance->Guides.DeformedResource->GetPositionOffsetBuffer(FHairStrandsDeformedResource::EFrameType::Current));
@@ -1685,7 +1691,8 @@ static FHairGroupPublicData::FVertexFactoryInput InternalComputeHairStrandsVerte
 			OutVFInput.Strands.Common.PrevPositionOffset = FVector3f(Instance->Strands.DeformedResource->GetPositionOffset(FHairStrandsDeformedResource::EFrameType::Current)); // LWC precision loss
 		}
 		CONVERT_HAIRSSTRANDS_VF_PARAMETERS(OutVFInput.Strands.TangentBuffer,			Instance->Strands.DeformedResource->TangentBuffer);
-		CONVERT_HAIRSSTRANDS_VF_PARAMETERS(OutVFInput.Strands.AttributeBuffer,			(ViewMode == EGroomViewMode::RenderHairStrands || ViewMode == EGroomViewMode::Cluster || ViewMode == EGroomViewMode::ClusterAABB) ? Instance->Strands.DebugAttributeBuffer : Instance->Strands.RestResource->AttributeBuffer);
+		CONVERT_HAIRSSTRANDS_VF_PARAMETERS(OutVFInput.Strands.CurveAttributeBuffer,		(ViewMode == EGroomViewMode::RenderHairStrands || ViewMode == EGroomViewMode::Cluster || ViewMode == EGroomViewMode::ClusterAABB) ? Instance->Strands.DebugCurveAttributeBuffer : Instance->Strands.RestResource->CurveAttributeBuffer);
+		CONVERT_HAIRSSTRANDS_VF_PARAMETERS(OutVFInput.Strands.PointAttributeBuffer,		Instance->Strands.RestResource->PointAttributeBuffer);
 		CONVERT_HAIRSSTRANDS_VF_PARAMETERS(OutVFInput.Strands.PointToCurveBuffer,		Instance->Strands.RestResource->PointToCurveBuffer);
 		CONVERT_HAIRSSTRANDS_VF_PARAMETERS(OutVFInput.Strands.CurveBuffer,				Instance->Strands.RestResource->CurveBuffer);
 
@@ -1698,7 +1705,8 @@ static FHairGroupPublicData::FVertexFactoryInput InternalComputeHairStrandsVerte
 		CONVERT_HAIRSSTRANDS_VF_PARAMETERS(OutVFInput.Strands.PositionBuffer,			Instance->Strands.RestResource->PositionBuffer);
 		CONVERT_HAIRSSTRANDS_VF_PARAMETERS(OutVFInput.Strands.PrevPositionBuffer,		Instance->Strands.RestResource->PositionBuffer);
 		CONVERT_HAIRSSTRANDS_VF_PARAMETERS(OutVFInput.Strands.TangentBuffer,			Instance->Strands.RestResource->TangentBuffer);
-		CONVERT_HAIRSSTRANDS_VF_PARAMETERS(OutVFInput.Strands.AttributeBuffer,			(ViewMode == EGroomViewMode::Cluster || ViewMode == EGroomViewMode::ClusterAABB) ? Instance->Strands.DebugAttributeBuffer : Instance->Strands.RestResource->AttributeBuffer);
+		CONVERT_HAIRSSTRANDS_VF_PARAMETERS(OutVFInput.Strands.CurveAttributeBuffer,		(ViewMode == EGroomViewMode::Cluster || ViewMode == EGroomViewMode::ClusterAABB) ? Instance->Strands.DebugCurveAttributeBuffer : Instance->Strands.RestResource->CurveAttributeBuffer);
+		CONVERT_HAIRSSTRANDS_VF_PARAMETERS(OutVFInput.Strands.PointAttributeBuffer,		Instance->Strands.RestResource->PointAttributeBuffer);
 		CONVERT_HAIRSSTRANDS_VF_PARAMETERS(OutVFInput.Strands.PointToCurveBuffer,		Instance->Strands.RestResource->PointToCurveBuffer);
 		CONVERT_HAIRSSTRANDS_VF_PARAMETERS(OutVFInput.Strands.CurveBuffer,				Instance->Strands.RestResource->CurveBuffer);
 		CONVERT_HAIRSSTRANDS_VF_PARAMETERS(OutVFInput.Strands.PositionOffsetBuffer,		Instance->Strands.RestResource->PositionOffsetBuffer);
@@ -1725,11 +1733,13 @@ static FHairGroupPublicData::FVertexFactoryInput InternalComputeHairStrandsVerte
 
 	if (bUseGuideAttributeOffsets)
 	{
-		PACK_HAIR_ATTRIBUTE_OFFSETS(OutVFInput.Strands.Common.AttributeOffsets, Instance->Guides.Data->AttributeOffsets);
+		PACK_CURVE_HAIR_ATTRIBUTE_OFFSETS(OutVFInput.Strands.Common.CurveAttributeOffsets, Instance->Guides.Data->CurveAttributeOffsets);
+		PACK_POINT_HAIR_ATTRIBUTE_OFFSETS(OutVFInput.Strands.Common.PointAttributeOffsets, Instance->Guides.Data->PointAttributeOffsets);
 	}
 	else
 	{
-		PACK_HAIR_ATTRIBUTE_OFFSETS(OutVFInput.Strands.Common.AttributeOffsets, Instance->Strands.Data->AttributeOffsets);
+		PACK_CURVE_HAIR_ATTRIBUTE_OFFSETS(OutVFInput.Strands.Common.CurveAttributeOffsets, Instance->Strands.Data->CurveAttributeOffsets);
+		PACK_POINT_HAIR_ATTRIBUTE_OFFSETS(OutVFInput.Strands.Common.PointAttributeOffsets, Instance->Strands.Data->PointAttributeOffsets);
 	}
 
 #if RHI_RAYTRACING
@@ -2125,11 +2135,11 @@ void ComputeHairStrandsInterpolation(
 			{
 				// Create an debug buffer for storing cluster visualization data. This is only used for debug purpose, hence only enable in editor build.
 				// Special case for debug mode were the attribute buffer is patch with some custom data to show hair properties (strands belonging to the same cluster, ...)
-				if (Instance->Strands.DebugAttributeBuffer.Buffer == nullptr)
+				if (Instance->Strands.DebugCurveAttributeBuffer.Buffer == nullptr)
 				{
-					CreateHairStrandsDebugAttributeBuffer(GraphBuilder, &Instance->Strands.DebugAttributeBuffer, Instance->Strands.Data->Attributes.GetBulkDataSize(), FName(Instance->Debug.MeshComponentName));
+					CreateHairStrandsDebugAttributeBuffer(GraphBuilder, &Instance->Strands.DebugCurveAttributeBuffer, Instance->Strands.Data->CurveAttributes.GetBulkDataSize(), FName(Instance->Debug.MeshComponentName));
 				}
-				FRDGImportedBuffer OutRenAttributeBuffer = Register(GraphBuilder, Instance->Strands.DebugAttributeBuffer, ERDGImportedBufferFlags::CreateUAV);
+				FRDGImportedBuffer OutRenCurveAttributeBuffer = Register(GraphBuilder, Instance->Strands.DebugCurveAttributeBuffer, ERDGImportedBufferFlags::CreateUAV);
 
 				AddPatchAttributePass(
 					GraphBuilder,
@@ -2138,14 +2148,14 @@ void ComputeHairStrandsInterpolation(
 					PatchMode,
 					bValidGuide,
 					bUseSingleGuide,
-					Instance->Strands.Data->AttributeOffsets,
-					Register(GraphBuilder, Instance->Strands.RestResource->AttributeBuffer, ERDGImportedBufferFlags::CreateSRV).Buffer,
+					Instance->Strands.Data->CurveAttributeOffsets,
+					Register(GraphBuilder, Instance->Strands.RestResource->CurveAttributeBuffer, ERDGImportedBufferFlags::CreateSRV).Buffer,
 					RegisterAsSRV(GraphBuilder, Instance->Strands.RestResource->CurveBuffer),
 					RegisterAsSRV(GraphBuilder, Instance->Strands.ClusterCullingResource->VertexToClusterIdBuffer),
 					bValidGuide &&  bUseSingleGuide ? RegisterAsSRV(GraphBuilder, Instance->Strands.InterpolationResource->InterpolationBuffer) : nullptr,
 					bValidGuide && !bUseSingleGuide ? RegisterAsSRV(GraphBuilder, Instance->Strands.InterpolationResource->Interpolation0Buffer) : nullptr,
 					bValidGuide && !bUseSingleGuide ? RegisterAsSRV(GraphBuilder, Instance->Strands.InterpolationResource->Interpolation1Buffer) : nullptr,
-					OutRenAttributeBuffer);
+					OutRenCurveAttributeBuffer);
 			}
 			#endif
 
diff --git a/Engine/Plugins/Runtime/HairStrands/Source/HairStrandsCore/Private/HairStrandsVertexFactory.cpp b/Engine/Plugins/Runtime/HairStrands/Source/HairStrandsCore/Private/HairStrandsVertexFactory.cpp
--- a/Engine/Plugins/Runtime/HairStrands/Source/HairStrandsCore/Private/HairStrandsVertexFactory.cpp
+++ b/Engine/Plugins/Runtime/HairStrands/Source/HairStrandsCore/Private/HairStrandsVertexFactory.cpp
@@ -98,7 +98,8 @@ FHairStrandsVertexFactoryUniformShaderParameters FHairGroupInstance::GetHairStan
 
 	Out.Resources.PositionBuffer					= VFInput.Strands.PositionBufferRHISRV;
 	Out.Resources.PositionOffsetBuffer 				= VFInput.Strands.PositionOffsetBufferRHISRV;
-	Out.Resources.AttributeBuffer					= VFInput.Strands.AttributeBufferRHISRV;
+	Out.Resources.PointAttributeBuffer				= VFInput.Strands.PointAttributeBufferRHISRV;
+	Out.Resources.CurveAttributeBuffer				= VFInput.Strands.CurveAttributeBufferRHISRV;
 	Out.Resources.CurveBuffer						= VFInput.Strands.CurveBufferRHISRV;
 	Out.Resources.PointToCurveBuffer				= VFInput.Strands.PointToCurveBufferRHISRV;
 	Out.Resources.TangentBuffer 					= VFInput.Strands.TangentBufferRHISRV;
@@ -108,7 +109,8 @@ FHairStrandsVertexFactoryUniformShaderParameters FHairGroupInstance::GetHairStan
 	// swap in some default data for those buffers that are not valid yet
 	if (!Out.Resources.PositionBuffer) 						{ Out.Resources.PositionBuffer = GDummyCulledDispatchVertexIdsBuffer.SRVFloat; }
 	if (!Out.Resources.PositionOffsetBuffer) 				{ Out.Resources.PositionOffsetBuffer = GDummyCulledDispatchVertexIdsBuffer.SRVFloat; }
-	if (!Out.Resources.AttributeBuffer) 					{ Out.Resources.AttributeBuffer = GDummyCulledDispatchVertexIdsBuffer.SRVByteAddress; }
+	if (!Out.Resources.CurveAttributeBuffer) 				{ Out.Resources.CurveAttributeBuffer = GDummyCulledDispatchVertexIdsBuffer.SRVByteAddress; }
+	if (!Out.Resources.PointAttributeBuffer) 				{ Out.Resources.PointAttributeBuffer = GDummyCulledDispatchVertexIdsBuffer.SRVByteAddress; }
 	if (!Out.Resources.CurveBuffer) 						{ Out.Resources.CurveBuffer = GDummyCulledDispatchVertexIdsBuffer.SRVByteAddress; }
 	if (!Out.Resources.PointToCurveBuffer) 					{ Out.Resources.PointToCurveBuffer = GDummyCulledDispatchVertexIdsBuffer.SRVUint; }
 	if (!Out.Resources.TangentBuffer) 						{ Out.Resources.TangentBuffer = GDummyCulledDispatchVertexIdsBuffer.SRVFloat; }
diff --git a/Engine/Plugins/Runtime/HairStrands/Source/HairStrandsCore/Public/GroomAsset.h b/Engine/Plugins/Runtime/HairStrands/Source/HairStrandsCore/Public/GroomAsset.h
--- a/Engine/Plugins/Runtime/HairStrands/Source/HairStrandsCore/Public/GroomAsset.h
+++ b/Engine/Plugins/Runtime/HairStrands/Source/HairStrandsCore/Public/GroomAsset.h
@@ -115,7 +115,8 @@ struct HAIRSTRANDSCORE_API FHairGroupData
 		{
 			uint32 Total = 0;
 			Total += BulkData.Positions.IsBulkDataLoaded() 		? BulkData.Positions.GetBulkDataSize()   : 0;
-			Total += BulkData.Attributes.IsBulkDataLoaded()		? BulkData.Attributes.GetBulkDataSize() : 0;
+			Total += BulkData.CurveAttributes.IsBulkDataLoaded()? BulkData.CurveAttributes.GetBulkDataSize() : 0;
+			Total += BulkData.PointAttributes.IsBulkDataLoaded()? BulkData.PointAttributes.GetBulkDataSize() : 0;
 			Total += BulkData.Curves.IsBulkDataLoaded() 		? BulkData.Curves.GetBulkDataSize(): 0;
 			Total += BulkData.PointToCurve.IsBulkDataLoaded()	? BulkData.PointToCurve.GetBulkDataSize() : 0;
 			return Total;
diff --git a/Engine/Plugins/Runtime/HairStrands/Source/HairStrandsCore/Public/GroomInstance.h b/Engine/Plugins/Runtime/HairStrands/Source/HairStrandsCore/Public/GroomInstance.h
--- a/Engine/Plugins/Runtime/HairStrands/Source/HairStrandsCore/Public/GroomInstance.h
+++ b/Engine/Plugins/Runtime/HairStrands/Source/HairStrandsCore/Public/GroomInstance.h
@@ -136,7 +136,7 @@ struct HAIRSTRANDSCORE_API FHairGroupInstance : public FHairStrandsInstance
 		int   CachedProceduralSplits = 0;
 		#endif
 
-		FRDGExternalBuffer DebugAttributeBuffer;
+		FRDGExternalBuffer DebugCurveAttributeBuffer;
 		FHairGroupInstanceModifer Modifier;
 
 		FHairStrandsUniformBuffer UniformBuffer;
diff --git a/Engine/Plugins/Runtime/HairStrands/Source/HairStrandsCore/Public/GroomResources.h b/Engine/Plugins/Runtime/HairStrands/Source/HairStrandsCore/Public/GroomResources.h
--- a/Engine/Plugins/Runtime/HairStrands/Source/HairStrandsCore/Public/GroomResources.h
+++ b/Engine/Plugins/Runtime/HairStrands/Source/HairStrandsCore/Public/GroomResources.h
@@ -313,7 +313,8 @@ struct FHairStrandsRestResource : public FHairCommonResource
 		Total += GetBufferTotalNumBytes(PositionBuffer);
 		Total += GetBufferTotalNumBytes(PositionOffsetBuffer);
 		Total += GetBufferTotalNumBytes(TangentBuffer);
-		Total += GetBufferTotalNumBytes(AttributeBuffer);
+		Total += GetBufferTotalNumBytes(PointAttributeBuffer);
+		Total += GetBufferTotalNumBytes(CurveAttributeBuffer);
 		Total += GetBufferTotalNumBytes(PointToCurveBuffer);
 		Total += GetBufferTotalNumBytes(CurveBuffer);
 		return Total;
@@ -328,8 +329,11 @@ struct FHairStrandsRestResource : public FHairCommonResource
 	/* Strand hair tangent buffer (non-allocated unless used for static geometry) */
 	FRDGExternalBuffer TangentBuffer;
 
-	/* Strand hair attribute buffer */
-	FRDGExternalBuffer AttributeBuffer;
+	/* Strand hair per-point attribute buffer */
+	FRDGExternalBuffer PointAttributeBuffer;
+
+	/* Strand hair per-curve attribute buffer */
+	FRDGExternalBuffer CurveAttributeBuffer;
 
 	/* Strand hair vertex to curve index mapping */
 	FRDGExternalBuffer PointToCurveBuffer;
diff --git a/Engine/Plugins/Runtime/HairStrands/Source/HairStrandsCore/Public/HairStrandsDatas.h b/Engine/Plugins/Runtime/HairStrands/Source/HairStrandsCore/Public/HairStrandsDatas.h
--- a/Engine/Plugins/Runtime/HairStrands/Source/HairStrandsCore/Public/HairStrandsDatas.h
+++ b/Engine/Plugins/Runtime/HairStrands/Source/HairStrandsCore/Public/HairStrandsDatas.h
@@ -483,16 +483,18 @@ struct HAIRSTRANDSCORE_API FHairStrandsBulkData
 	float MaxRadius = 0;
 	FBox BoundingBox = FBox(EForceInit::ForceInit);
 	uint32 Flags = 0;
-	uint32 AttributeOffsets[HAIR_ATTRIBUTE_COUNT] = {0};
+	uint32 CurveAttributeOffsets[HAIR_CURVE_ATTRIBUTE_COUNT] = {0};
+	uint32 PointAttributeOffsets[HAIR_POINT_ATTRIBUTE_COUNT] = {0};
 
 	/** Imported attribute info */
 	uint32 ImportedAttributes = 0;
 	uint32 ImportedAttributeFlags = 0;
 
-	FByteBulkData Positions;	// Size = PointCount
-	FByteBulkData Attributes;	// Size = x*PointCount + y*CurveCount (depends on the stored attributes, which could be per-vertex(x) or per-curve(y))
-	FByteBulkData PointToCurve; // Size = PointCount
-	FByteBulkData Curves;		// Size = CurveCount
+	FByteBulkData Positions;		// Size = PointCount
+	FByteBulkData CurveAttributes;	// Size = y*CurveCount (depends on the per-curve stored attributes)
+	FByteBulkData PointAttributes;	// Size = x*PointCount (depends on the per-point stored attributes)
+	FByteBulkData PointToCurve; 	// Size = PointCount
+	FByteBulkData Curves;			// Size = CurveCount
 };
 
 /** Hair strands debug data */
diff --git a/Engine/Shaders/Private/HairStrands/HairStrandsAttributeCommon.ush b/Engine/Shaders/Private/HairStrands/HairStrandsAttributeCommon.ush
--- a/Engine/Shaders/Private/HairStrands/HairStrandsAttributeCommon.ush
+++ b/Engine/Shaders/Private/HairStrands/HairStrandsAttributeCommon.ush
@@ -6,24 +6,27 @@
 
 ////////////////////////////////////////////////////////////////////////////////
 // Data stride in bytes
-#define HAIR_ATTRIBUTE_STRIDE_ROOTUV		4
-#define HAIR_ATTRIBUTE_STRIDE_SEED			1
-#define HAIR_ATTRIBUTE_STRIDE_LENGTH		2
-#define HAIR_ATTRIBUTE_STRIDE_CLUMPID		2
-#define HAIR_ATTRIBUTE_STRIDE_CLUMPID3		8
-#define HAIR_ATTRIBUTE_STRIDE_BASECOLOR		4
-#define HAIR_ATTRIBUTE_STRIDE_ROUGHNESS		1
-#define HAIR_ATTRIBUTE_STRIDE_AO			1
+// Data stride in bytes
+#define HAIR_CURVE_ATTRIBUTE_STRIDE_ROOTUV		4
+#define HAIR_CURVE_ATTRIBUTE_STRIDE_SEED		1
+#define HAIR_CURVE_ATTRIBUTE_STRIDE_LENGTH		2
+#define HAIR_CURVE_ATTRIBUTE_STRIDE_CLUMPID		2
+#define HAIR_CURVE_ATTRIBUTE_STRIDE_CLUMPID3	8
+
+#define HAIR_POINT_ATTRIBUTE_STRIDE_COLOR		4
+#define HAIR_POINT_ATTRIBUTE_STRIDE_ROUGHNESS	1
+#define HAIR_POINT_ATTRIBUTE_STRIDE_AO			1
 
 // Data offset in bytes. The order needs to match HairStrandsDefinitions.h
-#define HAIR_ATTRIBUTE_OFFSET_ROOTUV(InOffsets)		(InOffsets[0].x)
-#define HAIR_ATTRIBUTE_OFFSET_SEED(InOffsets)		(InOffsets[0].y)
-#define HAIR_ATTRIBUTE_OFFSET_LENGTH(InOffsets)		(InOffsets[0].z)
-#define HAIR_ATTRIBUTE_OFFSET_AO(InOffsets)			(InOffsets[0].w)
-#define HAIR_ATTRIBUTE_OFFSET_BASECOLOR(InOffsets)	(InOffsets[1].x)
-#define HAIR_ATTRIBUTE_OFFSET_ROUGHNESS(InOffsets)	(InOffsets[1].y)
-#define HAIR_ATTRIBUTE_OFFSET_CLUMPID(InOffsets)	(InOffsets[1].z)
-#define HAIR_ATTRIBUTE_OFFSET_CLUMPID3(InOffsets)	(InOffsets[1].w)
+#define HAIR_CURVE_ATTRIBUTE_OFFSET_ROOTUV(InOffsets)		(InOffsets[0].x)
+#define HAIR_CURVE_ATTRIBUTE_OFFSET_SEED(InOffsets)			(InOffsets[0].y)
+#define HAIR_CURVE_ATTRIBUTE_OFFSET_LENGTH(InOffsets)		(InOffsets[0].z)
+#define HAIR_CURVE_ATTRIBUTE_OFFSET_CLUMPID(InOffsets)		(InOffsets[0].w)
+#define HAIR_CURVE_ATTRIBUTE_OFFSET_CLUMPID3(InOffsets)		(InOffsets[1].x)
+
+#define HAIR_POINT_ATTRIBUTE_OFFSET_COLOR(InOffsets)		(InOffsets[0].x)
+#define HAIR_POINT_ATTRIBUTE_OFFSET_ROUGHNESS(InOffsets)	(InOffsets[0].y)
+#define HAIR_POINT_ATTRIBUTE_OFFSET_AO(InOffsets)			(InOffsets[0].z)
 
 ////////////////////////////////////////////////////////////////////////////////
 // Attribute access functions
@@ -122,7 +125,7 @@ float2 GetHairStrandsUV(uint HairPrimitiveId, float2 HairPrimitiveUV)
 
 float2 GetHairStrandsDimensions(uint HairPrimitiveId, float2 HairPrimitiveUV)
 {
-	const float CurveLength = UnpackHairLength(InternalGetHairStrandsAttributesPerCurve(HairPrimitiveId, HAIR_ATTRIBUTE_OFFSET_LENGTH(HairStrandsVF.AttributeOffsets), HAIR_ATTRIBUTE_STRIDE_LENGTH, HairStrandsVF.AttributeBuffer, HairStrandsVF.PointToCurveBuffer));
+	const float CurveLength = UnpackHairLength(InternalGetHairStrandsAttributesPerCurve(HairPrimitiveId, HAIR_CURVE_ATTRIBUTE_OFFSET_LENGTH(HairStrandsVF.CurveAttributeOffsets), HAIR_CURVE_ATTRIBUTE_STRIDE_LENGTH, HairStrandsVF.CurveAttributeBuffer, HairStrandsVF.PointToCurveBuffer));
 	const float CoordU = InternalGetHairStrandsCoordU(HairPrimitiveId, HairPrimitiveUV);
 	const float Radius = InternalGetHairStrandsWorldRadius(HairPrimitiveId, HairPrimitiveUV);
 	return float2(CurveLength * CoordU, Radius);
@@ -130,10 +133,10 @@ float2 GetHairStrandsDimensions(uint HairPrimitiveId, float2 HairPrimitiveUV)
 
 float2 GetHairStrandsRootUV(uint HairPrimitiveId, float2 HairPrimitiveUV)
 {
-	const uint RootUVOffset = HAIR_ATTRIBUTE_OFFSET_ROOTUV(HairStrandsVF.AttributeOffsets);
+	const uint RootUVOffset = HAIR_CURVE_ATTRIBUTE_OFFSET_ROOTUV(HairStrandsVF.CurveAttributeOffsets);
 	if (IsAttributeValid(RootUVOffset))
 	{
-		return UnpackHairRootUV(InternalGetHairStrandsAttributesPerCurve(HairPrimitiveId, RootUVOffset, HAIR_ATTRIBUTE_STRIDE_ROOTUV, HairStrandsVF.AttributeBuffer, HairStrandsVF.PointToCurveBuffer));
+		return UnpackHairRootUV(InternalGetHairStrandsAttributesPerCurve(HairPrimitiveId, RootUVOffset, HAIR_CURVE_ATTRIBUTE_STRIDE_ROOTUV, HairStrandsVF.CurveAttributeBuffer, HairStrandsVF.PointToCurveBuffer));
 	}
 	else
 	{
@@ -143,20 +146,20 @@ float2 GetHairStrandsRootUV(uint HairPrimitiveId, float2 HairPrimitiveUV)
 
 float GetHairStrandsSeed(uint HairPrimitiveId, float2 HairPrimitiveUV)
 {
-	return UnpackHairSeed(InternalGetHairStrandsAttributesPerCurve(HairPrimitiveId, HAIR_ATTRIBUTE_OFFSET_SEED(HairStrandsVF.AttributeOffsets), HAIR_ATTRIBUTE_STRIDE_SEED, HairStrandsVF.AttributeBuffer, HairStrandsVF.PointToCurveBuffer));
+	return UnpackHairSeed(InternalGetHairStrandsAttributesPerCurve(HairPrimitiveId, HAIR_CURVE_ATTRIBUTE_OFFSET_SEED(HairStrandsVF.CurveAttributeOffsets), HAIR_CURVE_ATTRIBUTE_STRIDE_SEED, HairStrandsVF.CurveAttributeBuffer, HairStrandsVF.PointToCurveBuffer));
 }
 
 uint3 GetHairStrandsClumpID(uint HairPrimitiveId, float2 HairPrimitiveUV)
 {
-	const uint ClumpIdOffset  = HAIR_ATTRIBUTE_OFFSET_CLUMPID(HairStrandsVF.AttributeOffsets);
-	const uint ClumpId3Offset = HAIR_ATTRIBUTE_OFFSET_CLUMPID3(HairStrandsVF.AttributeOffsets);
+	const uint ClumpIdOffset  = HAIR_CURVE_ATTRIBUTE_OFFSET_CLUMPID(HairStrandsVF.CurveAttributeOffsets);
+	const uint ClumpId3Offset = HAIR_CURVE_ATTRIBUTE_OFFSET_CLUMPID3(HairStrandsVF.CurveAttributeOffsets);
 	if (IsAttributeValid(ClumpIdOffset))
 	{
-		return UnpackHairClumpID(InternalGetHairStrandsAttributesPerCurve(HairPrimitiveId, ClumpIdOffset, HAIR_ATTRIBUTE_STRIDE_CLUMPID, HairStrandsVF.AttributeBuffer, HairStrandsVF.PointToCurveBuffer));
+		return UnpackHairClumpID(InternalGetHairStrandsAttributesPerCurve(HairPrimitiveId, ClumpIdOffset, HAIR_CURVE_ATTRIBUTE_STRIDE_CLUMPID, HairStrandsVF.CurveAttributeBuffer, HairStrandsVF.PointToCurveBuffer));
 	}
 	else if (IsAttributeValid(ClumpId3Offset))
 	{
-		return UnpackHairClumpID(InternalGetHairStrandsAttributesPerCurve64(HairPrimitiveId, ClumpId3Offset, HAIR_ATTRIBUTE_STRIDE_CLUMPID3, HairStrandsVF.AttributeBuffer, HairStrandsVF.PointToCurveBuffer));
+		return UnpackHairClumpID(InternalGetHairStrandsAttributesPerCurve64(HairPrimitiveId, ClumpId3Offset, HAIR_CURVE_ATTRIBUTE_STRIDE_CLUMPID3, HairStrandsVF.CurveAttributeBuffer, HairStrandsVF.PointToCurveBuffer));
 	}
 	else
 	{
@@ -166,11 +169,11 @@ uint3 GetHairStrandsClumpID(uint HairPrimitiveId, float2 HairPrimitiveUV)
 
 float3 GetHairStrandsBaseColor(uint HairPrimitiveId, float2 HairPrimitiveUV)
 {
-	const uint BaseColorOffset = HAIR_ATTRIBUTE_OFFSET_BASECOLOR(HairStrandsVF.AttributeOffsets);
-	if (IsAttributeValid(BaseColorOffset))
+	const uint ColorOffset = HAIR_POINT_ATTRIBUTE_OFFSET_COLOR(HairStrandsVF.PointAttributeOffsets);
+	if (IsAttributeValid(ColorOffset))
 	{
-		const float3 Color0 = UnpackHairColor(InternalGetHairStrandsAttributesPerVertex(HairPrimitiveId,   BaseColorOffset, HAIR_ATTRIBUTE_STRIDE_BASECOLOR, HairStrandsVF.AttributeBuffer));
-		const float3 Color1 = UnpackHairColor(InternalGetHairStrandsAttributesPerVertex(HairPrimitiveId+1, BaseColorOffset, HAIR_ATTRIBUTE_STRIDE_BASECOLOR, HairStrandsVF.AttributeBuffer));
+		const float3 Color0 = UnpackHairColor(InternalGetHairStrandsAttributesPerVertex(HairPrimitiveId,   ColorOffset, HAIR_POINT_ATTRIBUTE_STRIDE_COLOR, HairStrandsVF.PointAttributeBuffer));
+		const float3 Color1 = UnpackHairColor(InternalGetHairStrandsAttributesPerVertex(HairPrimitiveId+1, ColorOffset, HAIR_POINT_ATTRIBUTE_STRIDE_COLOR, HairStrandsVF.PointAttributeBuffer));
 		return lerp(Color0, Color1, HairPrimitiveUV.x);
 	}
 	else
@@ -181,11 +184,11 @@ float3 GetHairStrandsBaseColor(uint HairPrimitiveId, float2 HairPrimitiveUV)
 
 float GetHairStrandsRoughness(uint HairPrimitiveId, float2 HairPrimitiveUV)
 {
-	const uint RoughnessOffset = HAIR_ATTRIBUTE_OFFSET_ROUGHNESS(HairStrandsVF.AttributeOffsets);
+	const uint RoughnessOffset = HAIR_POINT_ATTRIBUTE_OFFSET_ROUGHNESS(HairStrandsVF.PointAttributeOffsets);
 	if (IsAttributeValid(RoughnessOffset))
 	{
-		const float Roughness0 = UnpackHairRoughness(InternalGetHairStrandsAttributesPerVertex(HairPrimitiveId,   RoughnessOffset, HAIR_ATTRIBUTE_STRIDE_ROUGHNESS, HairStrandsVF.AttributeBuffer));
-		const float Roughness1 = UnpackHairRoughness(InternalGetHairStrandsAttributesPerVertex(HairPrimitiveId+1, RoughnessOffset, HAIR_ATTRIBUTE_STRIDE_ROUGHNESS, HairStrandsVF.AttributeBuffer));
+		const float Roughness0 = UnpackHairRoughness(InternalGetHairStrandsAttributesPerVertex(HairPrimitiveId,   RoughnessOffset, HAIR_POINT_ATTRIBUTE_STRIDE_ROUGHNESS, HairStrandsVF.PointAttributeBuffer));
+		const float Roughness1 = UnpackHairRoughness(InternalGetHairStrandsAttributesPerVertex(HairPrimitiveId+1, RoughnessOffset, HAIR_POINT_ATTRIBUTE_STRIDE_ROUGHNESS, HairStrandsVF.PointAttributeBuffer));
 		return lerp(Roughness0, Roughness1, HairPrimitiveUV.x);
 	}
 	else
@@ -196,11 +199,11 @@ float GetHairStrandsRoughness(uint HairPrimitiveId, float2 HairPrimitiveUV)
 
 float GetHairStrandsAO(uint HairPrimitiveId, float2 HairPrimitiveUV)
 {
-	const uint AoOffset = HAIR_ATTRIBUTE_OFFSET_AO(HairStrandsVF.AttributeOffsets);
+	const uint AoOffset = HAIR_POINT_ATTRIBUTE_OFFSET_AO(HairStrandsVF.PointAttributeOffsets);
 	if (IsAttributeValid(AoOffset))
 	{
-		const float AO0 = UnpackHairRoughness(InternalGetHairStrandsAttributesPerVertex(HairPrimitiveId, AoOffset, HAIR_ATTRIBUTE_STRIDE_AO, HairStrandsVF.AttributeBuffer));
-		const float AO1 = UnpackHairRoughness(InternalGetHairStrandsAttributesPerVertex(HairPrimitiveId+1, AoOffset, HAIR_ATTRIBUTE_STRIDE_AO, HairStrandsVF.AttributeBuffer));
+		const float AO0 = UnpackHairRoughness(InternalGetHairStrandsAttributesPerVertex(HairPrimitiveId, AoOffset, HAIR_POINT_ATTRIBUTE_STRIDE_AO, HairStrandsVF.PointAttributeBuffer));
+		const float AO1 = UnpackHairRoughness(InternalGetHairStrandsAttributesPerVertex(HairPrimitiveId+1, AoOffset, HAIR_POINT_ATTRIBUTE_STRIDE_AO, HairStrandsVF.PointAttributeBuffer));
 		return lerp(AO0, AO1, HairPrimitiveUV.x);
 	}
 	else
diff --git a/Engine/Shaders/Private/HairStrands/HairStrandsInterpolation.usf b/Engine/Shaders/Private/HairStrands/HairStrandsInterpolation.usf
--- a/Engine/Shaders/Private/HairStrands/HairStrandsInterpolation.usf
+++ b/Engine/Shaders/Private/HairStrands/HairStrandsInterpolation.usf
@@ -552,7 +552,7 @@ void MainCS(uint2 DispatchThreadId : SV_DispatchThreadID)
 uint				CurveCount;
 
 Buffer<uint>		RenCurveBuffer;
-uint4				RenAttributeOffsets[HAIR_ATTRIBUTE_OFFSET_COUNT];
+uint4				RenCurveAttributeOffsets[HAIR_CURVE_ATTRIBUTE_OFFSET_COUNT];
 Buffer<int>			RenVertexToClusterIdBuffer;
 
 #if PERMUTATION_SIMULATION == 1
@@ -627,7 +627,7 @@ void MainCS(uint2 DispatchThreadId : SV_DispatchThreadID, uint GroupThread1D : S
 		Out |= (SeedValues[GroupThread1D+2] & 0xFF) << 16;
 		Out |= (SeedValues[GroupThread1D+3] & 0xFF) << 24;
 
-		const uint OffsetInBytes = HAIR_ATTRIBUTE_OFFSET_SEED(RenAttributeOffsets);
+		const uint OffsetInBytes = HAIR_CURVE_ATTRIBUTE_OFFSET_SEED(RenCurveAttributeOffsets);
 		OutRenAttributeBuffer.Store(OffsetInBytes + CurveIndex, Out);
 	}
 }
diff --git a/Engine/Shaders/Private/HairStrands/HairStrandsTexturesGeneration.usf b/Engine/Shaders/Private/HairStrands/HairStrandsTexturesGeneration.usf
--- a/Engine/Shaders/Private/HairStrands/HairStrandsTexturesGeneration.usf
+++ b/Engine/Shaders/Private/HairStrands/HairStrandsTexturesGeneration.usf
@@ -287,25 +287,25 @@ bool Trace(
 							OutTangent = Tangent;
 							OutDepth = Distance;
 
-							OutSeed         = UnpackHairSeed(InternalGetHairStrandsAttributesPerCurve(P0, HAIR_ATTRIBUTE_OFFSET_SEED(InVF_AttributeOffsets), HAIR_ATTRIBUTE_STRIDE_SEED, InVF_AttributeBuffer, InVF_PointToCurveBuffer));
-							OutRootUV       = UnpackHairRootUV(InternalGetHairStrandsAttributesPerCurve(P0, HAIR_ATTRIBUTE_OFFSET_ROOTUV(InVF_AttributeOffsets), HAIR_ATTRIBUTE_STRIDE_ROOTUV, InVF_AttributeBuffer, InVF_PointToCurveBuffer));
+							OutSeed         = UnpackHairSeed(InternalGetHairStrandsAttributesPerCurve(P0, HAIR_CURVE_ATTRIBUTE_OFFSET_SEED(InVF_CurveAttributeOffsets), HAIR_CURVE_ATTRIBUTE_STRIDE_SEED, InVF_CurveAttributeBuffer, InVF_PointToCurveBuffer));
+							OutRootUV       = UnpackHairRootUV(InternalGetHairStrandsAttributesPerCurve(P0, HAIR_CURVE_ATTRIBUTE_OFFSET_ROOTUV(InVF_CurveAttributeOffsets), HAIR_CURVE_ATTRIBUTE_STRIDE_ROOTUV, InVF_CurveAttributeBuffer, InVF_PointToCurveBuffer));
 							OutU			= lerp(CP_Curr.UCoord, CP_Next.UCoord, VertexU);
 							OutMaterial.xyz = 0;
 							OutMaterial.w   = 0;
 
-							const uint BaseColorOffset = HAIR_ATTRIBUTE_OFFSET_BASECOLOR(InVF_AttributeOffsets);
-							if (IsAttributeValid(BaseColorOffset))
+							const uint ColorOffset = HAIR_POINT_ATTRIBUTE_OFFSET_COLOR(InVF_PointAttributeOffsets);
+							if (IsAttributeValid(ColorOffset))
 							{
-								const float3 C0 = UnpackHairColor(InternalGetHairStrandsAttributesPerVertex(P0, BaseColorOffset, HAIR_ATTRIBUTE_STRIDE_BASECOLOR, InVF_AttributeBuffer));
-								const float3 C1 = UnpackHairColor(InternalGetHairStrandsAttributesPerVertex(P1, BaseColorOffset, HAIR_ATTRIBUTE_STRIDE_BASECOLOR, InVF_AttributeBuffer));
+								const float3 C0 = UnpackHairColor(InternalGetHairStrandsAttributesPerVertex(P0, ColorOffset, HAIR_POINT_ATTRIBUTE_STRIDE_COLOR, InVF_PointAttributeBuffer));
+								const float3 C1 = UnpackHairColor(InternalGetHairStrandsAttributesPerVertex(P1, ColorOffset, HAIR_POINT_ATTRIBUTE_STRIDE_COLOR, InVF_PointAttributeBuffer));
 								OutMaterial.xyz = lerp(C0, C1, VertexU);
 							}
 
-							const uint RoughnessOffset = HAIR_ATTRIBUTE_OFFSET_ROUGHNESS(InVF_AttributeOffsets);
+							const uint RoughnessOffset = HAIR_POINT_ATTRIBUTE_OFFSET_ROUGHNESS(InVF_PointAttributeOffsets);
 							if (IsAttributeValid(RoughnessOffset))
 							{
-								const float R0 = UnpackHairRoughness(InternalGetHairStrandsAttributesPerVertex(P0, RoughnessOffset, HAIR_ATTRIBUTE_STRIDE_ROUGHNESS, InVF_AttributeBuffer));
-								const float R1 = UnpackHairRoughness(InternalGetHairStrandsAttributesPerVertex(P1, RoughnessOffset, HAIR_ATTRIBUTE_STRIDE_ROUGHNESS, InVF_AttributeBuffer));
+								const float R0 = UnpackHairRoughness(InternalGetHairStrandsAttributesPerVertex(P0, RoughnessOffset, HAIR_POINT_ATTRIBUTE_STRIDE_ROUGHNESS, InVF_PointAttributeBuffer));
+								const float R1 = UnpackHairRoughness(InternalGetHairStrandsAttributesPerVertex(P1, RoughnessOffset, HAIR_POINT_ATTRIBUTE_STRIDE_ROUGHNESS, InVF_PointAttributeBuffer));
 								OutMaterial.w = lerp(R0, R1, VertexU);
 							}
 
diff --git a/Engine/Shaders/Private/HairStrands/HairStrandsVertexFactoryCommon.ush b/Engine/Shaders/Private/HairStrands/HairStrandsVertexFactoryCommon.ush
--- a/Engine/Shaders/Private/HairStrands/HairStrandsVertexFactoryCommon.ush
+++ b/Engine/Shaders/Private/HairStrands/HairStrandsVertexFactoryCommon.ush
@@ -162,12 +162,14 @@ FHairCurve UnpackHairCurve(uint In)
 	uint				Prefix##_CurveCount; \
 	float4x4			Prefix##_LocalToWorldPrimitiveTransform; \
 	float4x4			Prefix##_LocalToTranslatedWorldPrimitiveTransform; \
-	uint4				Prefix##_AttributeOffsets[HAIR_ATTRIBUTE_OFFSET_COUNT]; \
+	uint4				Prefix##_CurveAttributeOffsets[HAIR_CURVE_ATTRIBUTE_OFFSET_COUNT]; \
+	uint4				Prefix##_PointAttributeOffsets[HAIR_POINT_ATTRIBUTE_OFFSET_COUNT]; \
 	Buffer<uint4>		Prefix##_PositionBuffer; \
 	Buffer<float4>		Prefix##_PositionOffsetBuffer; \
 	Buffer<uint>		Prefix##_CurveBuffer; \
 	Buffer<uint>		Prefix##_PointToCurveBuffer; \
-	ByteAddressBuffer 	Prefix##_AttributeBuffer; \
+	ByteAddressBuffer 	Prefix##_CurveAttributeBuffer; \
+	ByteAddressBuffer 	Prefix##_PointAttributeBuffer; \
 	Buffer<uint>		Prefix##_CullingIndirectBuffer; \
 	Buffer<uint>		Prefix##_CullingIndexBuffer; \
 	Buffer<float>		Prefix##_CullingRadiusScaleBuffer; \
diff --git a/Engine/Shaders/Shared/HairStrandsDefinitions.h b/Engine/Shaders/Shared/HairStrandsDefinitions.h
--- a/Engine/Shaders/Shared/HairStrandsDefinitions.h
+++ b/Engine/Shaders/Shared/HairStrandsDefinitions.h
@@ -10,16 +10,19 @@
 // Change this to force recompilation of all strata dependent shaders (use https://www.random.org/cgi-bin/randbyte?nbytes=4&format=h)
 #define HAIRSTRANDS_SHADER_VERSION 0x3ba5af0e
 
-// Attribute index
-#define HAIR_ATTRIBUTE_ROOTUV 0
-#define HAIR_ATTRIBUTE_SEED 1
-#define HAIR_ATTRIBUTE_LENGTH 2
-#define HAIR_ATTRIBUTE_BASECOLOR 3
-#define HAIR_ATTRIBUTE_ROUGHNESS 4
-#define HAIR_ATTRIBUTE_AO 5
-#define HAIR_ATTRIBUTE_CLUMPID 6
-#define HAIR_ATTRIBUTE_CLUMPID3 7
-#define HAIR_ATTRIBUTE_COUNT 8
+// Curve Attribute index
+#define HAIR_CURVE_ATTRIBUTE_ROOTUV 0
+#define HAIR_CURVE_ATTRIBUTE_SEED 1
+#define HAIR_CURVE_ATTRIBUTE_LENGTH 2
+#define HAIR_CURVE_ATTRIBUTE_CLUMPID 3
+#define HAIR_CURVE_ATTRIBUTE_CLUMPID3 4
+#define HAIR_CURVE_ATTRIBUTE_COUNT 5
+
+// Point Attribute index
+#define HAIR_POINT_ATTRIBUTE_COLOR 0
+#define HAIR_POINT_ATTRIBUTE_ROUGHNESS 1
+#define HAIR_POINT_ATTRIBUTE_AO 2
+#define HAIR_POINT_ATTRIBUTE_COUNT 3
 
 // Groom limits (based on encoding)
 #define HAIR_MAX_NUM_POINT_PER_CURVE ((1u<<8)-1u)
@@ -58,19 +61,23 @@
 
 #define MAX_HAIR_MACROGROUP_COUNT 16
 
-// HAIR_ATTRIBUTE_MAX rounded to 4
-#define HAIR_ATTRIBUTE_OFFSET_COUNT ((HAIR_ATTRIBUTE_COUNT + 3) / 4)
+// HAIR_XXX_ATTRIBUTE_MAX rounded to 4
+#define HAIR_CURVE_ATTRIBUTE_OFFSET_COUNT ((HAIR_CURVE_ATTRIBUTE_COUNT + 3) / 4)
+#define HAIR_POINT_ATTRIBUTE_OFFSET_COUNT ((HAIR_POINT_ATTRIBUTE_COUNT + 3) / 4)
 
 // Pack all offset into uint4
-#define PACK_HAIR_ATTRIBUTE_OFFSETS(Out, In) \
-	for (uint32 AttributeIt4 = 0; AttributeIt4 < HAIR_ATTRIBUTE_OFFSET_COUNT; ++AttributeIt4) \
+#define PACK_HAIR_ATTRIBUTE_OFFSETS(Out, In, Offset, Count) \
+	for (uint32 AttributeIt4 = 0; AttributeIt4 < Count; ++AttributeIt4) \
 	{ \
-		Out[AttributeIt4] = FUintVector4(0xFFFFFFFF, 0xFFFFFFFF, 0xFFFFFFFF, 0xFFFFFFFF); \
+		Out[AttributeIt4] = FUintVector4(HAIR_ATTRIBUTE_INVALID_OFFSET, HAIR_ATTRIBUTE_INVALID_OFFSET, HAIR_ATTRIBUTE_INVALID_OFFSET, HAIR_ATTRIBUTE_INVALID_OFFSET); \
 	} \
-	for (uint32 AttributeIt = 0; AttributeIt < HAIR_ATTRIBUTE_COUNT; ++AttributeIt) \
+	for (uint32 AttributeIt = 0; AttributeIt < Offset; ++AttributeIt) \
 	{ \
 		const uint32 Index4 = AttributeIt & (~0x3); \
 		const uint32 SubIndex = AttributeIt - Index4; \
 		const uint32 IndexDiv4 = Index4 >> 2u; \
 		Out[IndexDiv4][SubIndex] = In[AttributeIt]; \
-	}
\ No newline at end of file
+	}
+
+#define PACK_CURVE_HAIR_ATTRIBUTE_OFFSETS(Out, In) PACK_HAIR_ATTRIBUTE_OFFSETS(Out, In, HAIR_CURVE_ATTRIBUTE_OFFSET_COUNT, HAIR_CURVE_ATTRIBUTE_COUNT)
+#define PACK_POINT_HAIR_ATTRIBUTE_OFFSETS(Out, In) PACK_HAIR_ATTRIBUTE_OFFSETS(Out, In, HAIR_POINT_ATTRIBUTE_OFFSET_COUNT, HAIR_POINT_ATTRIBUTE_COUNT)
\ No newline at end of file
diff --git a/Engine/Source/Runtime/Renderer/Private/HairStrands/HairStrandsVisibility.cpp b/Engine/Source/Runtime/Renderer/Private/HairStrands/HairStrandsVisibility.cpp
--- a/Engine/Source/Runtime/Renderer/Private/HairStrands/HairStrandsVisibility.cpp
+++ b/Engine/Source/Runtime/Renderer/Private/HairStrands/HairStrandsVisibility.cpp
@@ -2749,7 +2749,8 @@ void AddMeshDrawTransitionPass(
 			ExternalAccessQueue.Add(VFInput.Strands.PositionBuffer.Buffer);
 			ExternalAccessQueue.Add(VFInput.Strands.PrevPositionBuffer.Buffer);
 			ExternalAccessQueue.Add(VFInput.Strands.TangentBuffer.Buffer);
-			ExternalAccessQueue.Add(VFInput.Strands.AttributeBuffer.Buffer);
+			ExternalAccessQueue.Add(VFInput.Strands.CurveAttributeBuffer.Buffer);
+			ExternalAccessQueue.Add(VFInput.Strands.PointAttributeBuffer.Buffer);
 			ExternalAccessQueue.Add(VFInput.Strands.PointToCurveBuffer.Buffer);
 			ExternalAccessQueue.Add(VFInput.Strands.PositionOffsetBuffer.Buffer);
 			ExternalAccessQueue.Add(VFInput.Strands.PrevPositionOffsetBuffer.Buffer);
@@ -2767,7 +2768,8 @@ void AddMeshDrawTransitionPass(
 				VFInput.Strands.PositionBuffer				= FRDGImportedBuffer();
 				VFInput.Strands.PrevPositionBuffer			= FRDGImportedBuffer();
 				VFInput.Strands.TangentBuffer				= FRDGImportedBuffer();
-				VFInput.Strands.AttributeBuffer				= FRDGImportedBuffer();
+				VFInput.Strands.CurveAttributeBuffer		= FRDGImportedBuffer();
+				VFInput.Strands.PointAttributeBuffer		= FRDGImportedBuffer();
 				VFInput.Strands.PointToCurveBuffer			= FRDGImportedBuffer();
 				VFInput.Strands.PositionOffsetBuffer		= FRDGImportedBuffer();
 				VFInput.Strands.PrevPositionOffsetBuffer	= FRDGImportedBuffer();
@@ -2843,7 +2845,8 @@ FHairStrandsInstanceParameters GetHairStrandsInstanceParameters(FRDGBuilder& Gra
 		Out.Resources.PositionOffsetBuffer	= Register(GraphBuilder, VFInput.Strands.PositionOffsetBufferExternal, ERDGImportedBufferFlags::CreateSRV).SRV;
 		Out.Resources.CurveBuffer			= Register(GraphBuilder, VFInput.Strands.CurveBufferExternal, ERDGImportedBufferFlags::CreateSRV).SRV;
 		Out.Resources.PointToCurveBuffer	= Register(GraphBuilder, VFInput.Strands.PointToCurveBufferExternal, ERDGImportedBufferFlags::CreateSRV).SRV;
-		Out.Resources.AttributeBuffer		= Register(GraphBuilder, VFInput.Strands.AttributeBufferExternal, ERDGImportedBufferFlags::CreateSRV).SRV;
+		Out.Resources.CurveAttributeBuffer	= Register(GraphBuilder, VFInput.Strands.CurveAttributeBufferExternal, ERDGImportedBufferFlags::CreateSRV).SRV;
+		Out.Resources.PointAttributeBuffer	= Register(GraphBuilder, VFInput.Strands.PointAttributeBufferExternal, ERDGImportedBufferFlags::CreateSRV).SRV;
 	}
 	else
 	{
@@ -2851,7 +2854,8 @@ FHairStrandsInstanceParameters GetHairStrandsInstanceParameters(FRDGBuilder& Gra
 		Out.Resources.PositionOffsetBuffer	= VFInput.Strands.PositionOffsetBuffer.SRV;
 		Out.Resources.CurveBuffer			= VFInput.Strands.CurveBuffer.SRV;
 		Out.Resources.PointToCurveBuffer	= VFInput.Strands.PointToCurveBuffer.SRV;
-		Out.Resources.AttributeBuffer		= VFInput.Strands.AttributeBuffer.SRV;
+		Out.Resources.CurveAttributeBuffer	= VFInput.Strands.CurveAttributeBuffer.SRV;
+		Out.Resources.PointAttributeBuffer	= VFInput.Strands.PointAttributeBuffer.SRV;
 	}
 
 	Out.Common = VFInput.Strands.Common;
diff --git a/Engine/Source/Runtime/Renderer/Public/HairStrandsInterface.h b/Engine/Source/Runtime/Renderer/Public/HairStrandsInterface.h
--- a/Engine/Source/Runtime/Renderer/Public/HairStrandsInterface.h
+++ b/Engine/Source/Runtime/Renderer/Public/HairStrandsInterface.h
@@ -43,7 +43,8 @@ BEGIN_SHADER_PARAMETER_STRUCT(FHairStrandsInstanceCommonParameters, RENDERER_API
 	SHADER_PARAMETER(FVector3f, PrevPositionOffset)
 	SHADER_PARAMETER(FMatrix44f, LocalToWorldPrimitiveTransform)
 	SHADER_PARAMETER(FMatrix44f, LocalToTranslatedWorldPrimitiveTransform)
-	SHADER_PARAMETER_ARRAY(FUintVector4, AttributeOffsets, [HAIR_ATTRIBUTE_OFFSET_COUNT])
+	SHADER_PARAMETER_ARRAY(FUintVector4, CurveAttributeOffsets, [HAIR_CURVE_ATTRIBUTE_OFFSET_COUNT])
+	SHADER_PARAMETER_ARRAY(FUintVector4, PointAttributeOffsets, [HAIR_POINT_ATTRIBUTE_OFFSET_COUNT])
 END_SHADER_PARAMETER_STRUCT()
 
 // Instance resources (RDG)
@@ -52,7 +53,8 @@ BEGIN_SHADER_PARAMETER_STRUCT(FHairStrandsInstanceResourceParameters, RENDERER_A
 	SHADER_PARAMETER_RDG_BUFFER_SRV(Buffer, PositionOffsetBuffer)
 	SHADER_PARAMETER_RDG_BUFFER_SRV(Buffer, CurveBuffer)
 	SHADER_PARAMETER_RDG_BUFFER_SRV(Buffer, PointToCurveBuffer)
-	SHADER_PARAMETER_RDG_BUFFER_SRV(ByteAddressBuffer, AttributeBuffer)
+	SHADER_PARAMETER_RDG_BUFFER_SRV(ByteAddressBuffer, CurveAttributeBuffer)
+	SHADER_PARAMETER_RDG_BUFFER_SRV(ByteAddressBuffer, PointAttributeBuffer)
 END_SHADER_PARAMETER_STRUCT()
 
 // Instance prev. resources (RDG)
@@ -77,7 +79,8 @@ BEGIN_SHADER_PARAMETER_STRUCT(FHairStrandsInstanceResourceRawParameters, RENDERE
 	SHADER_PARAMETER_SRV(Buffer<uint>, CurveBuffer)
 	SHADER_PARAMETER_SRV(Buffer<uint>, PointToCurveBuffer)
 	SHADER_PARAMETER_SRV(Buffer<float4>, TangentBuffer)
-	SHADER_PARAMETER_SRV(ByteAddressBuffer, AttributeBuffer)
+	SHADER_PARAMETER_SRV(ByteAddressBuffer, CurveAttributeBuffer)
+	SHADER_PARAMETER_SRV(ByteAddressBuffer, PointAttributeBuffer)
 END_SHADER_PARAMETER_STRUCT()
 
 // Instance prev. resources (Raw)
@@ -291,7 +294,8 @@ public:
 			FRDGImportedBuffer PositionBuffer;
 			FRDGImportedBuffer PrevPositionBuffer;
 			FRDGImportedBuffer TangentBuffer;
-			FRDGImportedBuffer AttributeBuffer;
+			FRDGImportedBuffer CurveAttributeBuffer;
+			FRDGImportedBuffer PointAttributeBuffer;
 			FRDGImportedBuffer PointToCurveBuffer;
 			FRDGImportedBuffer PositionOffsetBuffer;
 			FRDGImportedBuffer PrevPositionOffsetBuffer;
@@ -300,7 +304,8 @@ public:
 			FRDGExternalBuffer PositionBufferExternal;
 			FRDGExternalBuffer PrevPositionBufferExternal;
 			FRDGExternalBuffer TangentBufferExternal;
-			FRDGExternalBuffer AttributeBufferExternal;
+			FRDGExternalBuffer CurveAttributeBufferExternal;
+			FRDGExternalBuffer PointAttributeBufferExternal;
 			FRDGExternalBuffer PointToCurveBufferExternal;
 			FRDGExternalBuffer PositionOffsetBufferExternal;
 			FRDGExternalBuffer PrevPositionOffsetBufferExternal;
@@ -309,7 +314,8 @@ public:
 			FShaderResourceViewRHIRef PositionBufferRHISRV				= nullptr;
 			FShaderResourceViewRHIRef PrevPositionBufferRHISRV			= nullptr;
 			FShaderResourceViewRHIRef TangentBufferRHISRV				= nullptr;
-			FShaderResourceViewRHIRef AttributeBufferRHISRV				= nullptr;
+			FShaderResourceViewRHIRef CurveAttributeBufferRHISRV		= nullptr;
+			FShaderResourceViewRHIRef PointAttributeBufferRHISRV		= nullptr;
 			FShaderResourceViewRHIRef PointToCurveBufferRHISRV			= nullptr;
 			FShaderResourceViewRHIRef PositionOffsetBufferRHISRV		= nullptr;
 			FShaderResourceViewRHIRef PrevPositionOffsetBufferRHISRV	= nullptr;
