* Rename FHairGroupData into FHairGroupPlatformData for clarity.
* Rename cards/meshes build functions

diff --git a/Engine/Plugins/Experimental/HairModelingToolset/Source/HairModelingToolset/Private/GroomQueryUtil.cpp b/Engine/Plugins/Experimental/HairModelingToolset/Source/HairModelingToolset/Private/GroomQueryUtil.cpp
--- a/Engine/Plugins/Experimental/HairModelingToolset/Source/HairModelingToolset/Private/GroomQueryUtil.cpp
+++ b/Engine/Plugins/Experimental/HairModelingToolset/Source/HairModelingToolset/Private/GroomQueryUtil.cpp
@@ -40,7 +40,7 @@ void UE::GroomQueries::ExtractAllHairCards(AGroomActor* GroomActor,
 	for (int32 GroupIdx = 0; GroupIdx < NumHairGroups; ++GroupIdx)
 	{
 		const FHairGroupInfo& GroupInfo = Asset->HairGroupsInfo[GroupIdx];
-		const FHairGroupData& GroupData = Asset->HairGroupsData[GroupIdx];
+		const FHairGroupPlatformData& GroupData = Asset->HairGroupsPlatformData[GroupIdx];
 
 		int32 NumCardsLODs = GroupData.Cards.LODs.Num();
 		if (LODIndex >= NumCardsLODs)
@@ -520,7 +520,7 @@ void UE::GroomQueries::ProcessHairCurves(AGroomActor* GroomActor,
 	for (int32 GroupIdx = 0; GroupIdx < NumHairGroups; ++GroupIdx)
 	{
 		const FHairGroupInfo& GroupInfo = Asset->HairGroupsInfo[GroupIdx];
-		const FHairGroupData& GroupData = Asset->HairGroupsData[GroupIdx];
+		const FHairGroupPlatformData& GroupData = Asset->HairGroupsPlatformData[GroupIdx];
 
 		//int32 NumCurves = GroupInfo.NumCurves;
 		//int32 NumGuides = GroupInfo.NumGuides;
diff --git a/Engine/Plugins/Experimental/HairModelingToolset/Source/HairModelingToolset/Private/GroomToMeshTool.cpp b/Engine/Plugins/Experimental/HairModelingToolset/Source/HairModelingToolset/Private/GroomToMeshTool.cpp
--- a/Engine/Plugins/Experimental/HairModelingToolset/Source/HairModelingToolset/Private/GroomToMeshTool.cpp
+++ b/Engine/Plugins/Experimental/HairModelingToolset/Source/HairModelingToolset/Private/GroomToMeshTool.cpp
@@ -301,7 +301,7 @@ static void ProcessHairCurvePoints(AGroomActor* GroomActor,
 	for (int32 GroupIdx = 0; GroupIdx < NumHairGroups; ++GroupIdx)
 	{
 		const FHairGroupInfo& GroupInfo = Asset->HairGroupsInfo[GroupIdx];
-		const FHairGroupData& GroupData = Asset->HairGroupsData[GroupIdx];
+		const FHairGroupPlatformData& GroupData = Asset->HairGroupsPlatformData[GroupIdx];
 
 		//int32 NumCurves = GroupInfo.NumCurves;
 		//int32 NumGuides = GroupInfo.NumGuides;
@@ -347,7 +347,7 @@ static void ProcessHairCurves(AGroomActor* GroomActor,
 	for (int32 GroupIdx = 0; GroupIdx < NumHairGroups; ++GroupIdx)
 	{
 		const FHairGroupInfo& GroupInfo = Asset->HairGroupsInfo[GroupIdx];
-		const FHairGroupData& GroupData = Asset->HairGroupsData[GroupIdx];
+		const FHairGroupPlatformData& GroupData = Asset->HairGroupsPlatformData[GroupIdx];
 
 		//int32 NumCurves = GroupInfo.NumCurves;
 		//int32 NumGuides = GroupInfo.NumGuides;
diff --git a/Engine/Plugins/Importers/USDImporter/Source/USDSchemas/Private/USDGroomTranslator.cpp b/Engine/Plugins/Importers/USDImporter/Source/USDSchemas/Private/USDGroomTranslator.cpp
--- a/Engine/Plugins/Importers/USDImporter/Source/USDSchemas/Private/USDGroomTranslator.cpp
+++ b/Engine/Plugins/Importers/USDImporter/Source/USDSchemas/Private/USDGroomTranslator.cpp
@@ -241,7 +241,7 @@ protected:
 					GroomCacheProcessor = MakeUnique<FGroomCacheProcessor>(EGroomCacheType::Strands, AnimInfo.Attributes);
 
 					// ref. FGroomCacheImporter::ImportGroomCache
-					const TArray<FHairGroupData>& GroomHairGroupsData = GroomAsset->HairGroupsData;
+					const TArray<FHairGroupPlatformData>& GroomHairGroupsData = GroomAsset->HairGroupsPlatformData;
 
 					// Each frame is translated into a HairDescription and processed into HairGroupData
 					// Sample one extra frame so that we can interpolate between EndFrame - 1 and EndFrame
diff --git a/Engine/Plugins/Importers/USDImporter/Source/USDSchemas/Private/USDGroomTranslatorUtils.cpp b/Engine/Plugins/Importers/USDImporter/Source/USDSchemas/Private/USDGroomTranslatorUtils.cpp
--- a/Engine/Plugins/Importers/USDImporter/Source/USDSchemas/Private/USDGroomTranslatorUtils.cpp
+++ b/Engine/Plugins/Importers/USDImporter/Source/USDSchemas/Private/USDGroomTranslatorUtils.cpp
@@ -119,7 +119,7 @@ namespace UE::UsdGroomTranslatorUtils::Private
 				GroomBinding->SourceGeometryCache = Cast<UGeometryCache>(Settings.SourceMesh);
 				GroomBinding->TargetGeometryCache = Cast<UGeometryCache>(Settings.TargetMesh);
 			}
-			GroomBinding->HairGroupBulkDatas.Reserve(Settings.Groom->HairGroupsData.Num());
+			GroomBinding->HairGroupBulkDatas.Reserve(Settings.Groom->HairGroupsPlatformData.Num());
 			GroomBinding->NumInterpolationPoints = Settings.NumInterpolationPoints;
 			GroomBinding->MatchingSection = Settings.MatchingSection;
 
diff --git a/Engine/Plugins/Runtime/HairStrands/Source/HairStrandsCore/Private/GroomAsset.cpp b/Engine/Plugins/Runtime/HairStrands/Source/HairStrandsCore/Private/GroomAsset.cpp
--- a/Engine/Plugins/Runtime/HairStrands/Source/HairStrandsCore/Private/GroomAsset.cpp
+++ b/Engine/Plugins/Runtime/HairStrands/Source/HairStrandsCore/Private/GroomAsset.cpp
@@ -119,7 +119,7 @@ static FAutoConsoleVariableRef CVarHairStrandsDump_All(
 	TEXT("Dump all the loaded groom assets, groom binding assets, and instanciated groom components."),
 	FConsoleVariableDelegate::CreateStatic(DumpLoadedGroomData));
 
-uint32 FHairGroupData::FStrands::GetDataSize() const
+uint32 FHairGroupPlatformData::FStrands::GetDataSize() const
 {
 	uint32 Total = 0;
 	Total += FBaseWithInterpolation::GetDataSize();
@@ -161,18 +161,18 @@ void DumpLoadedGroomAssets(IConsoleVariable* InCVarPakTesterEnabled)
 	{
 		if (AssetIt)
 		{
-			const uint32 GroupCount = AssetIt->HairGroupsData.Num();
+			const uint32 GroupCount = AssetIt->HairGroupsPlatformData.Num();
 			for (uint32 GroupIt = 0; GroupIt < GroupCount; ++GroupIt)
 			{
-				const uint32 CPUMemorySize_Guides  = AssetIt->HairGroupsData[GroupIt].Guides.GetDataSize();
-				const uint32 CPUMemorySize_Strands = AssetIt->HairGroupsData[GroupIt].Strands.GetDataSize();
-				const uint32 CPUMemorySize_Cards   = AssetIt->HairGroupsData[GroupIt].Cards.GetDataSize();
-				const uint32 CPUMemorySize_Meshes  = AssetIt->HairGroupsData[GroupIt].Meshes.GetDataSize();
+				const uint32 CPUMemorySize_Guides  = AssetIt->HairGroupsPlatformData[GroupIt].Guides.GetDataSize();
+				const uint32 CPUMemorySize_Strands = AssetIt->HairGroupsPlatformData[GroupIt].Strands.GetDataSize();
+				const uint32 CPUMemorySize_Cards   = AssetIt->HairGroupsPlatformData[GroupIt].Cards.GetDataSize();
+				const uint32 CPUMemorySize_Meshes  = AssetIt->HairGroupsPlatformData[GroupIt].Meshes.GetDataSize();
 
-				const uint32 GPUMemorySize_Guides  = AssetIt->HairGroupsData[GroupIt].Guides.GetResourcesSize();
-				const uint32 GPUMemorySize_Strands = AssetIt->HairGroupsData[GroupIt].Strands.GetResourcesSize();
-				const uint32 GPUMemorySize_Cards   = AssetIt->HairGroupsData[GroupIt].Cards.GetResourcesSize();
-				const uint32 GPUMemorySize_Meshes  = AssetIt->HairGroupsData[GroupIt].Meshes.GetResourcesSize();
+				const uint32 GPUMemorySize_Guides  = AssetIt->HairGroupsPlatformData[GroupIt].Guides.GetResourcesSize();
+				const uint32 GPUMemorySize_Strands = AssetIt->HairGroupsPlatformData[GroupIt].Strands.GetResourcesSize();
+				const uint32 GPUMemorySize_Cards   = AssetIt->HairGroupsPlatformData[GroupIt].Cards.GetResourcesSize();
+				const uint32 GPUMemorySize_Meshes  = AssetIt->HairGroupsPlatformData[GroupIt].Meshes.GetResourcesSize();
 
 				uint32 CPUMemorySize = 0;
 				uint32 GPUMemorySize = 0;
@@ -499,7 +499,7 @@ static bool BuildHairGroup(
 	const TArray<FHairGroupsInterpolation>& HairGroupsInterpolation,
 	const TArray<FHairGroupsLOD>& HairGroupsLOD,
 	TArray<FHairGroupInfoWithVisibility>& OutHairGroupsInfo,
-	TArray<FHairGroupData>& OutHairGroupsData)
+	TArray<FHairGroupPlatformData>& OutHairGroupsData)
 {
 	const bool bIsValid = HairDescriptionGroups.IsValid();
 	if (bIsValid)
@@ -542,7 +542,7 @@ static bool BuildHairGroupCluster(
 	const TArray<FHairGroupsInterpolation>& HairGroupsInterpolation,
 	const TArray<FHairGroupsLOD>& HairGroupsLOD,
 	TArray<FHairGroupInfoWithVisibility>& OutHairGroupsInfo,
-	TArray<FHairGroupData>& OutHairGroupsData)
+	TArray<FHairGroupPlatformData>& OutHairGroupsData)
 {
 	const bool bIsValid = HairDescriptionGroups.IsValid();
 	if (bIsValid)
@@ -649,9 +649,9 @@ void UGroomAsset::ApplyStripFlags(uint8 StripFlags, const ITargetPlatform* CookT
 	const bool bIsStrandsStrippedForCook = !!(StripFlags & CDSF_StrandsStripped);
 
 	// Set the CookedOut flags as appropriate and compute the value for EffectiveLODBias
-	for (int32 GroupIndex = 0; GroupIndex < HairGroupsData.Num(); ++GroupIndex)
+	for (int32 GroupIndex = 0; GroupIndex < HairGroupsPlatformData.Num(); ++GroupIndex)
 	{
-		FHairGroupData& GroupData = HairGroupsData[GroupIndex];
+		FHairGroupPlatformData& GroupData = HairGroupsPlatformData[GroupIndex];
 		GroupData.bIsCookedOut = bIsStrandsStrippedForCook;
 
 		// Determine the max LOD for strands since they are cooked out if the platform doesn't support them
@@ -700,14 +700,14 @@ bool UGroomAsset::HasImportedStrandsData() const
 #endif
 
 // Serialization for *array* of hair elements
-static void InternalSerialize(FArchive& Ar, UObject* Owner, TArray<FHairGroupData::FCards::FLOD>& CardLODData);
-static void InternalSerialize(FArchive& Ar, UObject* Owner, TArray<FHairGroupData::FMeshes::FLOD>& MeshLODData);
-static void InternalSerialize(FArchive& Ar, UObject* Owner, TArray<FHairGroupData>& GroupData);
+static void InternalSerialize(FArchive& Ar, UObject* Owner, TArray<FHairGroupPlatformData::FCards::FLOD>& CardLODData);
+static void InternalSerialize(FArchive& Ar, UObject* Owner, TArray<FHairGroupPlatformData::FMeshes::FLOD>& MeshLODData);
+static void InternalSerialize(FArchive& Ar, UObject* Owner, TArray<FHairGroupPlatformData>& GroupData);
 
 // Serialization for hair elements
-static void InternalSerialize(FArchive& Ar, UObject* Owner, FHairGroupData::FCards::FLOD& CardLODData);
-static void InternalSerialize(FArchive& Ar, UObject* Owner, FHairGroupData::FMeshes::FLOD& MeshLODData);
-static void InternalSerialize(FArchive& Ar, UObject* Owner, FHairGroupData& GroupData);
+static void InternalSerialize(FArchive& Ar, UObject* Owner, FHairGroupPlatformData::FCards::FLOD& CardLODData);
+static void InternalSerialize(FArchive& Ar, UObject* Owner, FHairGroupPlatformData::FMeshes::FLOD& MeshLODData);
+static void InternalSerialize(FArchive& Ar, UObject* Owner, FHairGroupPlatformData& GroupData);
 
 void UGroomAsset::Serialize(FArchive& Ar)
 {
@@ -746,7 +746,7 @@ void UGroomAsset::Serialize(FArchive& Ar)
 		{
 			// When cooking data or serializing old format to new format,
 			// serialize the computed groom data
-			InternalSerialize(Ar, this, HairGroupsData);
+			InternalSerialize(Ar, this, HairGroupsPlatformData);
 		}
 #if WITH_EDITORONLY_DATA
 		else
@@ -830,7 +830,7 @@ TArray<FHairVertexFactoryTypesPerMaterialData> UGroomAsset::CollectVertexFactory
 	const int32 GroupCount = GetNumHairGroups();
 	for (int32 GroupIt = 0; GroupIt < GroupCount; GroupIt++)
 	{
-		const FHairGroupData& InGroupData = HairGroupsData[GroupIt];
+		const FHairGroupPlatformData& InGroupData = HairGroupsPlatformData[GroupIt];
 
 		// Strands
 		{
@@ -842,7 +842,7 @@ TArray<FHairVertexFactoryTypesPerMaterialData> UGroomAsset::CollectVertexFactory
 		if (IsHairStrandsEnabled(EHairStrandsShaderType::Cards, ShaderPlatform))
 		{
 			uint32 CardsLODIndex = 0;
-			for (const FHairGroupData::FCards::FLOD& LOD : InGroupData.Cards.LODs)
+			for (const FHairGroupPlatformData::FCards::FLOD& LOD : InGroupData.Cards.LODs)
 			{
 				if (LOD.IsValid())
 				{
@@ -865,7 +865,7 @@ TArray<FHairVertexFactoryTypesPerMaterialData> UGroomAsset::CollectVertexFactory
 		// Meshes
 		{
 			uint32 MeshesLODIndex = 0;
-			for (const FHairGroupData::FMeshes::FLOD& LOD : InGroupData.Meshes.LODs)
+			for (const FHairGroupPlatformData::FMeshes::FLOD& LOD : InGroupData.Meshes.LODs)
 			{
 				if (LOD.IsValid())
 				{
@@ -933,7 +933,7 @@ void UGroomAsset::UpdateResource()
 		AllChangeType = AllChangeType | ChangeType;
 
 		check(GroupIndex < uint32(GetNumHairGroups()));
-		FHairGroupData& GroupData = HairGroupsData[GroupIndex];
+		FHairGroupPlatformData& GroupData = HairGroupsPlatformData[GroupIndex];
 		if (GroupData.Guides.RestResource)
 		{
 			InternalUpdateResource(GroupData.Guides.RestResource);
@@ -958,7 +958,7 @@ void UGroomAsset::UpdateResource()
 					HairGroupsInterpolation,
 					HairGroupsLOD,
 					HairGroupsInfo,
-					HairGroupsData);
+					HairGroupsPlatformData);
 
 				GroupData.Strands.ClusterCullingResource = new FHairStrandsClusterCullingResource(GroupData.Strands.ClusterCullingBulkData, FHairResourceName(GetFName(), GroupIndex), GetAssetPathName());
 			}
@@ -977,8 +977,8 @@ void UGroomAsset::UpdateResource()
 		}
 	}
 
-	if (BuildCardsGeometry())  { AllChangeType |= GroomChangeType_Cards;  }
-	if (BuildMeshesGeometry()) { AllChangeType |= GroomChangeType_Meshes; }
+	if (BuildCardsData())  { AllChangeType |= GroomChangeType_Cards;  }
+	if (BuildMeshesData()) { AllChangeType |= GroomChangeType_Meshes; }
 
 	if (AllChangeType & (GroomChangeType_LOD | GroomChangeType_Cards | GroomChangeType_Meshes))
 	{
@@ -1022,8 +1022,8 @@ void UGroomAsset::ReleaseResource()
 
 void UGroomAsset::ReleaseGuidesResource(uint32 GroupIndex)
 {
-	if (GroupIndex >= uint32(HairGroupsData.Num())) return;
-	FHairGroupData& GroupData = HairGroupsData[GroupIndex];
+	if (GroupIndex >= uint32(HairGroupsPlatformData.Num())) return;
+	FHairGroupPlatformData& GroupData = HairGroupsPlatformData[GroupIndex];
 	if (GroupData.Guides.IsValid())
 	{
 		InternalReleaseResource(GroupData.Guides.RestResource);
@@ -1032,8 +1032,8 @@ void UGroomAsset::ReleaseGuidesResource(uint32 GroupIndex)
 
 void UGroomAsset::ReleaseStrandsResource(uint32 GroupIndex)
 {
-	if (GroupIndex >= uint32(HairGroupsData.Num())) return;
-	FHairGroupData& GroupData = HairGroupsData[GroupIndex];
+	if (GroupIndex >= uint32(HairGroupsPlatformData.Num())) return;
+	FHairGroupPlatformData& GroupData = HairGroupsPlatformData[GroupIndex];
 	if (GroupData.Strands.IsValid())
 	{
 		InternalReleaseResource(GroupData.Strands.RestResource);
@@ -1047,9 +1047,9 @@ void UGroomAsset::ReleaseStrandsResource(uint32 GroupIndex)
 
 void UGroomAsset::ReleaseCardsResource(uint32 GroupIndex)
 {
-	if (GroupIndex >= uint32(HairGroupsData.Num())) return;
-	FHairGroupData& GroupData = HairGroupsData[GroupIndex];
-	for (FHairGroupData::FCards::FLOD& LOD : GroupData.Cards.LODs)
+	if (GroupIndex >= uint32(HairGroupsPlatformData.Num())) return;
+	FHairGroupPlatformData& GroupData = HairGroupsPlatformData[GroupIndex];
+	for (FHairGroupPlatformData::FCards::FLOD& LOD : GroupData.Cards.LODs)
 	{
 		InternalReleaseResource(LOD.RestResource);
 		InternalReleaseResource(LOD.InterpolationResource);
@@ -1063,9 +1063,9 @@ void UGroomAsset::ReleaseCardsResource(uint32 GroupIndex)
 
 void UGroomAsset::ReleaseMeshesResource(uint32 GroupIndex)
 {
-	if (GroupIndex >= uint32(HairGroupsData.Num())) return;
-	FHairGroupData& GroupData = HairGroupsData[GroupIndex];
-	for (FHairGroupData::FMeshes::FLOD& LOD : GroupData.Meshes.LODs)
+	if (GroupIndex >= uint32(HairGroupsPlatformData.Num())) return;
+	FHairGroupPlatformData& GroupData = HairGroupsPlatformData[GroupIndex];
+	for (FHairGroupPlatformData::FMeshes::FLOD& LOD : GroupData.Meshes.LODs)
 	{
 		InternalReleaseResource(LOD.RestResource);
 		#if RHI_RAYTRACING
@@ -1084,7 +1084,7 @@ void UGroomAsset::UpdateHairGroupsInfo()
 	HairGroupsInfo.SetNum(GroupCount);
 
 	uint32 GroupIndex = 0;
-	for (FHairGroupData& Data: HairGroupsData)
+	for (FHairGroupPlatformData& Data: HairGroupsPlatformData)
 	{
 		FHairGroupInfoWithVisibility& Info = HairGroupsInfo[GroupIndex];
 		Info.GroupID = GroupIndex;
@@ -1251,9 +1251,9 @@ void UGroomAsset::PostLoad()
 		{
 			const bool bNeedInterpoldationData = NeedsInterpolationData(GroupIndex);
 			const bool bNeedToBuildData = 
-				(HairGroupsData[GroupIndex].Guides.BulkData.GetNumCurves() == 0 ||
-				 (bNeedInterpoldationData && HairGroupsData[GroupIndex].Strands.InterpolationBulkData.GetPointCount() == 0)) &&
-				 HairGroupsData[GroupIndex].Strands.BulkData.GetNumCurves() > 0; // Empty groom has no data to build
+				(HairGroupsPlatformData[GroupIndex].Guides.BulkData.GetNumCurves() == 0 ||
+				 (bNeedInterpoldationData && HairGroupsPlatformData[GroupIndex].Strands.InterpolationBulkData.GetPointCount() == 0)) &&
+				 HairGroupsPlatformData[GroupIndex].Strands.BulkData.GetNumCurves() > 0; // Empty groom has no data to build
 			if (bNeedToBuildData)
 			{
 				BuildHairGroup(
@@ -1263,7 +1263,7 @@ void UGroomAsset::PostLoad()
 					HairGroupsInterpolation,
 					HairGroupsLOD,
 					HairGroupsInfo,
-					HairGroupsData);
+					HairGroupsPlatformData);
 			}
 		}
 	}
@@ -1370,7 +1370,7 @@ void UGroomAsset::PreSave(FObjectPreSaveContext ObjectSaveContext)
 						HairGroupsInterpolation,
 						HairGroupsLOD,
 						HairGroupsInfo,
-						HairGroupsData);
+						HairGroupsPlatformData);
 				}
 				else if (bHasLODChanged)
 				{
@@ -1380,7 +1380,7 @@ void UGroomAsset::PreSave(FObjectPreSaveContext ObjectSaveContext)
 						HairGroupsInterpolation,
 						HairGroupsLOD,
 						HairGroupsInfo,
-						HairGroupsData);
+						HairGroupsPlatformData);
 				}
 			}
 		}
@@ -1477,7 +1477,7 @@ void UGroomAsset::PostEditChangeProperty(FPropertyChangedEvent& PropertyChangedE
 	if (bGeometryTypeChanged)
 	{
 		// If the asset didn't have any mesh or cards, we need to create/allocate the cards/mesh groups
-		const uint32 GroupCount = HairGroupsData.Num();
+		const uint32 GroupCount = HairGroupsPlatformData.Num();
 		if (IsHairStrandsEnabled(EHairStrandsShaderType::Cards) && HasGeometryType(EGroomGeometryType::Cards) && HairGroupsCards.Num() == 0)
 		{
 			HairGroupsCards.Init(FHairGroupsCardsSourceDescription(), GroupCount);
@@ -1577,12 +1577,12 @@ void UGroomAsset::PostInitProperties()
 
 int32 UGroomAsset::GetNumHairGroups() const
 {
-	return HairGroupsData.Num();
+	return HairGroupsPlatformData.Num();
 }
 
 ////////////////////////////////////////////////////////////////////////////////////////////////////////////
 // Array serialize
-void InternalSerialize(FArchive& Ar, UObject* Owner, TArray<FHairGroupData::FMeshes::FLOD>& LODs)
+void InternalSerialize(FArchive& Ar, UObject* Owner, TArray<FHairGroupPlatformData::FMeshes::FLOD>& LODs)
 {
 	uint32 Count = LODs.Num();
 	Ar << Count;
@@ -1596,7 +1596,7 @@ void InternalSerialize(FArchive& Ar, UObject* Owner, TArray<FHairGroupData::FMes
 	}
 }
 
-void InternalSerialize(FArchive& Ar, UObject* Owner, TArray<FHairGroupData::FCards::FLOD>& LODs)
+void InternalSerialize(FArchive& Ar, UObject* Owner, TArray<FHairGroupPlatformData::FCards::FLOD>& LODs)
 {
 	uint32 Count = LODs.Num();
 	Ar << Count;
@@ -1610,7 +1610,7 @@ void InternalSerialize(FArchive& Ar, UObject* Owner, TArray<FHairGroupData::FCar
 	}
 }
 
-void InternalSerialize(FArchive& Ar, UObject* Owner, TArray<FHairGroupData>& GroupDatas)
+void InternalSerialize(FArchive& Ar, UObject* Owner, TArray<FHairGroupPlatformData>& GroupDatas)
 {
 	uint32 GroupCount = GroupDatas.Num();
 	Ar << GroupCount;
@@ -1626,7 +1626,7 @@ void InternalSerialize(FArchive& Ar, UObject* Owner, TArray<FHairGroupData>& Gro
 
 ////////////////////////////////////////////////////////////////////////////////////////////////////////////
 // Element serialize
-static void InternalSerialize(FArchive& Ar, UObject* Owner, FHairGroupData::FCards::FLOD& CardLODData)
+static void InternalSerialize(FArchive& Ar, UObject* Owner, FHairGroupPlatformData::FCards::FLOD& CardLODData)
 {
 	if (!Ar.IsCooking() || !CardLODData.bIsCookedOut)
 	{
@@ -1645,13 +1645,13 @@ static void InternalSerialize(FArchive& Ar, UObject* Owner, FHairGroupData::FCar
 		FHairCardsInterpolationBulkData NoInterpolationBulkData;
 		NoInterpolationBulkData.Serialize(Ar);
 
-		FHairGroupData::FBaseWithInterpolation NoGuideData;
+		FHairGroupPlatformData::FBaseWithInterpolation NoGuideData;
 		NoGuideData.BulkData.Serialize(Ar, Owner);
 		NoGuideData.InterpolationBulkData.Serialize(Ar, Owner);
 	}
 }
 
-static void InternalSerialize(FArchive& Ar, UObject* Owner, FHairGroupData::FMeshes::FLOD& MeshLODData)
+static void InternalSerialize(FArchive& Ar, UObject* Owner, FHairGroupPlatformData::FMeshes::FLOD& MeshLODData)
 {
 	if (!Ar.IsCooking() || !MeshLODData.bIsCookedOut)
 	{
@@ -1665,11 +1665,11 @@ static void InternalSerialize(FArchive& Ar, UObject* Owner, FHairGroupData::FMes
 	}
 }
 
-static void InternalSerialize(FArchive& Ar, UObject* Owner, FHairGroupData& GroupData)
+static void InternalSerialize(FArchive& Ar, UObject* Owner, FHairGroupPlatformData& GroupData)
 {
 	Ar.UsingCustomVersion(FAnimObjectVersion::GUID);
 
-	FHairGroupData NoStrandsData;
+	FHairGroupPlatformData NoStrandsData;
 	if (!Ar.IsCooking() || !GroupData.bIsCookedOut)
 	{
 		GroupData.Strands.BulkData.Serialize(Ar, Owner);
@@ -1770,7 +1770,7 @@ bool UGroomAsset::AreGroupsValid() const
 	const uint32 GroupCount = HairGroupsInterpolation.Num();
 	return
 		GroupCount > 0 &&
-		HairGroupsData.Num() == GroupCount &&
+		HairGroupsPlatformData.Num() == GroupCount &&
 		HairGroupsPhysics.Num() == GroupCount &&
 		HairGroupsRendering.Num() == GroupCount &&
 		HairGroupsLOD.Num() == GroupCount;
@@ -1781,13 +1781,13 @@ void UGroomAsset::SetNumGroup(uint32 InGroupCount, bool bResetGroupData, bool bR
 	ReleaseResource();
 	if (bResetGroupData)
 	{
-		HairGroupsData.Reset();
+		HairGroupsPlatformData.Reset();
 	}
 
 	// In order to preserve the existing asset settings, the settings are only reset if the group count has changed.
-	if (InGroupCount != HairGroupsData.Num())
+	if (InGroupCount != HairGroupsPlatformData.Num())
 	{
-		HairGroupsData.Init(FHairGroupData(), InGroupCount);
+		HairGroupsPlatformData.Init(FHairGroupPlatformData(), InGroupCount);
 	}
 
 	if (InGroupCount != HairGroupsPhysics.Num())
@@ -1941,7 +1941,7 @@ namespace GroomDerivedDataCacheUtils
 			return FString();
 		}
 
-		if (Desc.GroupIndex >= GroomAsset->HairGroupsData.Num() || Desc.LODIndex >= GroomAsset->HairGroupsData[Desc.GroupIndex].Cards.LODs.Num())
+		if (Desc.GroupIndex >= GroomAsset->HairGroupsPlatformData.Num() || Desc.LODIndex >= GroomAsset->HairGroupsPlatformData[Desc.GroupIndex].Cards.LODs.Num())
 		{
 			return FString();
 		}
@@ -2293,8 +2293,8 @@ bool UGroomAsset::CacheDerivedData(uint32 GroupIndex)
 
 	// Cache each geometric representation separately
 	bSuccess |= CacheStrandsData(GroupIndex, DerivedDataKey);
-	bSuccess |= CacheCardsGeometry(GroupIndex, DerivedDataKey);
-	bSuccess |= CacheMeshesGeometry(GroupIndex);
+	bSuccess |= CacheCardsData(GroupIndex, DerivedDataKey);
+	bSuccess |= CacheMeshesData(GroupIndex);
 
 	return bSuccess;
 }
@@ -2342,10 +2342,10 @@ bool UGroomAsset::CacheStrandsData(uint32 GroupIndex, FString& OutDerivedDataKey
 		UE_CLOG(IsHairStrandsDDCLogEnable(), LogHairStrands, Log, TEXT("[Groom/DDC] Strands - Found (Groom:%s Group6:%d)."), *GetName(), GroupIndex);
 
 		// Reset hair group data to ensure previously loaded bulk data are cleared/cleaned priori to load new data.
-		HairGroupsData[GroupIndex] = FHairGroupData();
+		HairGroupsPlatformData[GroupIndex] = FHairGroupPlatformData();
 
 		FMemoryReaderView Ar(Data, /*bIsPersistent*/ true);
-		InternalSerialize(Ar, this, HairGroupsData[GroupIndex]);
+		InternalSerialize(Ar, this, HairGroupsPlatformData[GroupIndex]);
 	}
 	else
 	{
@@ -2376,13 +2376,13 @@ bool UGroomAsset::CacheStrandsData(uint32 GroupIndex, FString& OutDerivedDataKey
 			HairGroupsInterpolation,
 			HairGroupsLOD,
 			HairGroupsInfo,
-			HairGroupsData);
+			HairGroupsPlatformData);
 		
 		if (bSuccess)
 		{
 			TArray<uint8> WriteData;
 			FMemoryWriter Ar(WriteData, /*bIsPersistent*/ true);
-			InternalSerialize(Ar, this, HairGroupsData[GroupIndex]);
+			InternalSerialize(Ar, this, HairGroupsPlatformData[GroupIndex]);
 			FRequestOwner AsyncOwner(EPriority::Normal);
 			GetCache().PutValue({ {Name, Key, FValue::Compress(MakeSharedBufferFromArray(MoveTemp(WriteData)))} }, AsyncOwner);
 			AsyncOwner.KeepAlive();
@@ -2397,7 +2397,7 @@ bool UGroomAsset::CacheStrandsData(uint32 GroupIndex, FString& OutDerivedDataKey
 
 bool UGroomAsset::HasValidMeshesData(uint32 GroupIndex) const
 {
-	const FHairGroupData& HairGroupData = HairGroupsData[GroupIndex];
+	const FHairGroupPlatformData& HairGroupData = HairGroupsPlatformData[GroupIndex];
 	const FHairGroupsLOD& GroupsLOD = HairGroupsLOD[GroupIndex];
 	for (int32 LODIt = 0; LODIt < GroupsLOD.LODs.Num(); ++LODIt)
 	{
@@ -2418,7 +2418,7 @@ bool UGroomAsset::HasValidMeshesData(uint32 GroupIndex) const
 
 bool UGroomAsset::HasValidCardsData(uint32 GroupIndex) const
 {
-	const FHairGroupData& HairGroupData = HairGroupsData[GroupIndex];
+	const FHairGroupPlatformData& HairGroupData = HairGroupsPlatformData[GroupIndex];
 	const FHairGroupsLOD& GroupsLOD = HairGroupsLOD[GroupIndex];
 	for (int32 LODIt = 0; LODIt < GroupsLOD.LODs.Num(); ++LODIt)
 	{
@@ -2439,7 +2439,7 @@ bool UGroomAsset::HasValidCardsData(uint32 GroupIndex) const
 }
 
 // Return true if the meshes data have changed
-bool UGroomAsset::CacheCardsGeometry(uint32 GroupIndex, const FString& StrandsKey)
+bool UGroomAsset::CacheCardsData(uint32 GroupIndex, const FString& StrandsKey)
 {
 	const bool bHasValidCardsData = HasValidCardsData(GroupIndex);
 	if (!bHasValidCardsData)
@@ -2458,7 +2458,7 @@ bool UGroomAsset::CacheCardsGeometry(uint32 GroupIndex, const FString& StrandsKe
 		return false;
 	}
 
-	FHairGroupData& HairGroupData = HairGroupsData[GroupIndex];
+	FHairGroupPlatformData& HairGroupData = HairGroupsPlatformData[GroupIndex];
 
 	// We should never query the DDC for something that we won't be able to build. So check that first.
 	// The goal is to avoid paying the price of a guaranteed failure on high-latency DDC backends every run.
@@ -2512,7 +2512,7 @@ bool UGroomAsset::CacheCardsGeometry(uint32 GroupIndex, const FString& StrandsKe
 	{
 		UE_CLOG(IsHairStrandsDDCLogEnable(), LogHairStrands, Log, TEXT("[Groom/DDC] Cards - Not found (Groom:%s Group:%d)."), *GetName(), GroupIndex);
 
-		if (!BuildCardsGeometry(GroupIndex))
+		if (!BuildCardsData(GroupIndex))
 		{
 			CardsDerivedDataKey[GroupIndex].Empty();
 			return false;
@@ -2528,7 +2528,7 @@ bool UGroomAsset::CacheCardsGeometry(uint32 GroupIndex, const FString& StrandsKe
 
 	// Handle the case where the cards data is already cached in the DDC
 	// Need to populate the strands data with it
-	FHairGroupData& GroupData = HairGroupsData[GroupIndex];
+	FHairGroupPlatformData& GroupData = HairGroupsPlatformData[GroupIndex];
 	if (StrandsDerivedDataKey[GroupIndex].IsEmpty() && GroupData.Strands.BulkData.GetNumPoints() == 0)
 	{
 		for (int32 LODIt = 0; LODIt < GroupData.Cards.LODs.Num(); ++LODIt)
@@ -2536,7 +2536,7 @@ bool UGroomAsset::CacheCardsGeometry(uint32 GroupIndex, const FString& StrandsKe
 			int32 SourceIt = 0;
 			if (FHairGroupsCardsSourceDescription* Desc = GetSourceDescription(HairGroupsCards, GroupIndex, LODIt, SourceIt))
 			{
-				FHairGroupData::FCards::FLOD& LOD = GroupData.Cards.LODs[LODIt];
+				FHairGroupPlatformData::FCards::FLOD& LOD = GroupData.Cards.LODs[LODIt];
 				GroupData.Strands.BulkData = LOD.Guides.BulkData;
 				GroupData.Guides.BulkData = LOD.Guides.BulkData;
 				break;
@@ -2560,7 +2560,7 @@ bool UGroomAsset::CacheCardsGeometry(uint32 GroupIndex, const FString& StrandsKe
 }
 
 // Return true if the meshes data have changed
-bool UGroomAsset::CacheMeshesGeometry(uint32 GroupIndex)
+bool UGroomAsset::CacheMeshesData(uint32 GroupIndex)
 {
 	const bool bHasValidMeshesData = UGroomAsset::HasValidMeshesData(GroupIndex);
 	if (!bHasValidMeshesData)
@@ -2579,7 +2579,7 @@ bool UGroomAsset::CacheMeshesGeometry(uint32 GroupIndex)
 		return false;
 	}
 
-	FHairGroupData& HairGroupData = HairGroupsData[GroupIndex];
+	FHairGroupPlatformData& HairGroupData = HairGroupsPlatformData[GroupIndex];
 
 	// We should never query the DDC for something that we won't be able to build. So check that first.
 	// The goal is to avoid paying the price of a guaranteed failure on high-latency DDC backends every run.
@@ -2623,7 +2623,7 @@ bool UGroomAsset::CacheMeshesGeometry(uint32 GroupIndex)
 	{
 		UE_CLOG(IsHairStrandsDDCLogEnable(), LogHairStrands, Log, TEXT("[Groom/DDC] Meshes - Not found (Groom:%s Group:%d)."), *GetName(), GroupIndex);
 
-		if (!BuildMeshesGeometry(GroupIndex))
+		if (!BuildMeshesData(GroupIndex))
 		{
 			MeshesDerivedDataKey[GroupIndex].Empty();
 			return false;
@@ -2647,7 +2647,7 @@ inline FString GetLODName(const UGroomAsset* Asset, uint32 LODIndex)
 	return Asset->GetOutermost()->GetName() + FString::Printf(TEXT("_LOD%d"), LODIndex);
 }
 
-bool UGroomAsset::BuildCardsGeometry(uint32 GroupIndex)
+bool UGroomAsset::BuildCardsData(uint32 GroupIndex)
 {
 	LLM_SCOPE_BYTAG(Groom);
 
@@ -2659,7 +2659,7 @@ bool UGroomAsset::BuildCardsGeometry(uint32 GroupIndex)
 	bool bHasChanged = HairGroupsCards.Num() != CachedHairGroupsCards.Num();
 
 	check(GroupIndex < uint32(GetNumHairGroups()));
-	FHairGroupData& GroupData = HairGroupsData[GroupIndex];
+	FHairGroupPlatformData& GroupData = HairGroupsPlatformData[GroupIndex];
 
 	// The settings might have been previously cached without the data having been built
 	const uint32 LODCount = HairGroupsLOD[GroupIndex].LODs.Num();
@@ -2731,7 +2731,7 @@ bool UGroomAsset::BuildCardsGeometry(uint32 GroupIndex)
 				return false;
 			}
 
-			FHairGroupData::FCards::FLOD& LOD = GroupData.Cards.LODs[LODIt];
+			FHairGroupPlatformData::FCards::FLOD& LOD = GroupData.Cards.LODs[LODIt];
 			LOD.BulkData.Reset();
 
 			// 0. Release geometry resources
@@ -2861,12 +2861,12 @@ bool UGroomAsset::BuildCardsGeometry(uint32 GroupIndex)
 	return bDataBuilt;
 }
 
-bool UGroomAsset::BuildCardsGeometry()
+bool UGroomAsset::BuildCardsData()
 {
 	bool bDataChanged = false;
 	for (uint32 GroupIndex = 0, GroupCount = GetNumHairGroups(); GroupIndex < GroupCount; ++GroupIndex)
 	{
-		bDataChanged |= CacheCardsGeometry(GroupIndex, StrandsDerivedDataKey[GroupIndex]);
+		bDataChanged |= CacheCardsData(GroupIndex, StrandsDerivedDataKey[GroupIndex]);
 	}
 
 	if (bDataChanged)
@@ -2874,7 +2874,7 @@ bool UGroomAsset::BuildCardsGeometry()
 		InitGuideResources();
 
 		// When building cards in an empty groom, the strands resources need to be initialized once
-		FHairGroupData& GroupData = HairGroupsData[0];
+		FHairGroupPlatformData& GroupData = HairGroupsPlatformData[0];
 		if (!GroupData.Strands.ClusterCullingResource)
 		{
 			InitStrandsResources();
@@ -2886,12 +2886,12 @@ bool UGroomAsset::BuildCardsGeometry()
 	return bDataChanged;
 }
 
-bool UGroomAsset::BuildMeshesGeometry()
+bool UGroomAsset::BuildMeshesData()
 {
 	bool bDataChanged = false;
 	for (uint32 GroupIndex = 0, GroupCount = GetNumHairGroups(); GroupIndex < GroupCount; ++GroupIndex)
 	{
-		bDataChanged |= CacheMeshesGeometry(GroupIndex);
+		bDataChanged |= CacheMeshesData(GroupIndex);
 	}
 
 	if (bDataChanged)
@@ -2902,7 +2902,7 @@ bool UGroomAsset::BuildMeshesGeometry()
 	return bDataChanged;
 }
 
-bool UGroomAsset::BuildMeshesGeometry(uint32 GroupIndex)
+bool UGroomAsset::BuildMeshesData(uint32 GroupIndex)
 {
 	if (!IsHairStrandsEnabled(EHairStrandsShaderType::Meshes) || HairGroupsMeshes.Num() == 0)
 	{
@@ -2911,7 +2911,7 @@ bool UGroomAsset::BuildMeshesGeometry(uint32 GroupIndex)
 
 	bool bHasChanged = HairGroupsMeshes.Num() != CachedHairGroupsMeshes.Num();
 
-	FHairGroupData& GroupData = HairGroupsData[GroupIndex];
+	FHairGroupPlatformData& GroupData = HairGroupsPlatformData[GroupIndex];
 
 	// The settings might have been previously cached without the data having been built
 	const uint32 LODCount = HairGroupsLOD[GroupIndex].LODs.Num();
@@ -2995,7 +2995,7 @@ bool UGroomAsset::BuildMeshesGeometry(uint32 GroupIndex)
 				return false;
 			}
 
-			FHairGroupData::FMeshes::FLOD& LOD = GroupData.Meshes.LODs[LODIt];
+			FHairGroupPlatformData::FMeshes::FLOD& LOD = GroupData.Meshes.LODs[LODIt];
 
 			if (Desc->ImportedMesh)
 			{
@@ -3056,7 +3056,7 @@ FHairStrandsRestResource* UGroomAsset::AllocateGuidesResources(uint32 GroupIndex
 {
 	if (GroupIndex < uint32(GetNumHairGroups()))
 	{
-		FHairGroupData& GroupData = HairGroupsData[GroupIndex];
+		FHairGroupPlatformData& GroupData = HairGroupsPlatformData[GroupIndex];
 		if (GroupData.Guides.HasValidData())
 		{
 			if (GroupData.Guides.RestResource == nullptr)
@@ -3073,7 +3073,7 @@ FHairStrandsInterpolationResource* UGroomAsset::AllocateInterpolationResources(u
 {
 	if (GroupIndex < uint32(GetNumHairGroups()))
 	{
-		FHairGroupData& GroupData = HairGroupsData[GroupIndex];
+		FHairGroupPlatformData& GroupData = HairGroupsPlatformData[GroupIndex];
 		check(GroupData.Guides.IsValid());
 		if (GroupData.Strands.InterpolationResource == nullptr)
 		{
@@ -3089,8 +3089,8 @@ FHairStrandsRaytracingResource* UGroomAsset::AllocateCardsRaytracingResources(ui
 {
 	if (IsRayTracingAllowed() && GroupIndex < uint32(GetNumHairGroups()) && LODIndex < uint32(HairGroupsLOD[GroupIndex].LODs.Num()))
 	{
-		FHairGroupData& GroupData = HairGroupsData[GroupIndex];
-		FHairGroupData::FCards::FLOD& LOD = GroupData.Cards.LODs[LODIndex];
+		FHairGroupPlatformData& GroupData = HairGroupsPlatformData[GroupIndex];
+		FHairGroupPlatformData::FCards::FLOD& LOD = GroupData.Cards.LODs[LODIndex];
 		check(LOD.BulkData.IsValid());
 
 		if (LOD.RaytracingResource == nullptr)
@@ -3106,8 +3106,8 @@ FHairStrandsRaytracingResource* UGroomAsset::AllocateMeshesRaytracingResources(u
 {
 	if (IsRayTracingAllowed() && GroupIndex < uint32(GetNumHairGroups()) && LODIndex < uint32(HairGroupsLOD[GroupIndex].LODs.Num()))
 	{
-		FHairGroupData& GroupData = HairGroupsData[GroupIndex];
-		FHairGroupData::FMeshes::FLOD& LOD = GroupData.Meshes.LODs[LODIndex];
+		FHairGroupPlatformData& GroupData = HairGroupsPlatformData[GroupIndex];
+		FHairGroupPlatformData::FMeshes::FLOD& LOD = GroupData.Meshes.LODs[LODIndex];
 		check(LOD.BulkData.IsValid());
 
 		if (LOD.RaytracingResource == nullptr)
@@ -3123,7 +3123,7 @@ FHairStrandsRaytracingResource* UGroomAsset::AllocateStrandsRaytracingResources(
 {
 	if (IsRayTracingAllowed() && GroupIndex < uint32(GetNumHairGroups()))
 	{
-		FHairGroupData& GroupData = HairGroupsData[GroupIndex];
+		FHairGroupPlatformData& GroupData = HairGroupsPlatformData[GroupIndex];
 		check(GroupData.Strands.HasValidData());
 
 		if (GroupData.Strands.RaytracingResource == nullptr)
@@ -3154,7 +3154,7 @@ void UGroomAsset::InitStrandsResources()
 
 	for (uint32 GroupIndex = 0, GroupCount = GetNumHairGroups(); GroupIndex < GroupCount; ++GroupIndex)
 	{
-		FHairGroupData& GroupData = HairGroupsData[GroupIndex];
+		FHairGroupPlatformData& GroupData = HairGroupsPlatformData[GroupIndex];
 
 		if (GroupData.Strands.HasValidData())
 		{
@@ -3198,14 +3198,14 @@ static void InitCardsTextureResources(UGroomAsset* GroomAsset)
 
 	for (uint32 GroupIndex = 0, GroupCount = GroomAsset->GetNumHairGroups(); GroupIndex < GroupCount; ++GroupIndex)
 	{
-		FHairGroupData& GroupData = GroomAsset->HairGroupsData[GroupIndex];
+		FHairGroupPlatformData& GroupData = GroomAsset->HairGroupsPlatformData[GroupIndex];
 
 		const uint32 LODCount = GroomAsset->HairGroupsLOD[GroupIndex].LODs.Num();
 		GroupData.Cards.LODs.SetNum(LODCount);
 
 		for (uint32 LODIt = 0; LODIt < LODCount; ++LODIt)
 		{
-			FHairGroupData::FCards::FLOD& LOD = GroupData.Cards.LODs[LODIt];
+			FHairGroupPlatformData::FCards::FLOD& LOD = GroupData.Cards.LODs[LODIt];
 
 			int32 SourceIt = 0;
 			const FHairGroupsCardsSourceDescription* Desc = GetSourceDescription(GroomAsset->HairGroupsCards, GroupIndex, LODIt, SourceIt);
@@ -3245,14 +3245,14 @@ void UGroomAsset::InitCardsResources()
 
 	for (uint32 GroupIndex = 0, GroupCount = GetNumHairGroups(); GroupIndex < GroupCount; ++GroupIndex)
 	{
-		FHairGroupData& GroupData = HairGroupsData[GroupIndex];
+		FHairGroupPlatformData& GroupData = HairGroupsPlatformData[GroupIndex];
 
 		const uint32 LODCount = HairGroupsLOD[GroupIndex].LODs.Num();
 		GroupData.Cards.LODs.SetNum(LODCount);
 
 		for (uint32 LODIt = 0; LODIt < LODCount; ++LODIt)
 		{
-			FHairGroupData::FCards::FLOD& LOD = GroupData.Cards.LODs[LODIt];
+			FHairGroupPlatformData::FCards::FLOD& LOD = GroupData.Cards.LODs[LODIt];
 
 			int32 SourceIt = 0;
 			const FHairGroupsCardsSourceDescription* Desc = GetSourceDescription(HairGroupsCards, GroupIndex, LODIt, SourceIt);
@@ -3309,7 +3309,7 @@ void UGroomAsset::InitMeshesResources()
 
 	for (uint32 GroupIndex = 0, GroupCount = GetNumHairGroups(); GroupIndex < GroupCount; ++GroupIndex)
 	{
-		FHairGroupData& GroupData = HairGroupsData[GroupIndex];
+		FHairGroupPlatformData& GroupData = HairGroupsPlatformData[GroupIndex];
 
 		const uint32 LODCount = GroupData.Meshes.LODs.Num();
 		for (uint32 LODIt = 0; LODIt < LODCount; ++LODIt)
@@ -3320,7 +3320,7 @@ void UGroomAsset::InitMeshesResources()
 			{
 				continue;
 			}
-			FHairGroupData::FMeshes::FLOD& LOD = GroupData.Meshes.LODs[LODIt];
+			FHairGroupPlatformData::FMeshes::FLOD& LOD = GroupData.Meshes.LODs[LODIt];
 			InternalReleaseResource(LOD.RestResource);
 
 			if (LOD.HasValidData())
@@ -3448,7 +3448,7 @@ bool UGroomAsset::IsGlobalInterpolationEnable(int32 GroupIndex, int32 LODIndex)
 
 bool UGroomAsset::IsSimulationEnable() const
 {
-	for (int32 GroupIndex = 0; GroupIndex < HairGroupsData.Num(); ++GroupIndex)
+	for (int32 GroupIndex = 0; GroupIndex < HairGroupsPlatformData.Num(); ++GroupIndex)
 	{
 		for (int32 LODIt = 0; LODIt < HairGroupsLOD.Num(); ++LODIt)
 		{
@@ -3477,7 +3477,7 @@ bool UGroomAsset::NeedsInterpolationData(int32 GroupIndex) const
 
 bool UGroomAsset::NeedsInterpolationData() const
 {
-	for (int32 GroupIndex = 0; GroupIndex < HairGroupsData.Num(); ++GroupIndex)
+	for (int32 GroupIndex = 0; GroupIndex < HairGroupsPlatformData.Num(); ++GroupIndex)
 	{
 		if (NeedsInterpolationData(GroupIndex))
 		{
@@ -3547,7 +3547,7 @@ void UGroomAsset::StripLODs(const TArray<int32>& LODsToKeep, bool bRebuildResour
 		const FHairDescriptionGroups& LocalHairDescriptionGroups = GetHairDescriptionGroups();
 		for (int32 GroupIt = 0; GroupIt < GroupCount; ++GroupIt)
 		{
-			FHairGroupData& GroupData = HairGroupsData[GroupIt];
+			FHairGroupPlatformData& GroupData = HairGroupsPlatformData[GroupIt];
 			InternalReleaseResource(GroupData.Strands.ClusterCullingResource);
 
 			BuildHairGroupCluster(
@@ -3556,7 +3556,7 @@ void UGroomAsset::StripLODs(const TArray<int32>& LODsToKeep, bool bRebuildResour
 				HairGroupsInterpolation,
 				HairGroupsLOD,
 				HairGroupsInfo,
-				HairGroupsData);
+				HairGroupsPlatformData);
 		}
 	}
 }
@@ -3566,7 +3566,7 @@ bool UGroomAsset::HasDebugData() const
 {
 	for (uint32 GroupIndex = 0, GroupCount = GetNumHairGroups(); GroupIndex < GroupCount; ++GroupIndex)
 	{
-		if (!HairGroupsData[GroupIndex].Debug.Data.IsValid())
+		if (!HairGroupsPlatformData[GroupIndex].Debug.Data.IsValid())
 		{
 			return false;
 		}
@@ -3590,7 +3590,7 @@ void UGroomAsset::CreateDebugData()
 		FHairStrandsDatas GuidesData;
 		FGroomBuilder::BuildData(LocalHairDescriptionGroups.HairGroups[GroupIndex], HairGroupsInterpolation[GroupIndex], DummyInfo, StrandsData, GuidesData);
 
-		FHairGroupData& GroupData = HairGroupsData[GroupIndex];
+		FHairGroupPlatformData& GroupData = HairGroupsPlatformData[GroupIndex];
 		CreateHairStrandsDebugDatas(StrandsData, GroupData.Debug.Data);
 
 		if (GroupData.Debug.Data.IsValid())
@@ -3704,7 +3704,7 @@ void UGroomAsset::SaveProceduralCards(uint32 DescIndex)
 
 	const int32 GroupIndex = Desc->GroupIndex;
 	const int32 LODIndex = Desc->LODIndex;
-	if (GroupIndex >= HairGroupsData.Num())
+	if (GroupIndex >= HairGroupsPlatformData.Num())
 		return;
 
 	// 1. Convert old parameters (ClusterDecimation, bUseCards) to new parameters (GenerationType & CardsCount)
@@ -3714,7 +3714,7 @@ void UGroomAsset::SaveProceduralCards(uint32 DescIndex)
 		const bool bNeedConversion = ClusterSettings.ClusterDecimation > 0;
 		if (bNeedConversion)
 		{
-			const int32 MaxCardCount = HairGroupsData[GroupIndex].Strands.BulkData.GetNumCurves();
+			const int32 MaxCardCount = HairGroupsPlatformData[GroupIndex].Strands.BulkData.GetNumCurves();
 
 			GeometrySettings.GenerationType = ClusterSettings.bUseGuide ? EHairCardsGenerationType::UseGuides : EHairCardsGenerationType::CardsCount;
 			GeometrySettings.CardsCount = FMath::Clamp(FMath::CeilToInt(ClusterSettings.ClusterDecimation * MaxCardCount), 1, MaxCardCount);
@@ -3899,9 +3899,9 @@ void UGroomAsset::GetResourceSizeEx(FResourceSizeEx& CumulativeResourceSize)
 	CumulativeResourceSize.AddDedicatedSystemMemoryBytes(HairGroupsCards.GetAllocatedSize());
 	CumulativeResourceSize.AddDedicatedSystemMemoryBytes(HairGroupsMeshes.GetAllocatedSize());
 	CumulativeResourceSize.AddDedicatedSystemMemoryBytes(HairGroupsMaterials.GetAllocatedSize());
-	CumulativeResourceSize.AddDedicatedSystemMemoryBytes(HairGroupsData.GetAllocatedSize());
+	CumulativeResourceSize.AddDedicatedSystemMemoryBytes(HairGroupsPlatformData.GetAllocatedSize());
 	
-	for (const FHairGroupData & GroupData : HairGroupsData)
+	for (const FHairGroupPlatformData & GroupData : HairGroupsPlatformData)
 	{		
 		CumulativeResourceSize.AddDedicatedVideoMemoryBytes(GroupData.Guides.GetResourcesSize());
 		CumulativeResourceSize.AddDedicatedVideoMemoryBytes(GroupData.Strands.GetResourcesSize());
diff --git a/Engine/Plugins/Runtime/HairStrands/Source/HairStrandsCore/Private/GroomBindingAsset.cpp b/Engine/Plugins/Runtime/HairStrands/Source/HairStrandsCore/Private/GroomBindingAsset.cpp
--- a/Engine/Plugins/Runtime/HairStrands/Source/HairStrandsCore/Private/GroomBindingAsset.cpp
+++ b/Engine/Plugins/Runtime/HairStrands/Source/HairStrandsCore/Private/GroomBindingAsset.cpp
@@ -538,7 +538,7 @@ bool UGroomBindingAsset::IsCompatible(const UGroomAsset* InGroom, const UGroomBi
 		for (uint32 GroupIt = 0; GroupIt < GroupCount; ++GroupIt)
 		{
 			{
-				const uint32 GroomCount = InGroom->HairGroupsData[GroupIt].Guides.BulkData.GetNumCurves();
+				const uint32 GroomCount = InGroom->HairGroupsPlatformData[GroupIt].Guides.BulkData.GetNumCurves();
 				const uint32 BindingCount = InBinding->GroupInfos[GroupIt].SimRootCount;
 
 				if (GroomCount != 0 && GroomCount != BindingCount)
@@ -558,7 +558,7 @@ bool UGroomBindingAsset::IsCompatible(const UGroomAsset* InGroom, const UGroomBi
 
 			if (IsHairStrandsEnabled(EHairStrandsShaderType::Strands))
 			{
-				const uint32 GroomCount = InGroom->HairGroupsData[GroupIt].Strands.BulkData.GetNumCurves();
+				const uint32 GroomCount = InGroom->HairGroupsPlatformData[GroupIt].Strands.BulkData.GetNumCurves();
 				const uint32 BindingCount = InBinding->GroupInfos[GroupIt].RenRootCount;
 
 				// Groom may have stripped strands data so GroomCount would be 0
diff --git a/Engine/Plugins/Runtime/HairStrands/Source/HairStrandsCore/Private/GroomBindingBuilder.cpp b/Engine/Plugins/Runtime/HairStrands/Source/HairStrandsCore/Private/GroomBindingBuilder.cpp
--- a/Engine/Plugins/Runtime/HairStrands/Source/HairStrandsCore/Private/GroomBindingBuilder.cpp
+++ b/Engine/Plugins/Runtime/HairStrands/Source/HairStrandsCore/Private/GroomBindingBuilder.cpp
@@ -2067,7 +2067,7 @@ static bool InternalBuildBinding_CPU(UGroomBindingAsset* BindingAsset, bool bIni
 		const uint32 MeshLODCount = TargetMeshData->GetNumLODs();
 
 		uint32 GroupIndex = 0;
-		for (const FHairGroupData& GroupData : GroomAsset->HairGroupsData)
+		for (const FHairGroupPlatformData& GroupData : GroomAsset->HairGroupsPlatformData)
 		{
 			FHairStrandsDatas StrandsData;
 			FHairStrandsDatas GuidesData;
@@ -2162,7 +2162,7 @@ static bool InternalBuildBinding_CPU(UGroomBindingAsset* BindingAsset, bool bIni
 			const uint32 CardsLODCount = OutHairGroupDatas[GroupIt].CardsRootData.Num();
 			for (uint32 CardsLODIt = 0; CardsLODIt < CardsLODCount; ++CardsLODIt)
 			{
-				if (BindingAsset->Groom->HairGroupsData[GroupIt].Cards.IsValid(CardsLODIt))
+				if (BindingAsset->Groom->HairGroupsPlatformData[GroupIt].Cards.IsValid(CardsLODIt))
 				{
 					FHairStrandsDatas LODGuidesData;
 					const bool bIsValid = GroomAsset->GetHairCardsGuidesDatas(GroupIt, CardsLODIt, LODGuidesData);
diff --git a/Engine/Plugins/Runtime/HairStrands/Source/HairStrandsCore/Private/GroomComponent.cpp b/Engine/Plugins/Runtime/HairStrands/Source/HairStrandsCore/Private/GroomComponent.cpp
--- a/Engine/Plugins/Runtime/HairStrands/Source/HairStrandsCore/Private/GroomComponent.cpp
+++ b/Engine/Plugins/Runtime/HairStrands/Source/HairStrandsCore/Private/GroomComponent.cpp
@@ -135,7 +135,7 @@ static FHairGroupDesc GetGroomGroupsDesc(const UGroomAsset* Asset, UGroomCompone
 	}
 
 	FHairGroupDesc O = Component->GroomGroupsDesc[GroupIndex];
-	O.HairLength = Asset->HairGroupsData[GroupIndex].Strands.BulkData.MaxLength;
+	O.HairLength = Asset->HairGroupsPlatformData[GroupIndex].Strands.BulkData.MaxLength;
 	O.LODBias 	 = Asset->EffectiveLODBias[GroupIndex] > 0 ? FMath::Max(O.LODBias, Asset->EffectiveLODBias[GroupIndex]) : O.LODBias;
 
 	// Width can be overridden at 3 levels:
@@ -147,7 +147,7 @@ static FHairGroupDesc GetGroomGroupsDesc(const UGroomAsset* Asset, UGroomCompone
 		O.HairWidth = 
 			Asset->HairGroupsRendering[GroupIndex].GeometrySettings.HairWidth_Override ? 
 			Asset->HairGroupsRendering[GroupIndex].GeometrySettings.HairWidth : 
-			Asset->HairGroupsData[GroupIndex].Strands.BulkData.MaxRadius * 2.0f;
+			Asset->HairGroupsPlatformData[GroupIndex].Strands.BulkData.MaxRadius * 2.0f;
 	}
 
 	if (!O.HairRootScale_Override)				{ O.HairRootScale				= Asset->HairGroupsRendering[GroupIndex].GeometrySettings.HairRootScale;				}
@@ -349,12 +349,12 @@ public:
 		const EShaderPlatform ShaderPlatform = GetScene().GetShaderPlatform();
 
 		const int32 GroupCount = Component->GroomAsset->GetNumHairGroups();
-		check(Component->GroomAsset->HairGroupsData.Num() == Component->HairGroupInstances.Num());
+		check(Component->GroomAsset->HairGroupsPlatformData.Num() == Component->HairGroupInstances.Num());
 		for (int32 GroupIt=0; GroupIt<GroupCount; GroupIt++)
 		{
 			const bool bIsVisible = Component->GroomAsset->HairGroupsInfo[GroupIt].bIsVisible;
 
-			const FHairGroupData& InGroupData = Component->GroomAsset->HairGroupsData[GroupIt];
+			const FHairGroupPlatformData& InGroupData = Component->GroomAsset->HairGroupsPlatformData[GroupIt];
 			FHairGroupInstance* HairInstance = Component->HairGroupInstances[GroupIt];
 			check(HairInstance->HairGroupPublicData);
 			HairInstance->bForceCards = Component->bUseCards;
@@ -411,7 +411,7 @@ public:
 			if (IsHairStrandsEnabled(EHairStrandsShaderType::Cards, ShaderPlatform))
 			{
 				uint32 CardsLODIndex = 0;
-				for (const FHairGroupData::FCards::FLOD& LOD : InGroupData.Cards.LODs)
+				for (const FHairGroupPlatformData::FCards::FLOD& LOD : InGroupData.Cards.LODs)
 				{
 					if (LOD.IsValid())
 					{
@@ -437,7 +437,7 @@ public:
 			if (IsHairStrandsEnabled(EHairStrandsShaderType::Meshes, ShaderPlatform))
 			{
 				uint32 MeshesLODIndex = 0;
-				for (const FHairGroupData::FMeshes::FLOD& LOD : InGroupData.Meshes.LODs)
+				for (const FHairGroupPlatformData::FMeshes::FLOD& LOD : InGroupData.Meshes.LODs)
 				{
 					if (LOD.IsValid())
 					{
@@ -1377,7 +1377,7 @@ void UGroomComponent::SwitchSimulationLOD(const int32 PreviousLOD, const int32 C
 	if (GroomAsset && PreviousLOD != CurrentLOD)
 	{
 		bool bRequiresSimulationUpdate = false;
-		for (int32 GroupIt = 0, GroupCount = GroomAsset->HairGroupsData.Num(); GroupIt < GroupCount; ++GroupIt)
+		for (int32 GroupIt = 0, GroupCount = GroomAsset->HairGroupsPlatformData.Num(); GroupIt < GroupCount; ++GroupIt)
 		{
 			if ((IsSimulationEnable(GroupIt, PreviousLOD) != IsSimulationEnable(GroupIt, CurrentLOD)))
 			{
@@ -1560,7 +1560,7 @@ void UGroomComponent::SetForcedLOD(int32 CurrLODIndex)
 		bool bNeedSimulationNeed = false;
 		if (GroomAsset)
 		{
-			for (int32 GroupIt = 0, GroupCount = GroomAsset->HairGroupsData.Num(); GroupIt < GroupCount; ++GroupIt)
+			for (int32 GroupIt = 0, GroupCount = GroomAsset->HairGroupsPlatformData.Num(); GroupIt < GroupCount; ++GroupIt)
 			{
 				for (uint32 LODIt = 0, LODCount = GroomAsset->GetLODCount(); LODIt < LODCount; ++LODIt)
 				{
@@ -1755,7 +1755,7 @@ FBoxSphereBounds UGroomComponent::CalcBounds(const FTransform& InLocalToWorld) c
 		FBox LocalHairBound(EForceInit::ForceInitToZero);
 		if (!GroomCacheBuffers.IsValid())
 		{
-			for (const FHairGroupData& GroupData : GroomAsset->HairGroupsData)
+			for (const FHairGroupPlatformData& GroupData : GroomAsset->HairGroupsPlatformData)
 			{
 				if (IsHairStrandsEnabled(EHairStrandsShaderType::Strands) && GroupData.Strands.HasValidData())
 				{
@@ -1929,7 +1929,7 @@ EHairGeometryType UGroomComponent::GetMaterialGeometryType(int32 ElementIndex) c
 	for (uint32 GroupIt = 0, GroupCount = GroomAsset->HairGroupsRendering.Num(); GroupIt < GroupCount; ++GroupIt)
 	{
 		// Material - Strands
-		const FHairGroupData& InGroupData = GroomAsset->HairGroupsData[GroupIt];
+		const FHairGroupPlatformData& InGroupData = GroomAsset->HairGroupsPlatformData[GroupIt];
 		if (IsHairStrandsEnabled(EHairStrandsShaderType::Strands, Platform))
 		{
 			const int32 SlotIndex = GroomAsset->GetMaterialIndex(GroomAsset->HairGroupsRendering[GroupIt].MaterialSlotName);
@@ -1943,7 +1943,7 @@ EHairGeometryType UGroomComponent::GetMaterialGeometryType(int32 ElementIndex) c
 		if (IsHairStrandsEnabled(EHairStrandsShaderType::Cards, Platform))
 		{
 			uint32 CardsLODIndex = 0;
-			for (const FHairGroupData::FCards::FLOD& LOD : InGroupData.Cards.LODs)
+			for (const FHairGroupPlatformData::FCards::FLOD& LOD : InGroupData.Cards.LODs)
 			{
 				if (LOD.IsValid())
 				{
@@ -1970,7 +1970,7 @@ EHairGeometryType UGroomComponent::GetMaterialGeometryType(int32 ElementIndex) c
 		if (IsHairStrandsEnabled(EHairStrandsShaderType::Meshes, Platform))
 		{
 			uint32 MeshesLODIndex = 0;
-			for (const FHairGroupData::FMeshes::FLOD& LOD : InGroupData.Meshes.LODs)
+			for (const FHairGroupPlatformData::FMeshes::FLOD& LOD : InGroupData.Meshes.LODs)
 			{
 				if (LOD.IsValid())
 				{
@@ -2010,11 +2010,11 @@ FHairStrandsRestResource* UGroomComponent::GetGuideStrandsRestResource(uint32 Gr
 		return nullptr;
 	}
 
-	if (!GroomAsset->HairGroupsData[GroupIndex].Guides.IsValid())
+	if (!GroomAsset->HairGroupsPlatformData[GroupIndex].Guides.IsValid())
 	{
 		return nullptr;
 	}
-	return GroomAsset->HairGroupsData[GroupIndex].Guides.RestResource;
+	return GroomAsset->HairGroupsPlatformData[GroupIndex].Guides.RestResource;
 }
 
 FHairStrandsDeformedResource* UGroomComponent::GetGuideStrandsDeformedResource(uint32 GroupIndex)
@@ -2197,11 +2197,11 @@ static USkeletalMeshComponent* ValidateBindingAsset(
 	}
 
 	// Validate against cards data
-	for (int32 GroupIt = 0, GroupCount = GroomAsset->HairGroupsData.Num(); GroupIt < GroupCount; ++GroupIt)
+	for (int32 GroupIt = 0, GroupCount = GroomAsset->HairGroupsPlatformData.Num(); GroupIt < GroupCount; ++GroupIt)
 	{
-		FHairGroupData& GroupData = GroomAsset->HairGroupsData[GroupIt];
+		FHairGroupPlatformData& GroupData = GroomAsset->HairGroupsPlatformData[GroupIt];
 		uint32 CardsLODIndex = 0;
-		for (FHairGroupData::FCards::FLOD& LOD : GroupData.Cards.LODs)
+		for (FHairGroupPlatformData::FCards::FLOD& LOD : GroupData.Cards.LODs)
 		{
 			if (LOD.IsValid())
 			{
@@ -2299,14 +2299,14 @@ void UGroomComponent::InitResources(bool bIsBindingReloading)
 	TArray<bool> bHasNeedGlobalDeformation;
 	TArray<bool> bHasNeedSimulation;
 	TArray<bool> bHasNeedDeformation;
-	bHasNeedGlobalDeformation.Init(false, GroomAsset->HairGroupsData.Num());
-	bHasNeedSimulation.Init(false, GroomAsset->HairGroupsData.Num());
-	bHasNeedDeformation.Init(false, GroomAsset->HairGroupsData.Num());
+	bHasNeedGlobalDeformation.Init(false, GroomAsset->HairGroupsPlatformData.Num());
+	bHasNeedSimulation.Init(false, GroomAsset->HairGroupsPlatformData.Num());
+	bHasNeedDeformation.Init(false, GroomAsset->HairGroupsPlatformData.Num());
 	
 	bool bHasAnyNeedSimulation = false;
 	bool bHasAnyNeedDeformation = false;
 	bool bHasAnyNeedGlobalDeformation = false;
-	for (int32 GroupIt = 0, GroupCount = GroomAsset->HairGroupsData.Num(); GroupIt < GroupCount; ++GroupIt)
+	for (int32 GroupIt = 0, GroupCount = GroomAsset->HairGroupsPlatformData.Num(); GroupIt < GroupCount; ++GroupIt)
 	{
 		for (uint32 LODIt = 0, LODCount = GroomAsset->GetLODCount(); LODIt < LODCount; ++LODIt)
 		{
@@ -2406,7 +2406,7 @@ void UGroomComponent::InitResources(bool bIsBindingReloading)
 					}
 				}
 
-				for (int32 GroupIt = 0, GroupCount = GroomAsset->HairGroupsData.Num(); GroupIt < GroupCount; ++GroupIt)
+				for (int32 GroupIt = 0, GroupCount = GroomAsset->HairGroupsPlatformData.Num(); GroupIt < GroupCount; ++GroupIt)
 				{
 					for (uint32 LODIt = 0, LODCount = GroomAsset->GetLODCount(); LODIt < LODCount; ++LODIt)
 					{
@@ -2439,7 +2439,7 @@ void UGroomComponent::InitResources(bool bIsBindingReloading)
 		if (bHasNeedBindingData)
 		{
 			LocalBindingAsset = BindingAsset;
-			if (LocalBindingAsset && LocalBindingAsset->HairGroupResources.Num() != GroomAsset->HairGroupsData.Num())
+			if (LocalBindingAsset && LocalBindingAsset->HairGroupResources.Num() != GroomAsset->HairGroupsPlatformData.Num())
 			{
 				LocalBindingAsset = nullptr;
 			}
@@ -2458,7 +2458,7 @@ void UGroomComponent::InitResources(bool bIsBindingReloading)
 	{
 		bHasNeedSkinningBinding = false;
 		bHasAnyNeedGlobalDeformation = false;
-		bHasNeedGlobalDeformation.Init(false, GroomAsset->HairGroupsData.Num());
+		bHasNeedGlobalDeformation.Init(false, GroomAsset->HairGroupsPlatformData.Num());
 	}
 
 	if (GroomCache)
@@ -2471,7 +2471,7 @@ void UGroomComponent::InitResources(bool bIsBindingReloading)
 		GroomCacheBuffers.Reset();
 	}
 
-	for (int32 GroupIt = 0, GroupCount = GroomAsset->HairGroupsData.Num(); GroupIt < GroupCount; ++GroupIt)
+	for (int32 GroupIt = 0, GroupCount = GroomAsset->HairGroupsPlatformData.Num(); GroupIt < GroupCount; ++GroupIt)
 	{
 		FHairGroupInstance* HairGroupInstance = new FHairGroupInstance();
 		HairGroupInstance->AddRef();
@@ -2505,7 +2505,7 @@ void UGroomComponent::InitResources(bool bIsBindingReloading)
 		HairGroupInstance->Debug.RigidPreviousLocalToWorld		= HairGroupInstance->Debug.RigidCurrentLocalToWorld;
 		HairGroupInstance->LocalToWorld							= HairGroupInstance->GetCurrentLocalToWorld();
 
-		FHairGroupData& GroupData = GroomAsset->HairGroupsData[GroupIt];
+		FHairGroupPlatformData& GroupData = GroomAsset->HairGroupsPlatformData[GroupIt];
 
 		const EHairInterpolationType HairInterpolationType = ToHairInterpolationType(GroomAsset->HairInterpolationType);
 
@@ -2742,7 +2742,7 @@ void UGroomComponent::InitResources(bool bIsBindingReloading)
 		}
 
 		// Cards resources
-		for (FHairGroupData::FCards::FLOD& LOD : GroupData.Cards.LODs)
+		for (FHairGroupPlatformData::FCards::FLOD& LOD : GroupData.Cards.LODs)
 		{
 			const uint32 CardsLODIndex = HairGroupInstance->Cards.LODs.Num();
 			FHairGroupInstance::FCards::FLOD& InstanceLOD = HairGroupInstance->Cards.LODs.AddDefaulted_GetRef();
@@ -2822,7 +2822,7 @@ void UGroomComponent::InitResources(bool bIsBindingReloading)
 		}
 
 		// Meshes resources
-		for (FHairGroupData::FMeshes::FLOD& LOD : GroupData.Meshes.LODs)
+		for (FHairGroupPlatformData::FMeshes::FLOD& LOD : GroupData.Meshes.LODs)
 		{
 			const int32 MeshLODIndex = HairGroupInstance->Meshes.LODs.Num();
 			FHairGroupInstance::FMeshes::FLOD& InstanceLOD = HairGroupInstance->Meshes.LODs.AddDefaulted_GetRef();
diff --git a/Engine/Plugins/Runtime/HairStrands/Source/HairStrandsCore/Private/GroomRBFDeformer.cpp b/Engine/Plugins/Runtime/HairStrands/Source/HairStrandsCore/Private/GroomRBFDeformer.cpp
--- a/Engine/Plugins/Runtime/HairStrands/Source/HairStrandsCore/Private/GroomRBFDeformer.cpp
+++ b/Engine/Plugins/Runtime/HairStrands/Source/HairStrandsCore/Private/GroomRBFDeformer.cpp
@@ -399,7 +399,7 @@ static void ApplyDeformationToGroom(const TArray<FRBFDeformedPositions>& Deforme
 				const FHairGroupsLOD& HairGroupLOD = GroomAsset->HairGroupsLOD[GroupIndex];
 
 				FHairGroupInfo& HairGroupsInfo = GroomAsset->HairGroupsInfo[GroupIndex];
-				FHairGroupData& HairGroupsData = GroomAsset->HairGroupsData[GroupIndex];
+				FHairGroupPlatformData& HairGroupsData = GroomAsset->HairGroupsPlatformData[GroupIndex];
 
 				FHairStrandsDatas StrandsData;
 				FHairStrandsDatas GuidesData;
@@ -422,8 +422,8 @@ static void ApplyDeformationToGroom(const TArray<FRBFDeformedPositions>& Deforme
 	GroomAsset->UpdateHairGroupsInfo();
 
 	// Update/reimport the cards/meshes geometry which have been deformed prior to call this function
-	GroomAsset->BuildCardsGeometry();
-	GroomAsset->BuildMeshesGeometry();
+	GroomAsset->BuildCardsData();
+	GroomAsset->BuildMeshesData();
 }
 
 static void ExtractSkeletalVertexPosition(
@@ -544,7 +544,7 @@ void FGroomRBFDeformer::GetRBFDeformedGroomAsset(const UGroomAsset* InGroomAsset
 		{
 			const uint32 GroupIndex = Group.Info.GroupID;
 			const FHairStrandsDatas& OriginalGuides = Group.Guides;
-			const FHairGroupData& HairGroupData = InGroomAsset->HairGroupsData[GroupIndex];
+			const FHairGroupPlatformData& HairGroupData = InGroomAsset->HairGroupsPlatformData[GroupIndex];
 
 			FHairStrandsDatas StrandsData;
 			FHairStrandsDatas GuidesData;
diff --git a/Engine/Plugins/Runtime/HairStrands/Source/HairStrandsCore/Private/GroomTextureBuilder.cpp b/Engine/Plugins/Runtime/HairStrands/Source/HairStrandsCore/Private/GroomTextureBuilder.cpp
--- a/Engine/Plugins/Runtime/HairStrands/Source/HairStrandsCore/Private/GroomTextureBuilder.cpp
+++ b/Engine/Plugins/Runtime/HairStrands/Source/HairStrandsCore/Private/GroomTextureBuilder.cpp
@@ -271,7 +271,7 @@ static void InternalBuildFollicleTexture_CPU(const TArray<FFollicleInfo>& InInfo
 		}
 
 		uint32 GroupIndex = 0;
-		for (const FHairGroupData& HairGroupData : Info.GroomAsset->HairGroupsData)
+		for (const FHairGroupPlatformData& HairGroupData : Info.GroomAsset->HairGroupsPlatformData)
 		{
 			FHairStrandsDatas StrandsData;
 			FHairStrandsDatas GuidesData;
@@ -980,7 +980,7 @@ static bool TraceTextures(
 			continue;
 		}
 
-		FHairGroupData& GroupData = InInfo.GroomAsset->HairGroupsData[GroupIndex];
+		FHairGroupPlatformData& GroupData = InInfo.GroomAsset->HairGroupsPlatformData[GroupIndex];
 		FHairGroupsRendering& RenderingData = InInfo.GroomAsset->HairGroupsRendering[GroupIndex];
 
 		FRDGBufferRef VoxelOffsetAndCount = GraphBuilder.RegisterExternalBuffer(GroupData.Debug.Resource->VoxelOffsetAndCount);
diff --git a/Engine/Plugins/Runtime/HairStrands/Source/HairStrandsCore/Private/HairStrandsCore.cpp b/Engine/Plugins/Runtime/HairStrands/Source/HairStrandsCore/Private/HairStrandsCore.cpp
--- a/Engine/Plugins/Runtime/HairStrands/Source/HairStrandsCore/Private/HairStrandsCore.cpp
+++ b/Engine/Plugins/Runtime/HairStrands/Source/HairStrandsCore/Private/HairStrandsCore.cpp
@@ -231,7 +231,7 @@ UGroomBindingAsset* FHairStrandsCore::CreateGroomBindingAsset(EGroomBindingMeshT
 			Out->SourceGeometryCache = Cast<UGeometryCache>(Source);
 			Out->TargetGeometryCache = Cast<UGeometryCache>(Target);
 		}
-		Out->HairGroupBulkDatas.Reserve(GroomAsset->HairGroupsData.Num());
+		Out->HairGroupBulkDatas.Reserve(GroomAsset->HairGroupsPlatformData.Num());
 		Out->NumInterpolationPoints = NumInterpolationPoints;
 		Out->MatchingSection = MatchingSection;
 		Out->MarkPackageDirty();
diff --git a/Engine/Plugins/Runtime/HairStrands/Source/HairStrandsCore/Private/Niagara/NiagaraDataInterfaceHairStrands.cpp b/Engine/Plugins/Runtime/HairStrands/Source/HairStrandsCore/Private/Niagara/NiagaraDataInterfaceHairStrands.cpp
--- a/Engine/Plugins/Runtime/HairStrands/Source/HairStrandsCore/Private/Niagara/NiagaraDataInterfaceHairStrands.cpp
+++ b/Engine/Plugins/Runtime/HairStrands/Source/HairStrandsCore/Private/Niagara/NiagaraDataInterfaceHairStrands.cpp
@@ -737,7 +737,7 @@ void UNiagaraDataInterfaceHairStrands::ExtractDatasAndResources(
 		OutLocalToWorld = SystemInstance ? SystemInstance->GetWorldTransform() : FTransform::Identity;
 		if (OutGroupIndex < DefaultSource->GetNumHairGroups())
 		{
-			OutStrandsRestResource = DefaultSource->HairGroupsData[OutGroupIndex].Guides.RestResource;
+			OutStrandsRestResource = DefaultSource->HairGroupsPlatformData[OutGroupIndex].Guides.RestResource;
 			OutGroomAsset = DefaultSource;
 		}
 	}
diff --git a/Engine/Plugins/Runtime/HairStrands/Source/HairStrandsCore/Public/GroomAsset.h b/Engine/Plugins/Runtime/HairStrands/Source/HairStrandsCore/Public/GroomAsset.h
--- a/Engine/Plugins/Runtime/HairStrands/Source/HairStrandsCore/Public/GroomAsset.h
+++ b/Engine/Plugins/Runtime/HairStrands/Source/HairStrandsCore/Public/GroomAsset.h
@@ -43,7 +43,7 @@ struct HAIRSTRANDSCORE_API FHairGroupLODInfo
 {
 	GENERATED_BODY()
 
-	UPROPERTY(VisibleAnywhere, Category = "Info", meta = (DisplayName = "Curve Count"))
+	UPROPERTY(VisibleAnywhere, Category = "Info", meta = (DisplayName = "Point Count"))
 	int32 NumPoints = 0;
 
 	UPROPERTY(VisibleAnywhere, Category = "Info", meta = (DisplayName = "Curve Count"))
@@ -100,7 +100,7 @@ struct HAIRSTRANDSCORE_API FHairGroupsMaterial
 
 
 /** Describe all data & resource for a groom asset's hair group */
-struct HAIRSTRANDSCORE_API FHairGroupData
+struct HAIRSTRANDSCORE_API FHairGroupPlatformData
 {
 	////////////////////////////////////////////////////////////////////////////
 	// Helper
@@ -475,7 +475,7 @@ public:
 	TArray<FHairGroupsMaterial> HairGroupsMaterials;
 
 	/** Store strands/cards/meshes data */
-	TArray<FHairGroupData> HairGroupsData;
+	TArray<FHairGroupPlatformData> HairGroupsPlatformData;
 
 	/** Enable radial basis function interpolation to be used instead of the local skin rigid transform */
 	UPROPERTY(EditAnywhere, EditFixedSize, BlueprintReadWrite, Category = "HairInterpolation", meta = (ToolTip = "Enable radial basis function interpolation to be used instead of the local skin rigid transform (WIP)", DisplayName = "RBF Interpolation"))
@@ -633,8 +633,8 @@ public:
 	bool IsMaterialUsed(int32 MaterialIndex) const;
 	TArray<FName> GetMaterialSlotNames() const;
 
-	bool BuildCardsGeometry();
-	bool BuildMeshesGeometry();
+	bool BuildCardsData();
+	bool BuildMeshesData();
 
 	enum EClassDataStripFlag : uint8
 	{
@@ -663,8 +663,8 @@ private:
 #if WITH_EDITORONLY_DATA
 	bool HasImportedStrandsData() const;
 
-	bool BuildCardsGeometry(uint32 GroupIndex);
-	bool BuildMeshesGeometry(uint32 GroupIndex);
+	bool BuildCardsData(uint32 GroupIndex);
+	bool BuildMeshesData(uint32 GroupIndex);
 
 	bool HasValidCardsData(uint32 GroupIndex) const;
 	bool HasValidMeshesData(uint32 GroupIndex) const;
@@ -690,8 +690,8 @@ public:
 	bool CacheDerivedDatas();
 	bool CacheDerivedData(uint32 GroupIndex);
 	bool CacheStrandsData(uint32 GroupIndex, FString& OutDerivedDataKey);
-	bool CacheCardsGeometry(uint32 GroupIndex, const FString& StrandsKey);
-	bool CacheMeshesGeometry(uint32 GroupIndex);
+	bool CacheCardsData(uint32 GroupIndex, const FString& StrandsKey);
+	bool CacheMeshesData(uint32 GroupIndex);
 
 	FString GetDerivedDataKey();
 	FString GetDerivedDataKeyForCards(uint32 GroupIt, const FString& StrandsKey);
diff --git a/Engine/Plugins/Runtime/HairStrands/Source/HairStrandsCore/Public/GroomBindingAsset.h b/Engine/Plugins/Runtime/HairStrands/Source/HairStrandsCore/Public/GroomBindingAsset.h
--- a/Engine/Plugins/Runtime/HairStrands/Source/HairStrandsCore/Public/GroomBindingAsset.h
+++ b/Engine/Plugins/Runtime/HairStrands/Source/HairStrandsCore/Public/GroomBindingAsset.h
@@ -19,7 +19,7 @@ class UGeometryCache;
 class UMaterialInterface;
 class UNiagaraSystem;
 class UGroomAsset;
-struct FHairGroupData;
+struct FHairGroupPlatformData;
 
 USTRUCT(BlueprintType)
 struct HAIRSTRANDSCORE_API FGoomBindingGroupInfo
diff --git a/Engine/Plugins/Runtime/HairStrands/Source/HairStrandsEditor/Private/GroomActions.cpp b/Engine/Plugins/Runtime/HairStrands/Source/HairStrandsEditor/Private/GroomActions.cpp
--- a/Engine/Plugins/Runtime/HairStrands/Source/HairStrandsEditor/Private/GroomActions.cpp
+++ b/Engine/Plugins/Runtime/HairStrands/Source/HairStrandsEditor/Private/GroomActions.cpp
@@ -117,8 +117,8 @@ void ExecuteRebuild(const FToolMenuContext& InContext)
 						FGroomHairGroupPreview& OutGroup = GroupsPreview->Groups.AddDefaulted_GetRef();
 						OutGroup.GroupID    	= GroupIndex;
 						OutGroup.GroupName		= GroomAsset->HairGroupsInfo[GroupIndex].GroupName;
-						OutGroup.CurveCount 	= GroomAsset->HairGroupsData[GroupIndex].Strands.BulkData.GetNumCurves();
-						OutGroup.GuideCount 	= GroomAsset->HairGroupsData[GroupIndex].Guides.BulkData.GetNumCurves();
+						OutGroup.CurveCount 	= GroomAsset->HairGroupsPlatformData[GroupIndex].Strands.BulkData.GetNumCurves();
+						OutGroup.GuideCount 	= GroomAsset->HairGroupsPlatformData[GroupIndex].Guides.BulkData.GetNumCurves();
 						OutGroup.Attributes 	= HairDescriptionGroups.HairGroups[GroupIndex].GetHairAttributes();
 						OutGroup.AttributeFlags = HairDescriptionGroups.HairGroups[GroupIndex].GetHairAttributeFlags();
 						OutGroup.InterpolationSettings = GroomAsset->HairGroupsInterpolation[GroupIndex];
diff --git a/Engine/Plugins/Runtime/HairStrands/Source/HairStrandsEditor/Private/GroomAssetDetails.cpp b/Engine/Plugins/Runtime/HairStrands/Source/HairStrandsEditor/Private/GroomAssetDetails.cpp
--- a/Engine/Plugins/Runtime/HairStrands/Source/HairStrandsEditor/Private/GroomAssetDetails.cpp
+++ b/Engine/Plugins/Runtime/HairStrands/Source/HairStrandsEditor/Private/GroomAssetDetails.cpp
@@ -1300,7 +1300,7 @@ void FGroomRenderingDetails::OnGenerateElementForLODs(TSharedRef<IPropertyHandle
 	];
 
 	// LOD Stats
-	const FHairStrandsClusterCullingBulkData& ClusterCullingBulkData = GroomAsset->HairGroupsData[GroupIndex].Strands.ClusterCullingBulkData;
+	const FHairStrandsClusterCullingBulkData& ClusterCullingBulkData = GroomAsset->HairGroupsPlatformData[GroupIndex].Strands.ClusterCullingBulkData;
 	if (ClusterCullingBulkData.IsValid() && LODIndex < ClusterCullingBulkData.LODInfos.Num())
 	{
 		const FHairLODInfo& LODInfo = ClusterCullingBulkData.LODInfos[LODIndex];
@@ -1480,13 +1480,13 @@ void FGroomRenderingDetails::OnGenerateElementForHairGroup(TSharedRef<IPropertyH
 		.ValueContent()
 		.HAlign(HAlign_Fill)
 		[
-			MakeHairStrandsInfoGrid(DetailFontInfo, GroomAsset->HairGroupsInfo[GroupIndex], GroomAsset->HairGroupsData[GroupIndex].Strands.BulkData.MaxRadius)
+			MakeHairStrandsInfoGrid(DetailFontInfo, GroomAsset->HairGroupsInfo[GroupIndex], GroomAsset->HairGroupsPlatformData[GroupIndex].Strands.BulkData.MaxRadius)
 		];
 		ChildrenBuilder.AddCustomRow(LOCTEXT("HairStrandsAttributeInfo_Array", "HairStrandsAttributeInfo"))
 		.ValueContent()
 		.HAlign(HAlign_Fill)
 		[
-			MakeHairStrandsAttributeInfoGrid(DetailFontInfo, GroomAsset->HairGroupsData[GroupIndex].Strands.BulkData.ImportedAttributes, GroomAsset->HairGroupsData[GroupIndex].Strands.BulkData.ImportedAttributeFlags)
+			MakeHairStrandsAttributeInfoGrid(DetailFontInfo, GroomAsset->HairGroupsPlatformData[GroupIndex].Strands.BulkData.ImportedAttributes, GroomAsset->HairGroupsPlatformData[GroupIndex].Strands.BulkData.ImportedAttributeFlags)
 		];
 	}
 
diff --git a/Engine/Plugins/Runtime/HairStrands/Source/HairStrandsEditor/Private/GroomCacheImporter.cpp b/Engine/Plugins/Runtime/HairStrands/Source/HairStrandsEditor/Private/GroomCacheImporter.cpp
--- a/Engine/Plugins/Runtime/HairStrands/Source/HairStrandsEditor/Private/GroomCacheImporter.cpp
+++ b/Engine/Plugins/Runtime/HairStrands/Source/HairStrandsEditor/Private/GroomCacheImporter.cpp
@@ -124,7 +124,7 @@ TArray<UGroomCache*> FGroomCacheImporter::ImportGroomCache(const FString& Source
 		FScopedSlowTask SlowTask(NumFrames, LOCTEXT("ImportGroomCache", "Importing GroomCache frames"));
 		SlowTask.MakeDialog();
 
-		const TArray<FHairGroupData>& GroomHairGroupsData = GroomAssetForCache->HairGroupsData;
+		const TArray<FHairGroupPlatformData>& GroupPlatformData = GroomAssetForCache->HairGroupsPlatformData;
 
 		// Each frame is translated into a HairDescription and processed into HairGroupData
 		for (int32 FrameIndex = AnimInfo.StartFrame; FrameIndex < AnimInfo.EndFrame + 1; ++FrameIndex)
@@ -149,11 +149,11 @@ TArray<UGroomCache*> FGroomCacheImporter::ImportGroomCache(const FString& Source
 
 				const uint32 GroupCount = HairDescriptionGroups.HairGroups.Num();
 
-				if (GroupCount != GroomHairGroupsData.Num())
+				if (GroupCount != GroupPlatformData.Num())
 				{
 					bSuccess = false;
 					UE_LOG(LogGroomCacheImporter, Warning, TEXT("GroomCache does not have the same number of groups as the static groom (%d instead of %d). Aborting GroomCache import."),
-						GroupCount, GroomHairGroupsData.Num());
+						GroupCount, GroupPlatformData.Num());
 					break;
 				}
 
@@ -172,14 +172,14 @@ TArray<UGroomCache*> FGroomCacheImporter::ImportGroomCache(const FString& Source
 				{
 					if (bImportStrandsCache)
 					{
-						if (HairGroupsData[GroupIndex].Strands.GetNumCurves() != GroomHairGroupsData[GroupIndex].Strands.BulkData.GetNumCurves() ||
-							HairGroupsData[GroupIndex].Strands.GetNumPoints() != GroomHairGroupsData[GroupIndex].Strands.BulkData.GetNumPoints())
+						if (HairGroupsData[GroupIndex].Strands.GetNumCurves() != GroupPlatformData[GroupIndex].Strands.BulkData.GetNumCurves() ||
+							HairGroupsData[GroupIndex].Strands.GetNumPoints() != GroupPlatformData[GroupIndex].Strands.BulkData.GetNumPoints())
 						{
 							bSuccess = false;
 							UE_LOG(LogGroomCacheImporter, Warning, TEXT("GroomCache frame %d does not have the same number of curves (%u) \
 								or vertices (%u) for the strands as the static groom (%u and %u respectively). Aborting GroomCache import."),
 								FrameIndex, HairGroupsData[GroupIndex].Strands.GetNumCurves(), HairGroupsData[GroupIndex].Strands.GetNumPoints(),
-								GroomHairGroupsData[GroupIndex].Strands.BulkData.GetNumCurves(), GroomHairGroupsData[GroupIndex].Strands.BulkData.GetNumPoints());
+								GroupPlatformData[GroupIndex].Strands.BulkData.GetNumCurves(), GroupPlatformData[GroupIndex].Strands.BulkData.GetNumPoints());
 							break;
 						}
 					}
@@ -191,14 +191,14 @@ TArray<UGroomCache*> FGroomCacheImporter::ImportGroomCache(const FString& Source
 							bGuidesOnly = true;
 						}
 
-						if (HairGroupsData[GroupIndex].Guides.GetNumCurves() != GroomHairGroupsData[GroupIndex].Guides.BulkData.GetNumCurves() ||
-							HairGroupsData[GroupIndex].Guides.GetNumPoints() != GroomHairGroupsData[GroupIndex].Guides.BulkData.GetNumPoints())
+						if (HairGroupsData[GroupIndex].Guides.GetNumCurves() != GroupPlatformData[GroupIndex].Guides.BulkData.GetNumCurves() ||
+							HairGroupsData[GroupIndex].Guides.GetNumPoints() != GroupPlatformData[GroupIndex].Guides.BulkData.GetNumPoints())
 						{
 							bSuccess = false;
 							UE_LOG(LogGroomCacheImporter, Warning, TEXT("GroomCache frame %d does not have the same number of curves (%u) \
 								or vertices (%u) for the guides as the static groom (%u and %u respectively). Aborting GroomCache import."),
 								FrameIndex, HairGroupsData[GroupIndex].Guides.GetNumCurves(), HairGroupsData[GroupIndex].Guides.GetNumPoints(),
-								GroomHairGroupsData[GroupIndex].Guides.BulkData.GetNumCurves(), GroomHairGroupsData[GroupIndex].Guides.BulkData.GetNumPoints());
+								GroupPlatformData[GroupIndex].Guides.BulkData.GetNumCurves(), GroupPlatformData[GroupIndex].Guides.BulkData.GetNumPoints());
 							break;
 						}
 					}
diff --git a/Engine/Plugins/Runtime/HairStrands/Source/HairStrandsEditor/Private/GroomImportOptionsWindow.cpp b/Engine/Plugins/Runtime/HairStrands/Source/HairStrandsEditor/Private/GroomImportOptionsWindow.cpp
--- a/Engine/Plugins/Runtime/HairStrands/Source/HairStrandsEditor/Private/GroomImportOptionsWindow.cpp
+++ b/Engine/Plugins/Runtime/HairStrands/Source/HairStrandsEditor/Private/GroomImportOptionsWindow.cpp
@@ -114,7 +114,7 @@ void SGroomImportOptionsWindow::UpdateStatus(UGroomHairGroupsPreview* Descriptio
 			return;
 		}
 
-		const TArray<FHairGroupData>& GroomHairGroupsData = GroomAssetForCache->HairGroupsData;
+		const TArray<FHairGroupPlatformData>& GroomHairGroupsData = GroomAssetForCache->HairGroupsPlatformData;
 		if (GroomHairGroupsData.Num() != Description->Groups.Num())
 		{
 			CurrentStatus = bGuidesOnly ? EHairDescriptionStatus::GuidesOnlyIncompatible : EHairDescriptionStatus::GroomCacheIncompatible;
