* Split strands/guide/group serialization for clarity.
* This fixes double serialization of cards/meshes.

diff --git a/Engine/Plugins/Runtime/HairStrands/Source/HairStrandsCore/Private/GroomAsset.cpp b/Engine/Plugins/Runtime/HairStrands/Source/HairStrandsCore/Private/GroomAsset.cpp
--- a/Engine/Plugins/Runtime/HairStrands/Source/HairStrandsCore/Private/GroomAsset.cpp
+++ b/Engine/Plugins/Runtime/HairStrands/Source/HairStrandsCore/Private/GroomAsset.cpp
@@ -652,7 +652,7 @@ void UGroomAsset::ApplyStripFlags(uint8 StripFlags, const ITargetPlatform* CookT
 	for (int32 GroupIndex = 0; GroupIndex < HairGroupsPlatformData.Num(); ++GroupIndex)
 	{
 		FHairGroupPlatformData& GroupData = HairGroupsPlatformData[GroupIndex];
-		GroupData.bIsCookedOut = bIsStrandsStrippedForCook;
+		GroupData.Strands.bIsCookedOut = bIsStrandsStrippedForCook;
 
 		// Determine the max LOD for strands since they are cooked out if the platform doesn't support them
 		// They are cooked out as a whole
@@ -700,15 +700,18 @@ bool UGroomAsset::HasImportedStrandsData() const
 #endif
 
 // Serialization for *array* of hair elements
-static void InternalSerialize(FArchive& Ar, UObject* Owner, TArray<FHairGroupPlatformData::FCards::FLOD>& CardLODData);
-static void InternalSerialize(FArchive& Ar, UObject* Owner, TArray<FHairGroupPlatformData::FMeshes::FLOD>& MeshLODData);
-static void InternalSerialize(FArchive& Ar, UObject* Owner, TArray<FHairGroupPlatformData>& GroupData);
+static void InternalSerializeCards(FArchive& Ar, UObject* Owner, TArray<FHairGroupPlatformData::FCards::FLOD>& CardLODData);
+static void InternalSerializeMeshes(FArchive& Ar, UObject* Owner, TArray<FHairGroupPlatformData::FMeshes::FLOD>& MeshLODData);
+static void InternalSerializePlatformDatas(FArchive& Ar, UObject* Owner, TArray<FHairGroupPlatformData>& GroupData);
 
 // Serialization for hair elements
-static void InternalSerialize(FArchive& Ar, UObject* Owner, FHairGroupPlatformData::FCards::FLOD& CardLODData);
-static void InternalSerialize(FArchive& Ar, UObject* Owner, FHairGroupPlatformData::FMeshes::FLOD& MeshLODData);
-static void InternalSerialize(FArchive& Ar, UObject* Owner, FHairGroupPlatformData& GroupData);
+static void InternalSerializeCard(FArchive& Ar, UObject* Owner, FHairGroupPlatformData::FCards::FLOD& CardLODData);
+static void InternalSerializeMesh(FArchive& Ar, UObject* Owner, FHairGroupPlatformData::FMeshes::FLOD& MeshLODData);
+static void InternalSerializeGuide(FArchive& Ar, UObject* Owner, FHairGroupPlatformData::FGuides& GuideData);
+static void InternalSerializeStrand(FArchive& Ar, UObject* Owner, FHairGroupPlatformData::FStrands& StrandData);
+static void InternalSerializePlatformData(FArchive& Ar, UObject* Owner, FHairGroupPlatformData& GroupData);
 
+// Serialize is only called for cooking??
 void UGroomAsset::Serialize(FArchive& Ar)
 {
 	uint8 ClassDataStripFlags = GenerateClassStripFlags(Ar);
@@ -746,7 +749,7 @@ void UGroomAsset::Serialize(FArchive& Ar)
 		{
 			// When cooking data or serializing old format to new format,
 			// serialize the computed groom data
-			InternalSerialize(Ar, this, HairGroupsPlatformData);
+			InternalSerializePlatformDatas(Ar, this, HairGroupsPlatformData);
 		}
 #if WITH_EDITORONLY_DATA
 		else
@@ -1582,7 +1585,7 @@ int32 UGroomAsset::GetNumHairGroups() const
 
 ////////////////////////////////////////////////////////////////////////////////////////////////////////////
 // Array serialize
-void InternalSerialize(FArchive& Ar, UObject* Owner, TArray<FHairGroupPlatformData::FMeshes::FLOD>& LODs)
+void InternalSerializeMeshes(FArchive& Ar, UObject* Owner, TArray<FHairGroupPlatformData::FMeshes::FLOD>& LODs)
 {
 	uint32 Count = LODs.Num();
 	Ar << Count;
@@ -1592,11 +1595,11 @@ void InternalSerialize(FArchive& Ar, UObject* Owner, TArray<FHairGroupPlatformDa
 	}
 	for (uint32 MeshIt = 0; MeshIt < Count; ++MeshIt)
 	{
-		InternalSerialize(Ar, Owner, LODs[MeshIt]);
+		InternalSerializeMesh(Ar, Owner, LODs[MeshIt]);
 	}
 }
 
-void InternalSerialize(FArchive& Ar, UObject* Owner, TArray<FHairGroupPlatformData::FCards::FLOD>& LODs)
+void InternalSerializeCards(FArchive& Ar, UObject* Owner, TArray<FHairGroupPlatformData::FCards::FLOD>& LODs)
 {
 	uint32 Count = LODs.Num();
 	Ar << Count;
@@ -1606,11 +1609,11 @@ void InternalSerialize(FArchive& Ar, UObject* Owner, TArray<FHairGroupPlatformDa
 	}
 	for (uint32 MeshIt = 0; MeshIt < Count; ++MeshIt)
 	{
-		InternalSerialize(Ar, Owner, LODs[MeshIt]);
+		InternalSerializeCard(Ar, Owner, LODs[MeshIt]);
 	}
 }
 
-void InternalSerialize(FArchive& Ar, UObject* Owner, TArray<FHairGroupPlatformData>& GroupDatas)
+void InternalSerializePlatformDatas(FArchive& Ar, UObject* Owner, TArray<FHairGroupPlatformData>& GroupDatas)
 {
 	uint32 GroupCount = GroupDatas.Num();
 	Ar << GroupCount;
@@ -1620,13 +1623,13 @@ void InternalSerialize(FArchive& Ar, UObject* Owner, TArray<FHairGroupPlatformDa
 	}
 	for (uint32 GroupIt = 0; GroupIt < GroupCount; ++GroupIt)
 	{
-		InternalSerialize(Ar, Owner, GroupDatas[GroupIt]);
+		InternalSerializePlatformData(Ar, Owner, GroupDatas[GroupIt]);
 	}
 }
 
 ////////////////////////////////////////////////////////////////////////////////////////////////////////////
 // Element serialize
-static void InternalSerialize(FArchive& Ar, UObject* Owner, FHairGroupPlatformData::FCards::FLOD& CardLODData)
+static void InternalSerializeCard(FArchive& Ar, UObject* Owner, FHairGroupPlatformData::FCards::FLOD& CardLODData)
 {
 	if (!Ar.IsCooking() || !CardLODData.bIsCookedOut)
 	{
@@ -1651,7 +1654,7 @@ static void InternalSerialize(FArchive& Ar, UObject* Owner, FHairGroupPlatformDa
 	}
 }
 
-static void InternalSerialize(FArchive& Ar, UObject* Owner, FHairGroupPlatformData::FMeshes::FLOD& MeshLODData)
+static void InternalSerializeMesh(FArchive& Ar, UObject* Owner, FHairGroupPlatformData::FMeshes::FLOD& MeshLODData)
 {
 	if (!Ar.IsCooking() || !MeshLODData.bIsCookedOut)
 	{
@@ -1665,49 +1668,55 @@ static void InternalSerialize(FArchive& Ar, UObject* Owner, FHairGroupPlatformDa
 	}
 }
 
-static void InternalSerialize(FArchive& Ar, UObject* Owner, FHairGroupPlatformData& GroupData)
+static void InternalSerializeGuide(FArchive& Ar, UObject* Owner, FHairGroupPlatformData::FGuides& GuideData)
+{
+	GuideData.BulkData.Serialize(Ar, Owner);
+}
+
+static void InternalSerializeStrand(FArchive& Ar, UObject* Owner, FHairGroupPlatformData::FStrands& StrandData)
 {
 	Ar.UsingCustomVersion(FAnimObjectVersion::GUID);
 
-	FHairGroupPlatformData NoStrandsData;
-	if (!Ar.IsCooking() || !GroupData.bIsCookedOut)
+	if (!Ar.IsCooking() || !StrandData.bIsCookedOut)
 	{
-		GroupData.Strands.BulkData.Serialize(Ar, Owner);
-		GroupData.Guides.BulkData.Serialize(Ar, Owner);
-		GroupData.Strands.InterpolationBulkData.Serialize(Ar, Owner);
-
+		StrandData.BulkData.Serialize(Ar, Owner);
+		StrandData.InterpolationBulkData.Serialize(Ar, Owner);
+		if (Ar.CustomVer(FAnimObjectVersion::GUID) >= FAnimObjectVersion::SerializeHairClusterCullingData)
+		{
+			StrandData.ClusterCullingBulkData.Serialize(Ar, Owner);
+		}
 	}
 	else
 	{
 		// Fall back to no data, but still serialize guide data as they an be used for cards simulation
 		// Theoritically, we should have something to detect if we are going to use or not guide (for 
 		// simulation or RBF deformation) on the target platform
-		NoStrandsData.Strands.BulkData.Serialize(Ar, Owner);
-		GroupData.Guides.BulkData.Serialize(Ar, Owner);
-		NoStrandsData.Strands.InterpolationBulkData.Serialize(Ar, Owner);
+		FHairGroupPlatformData::FStrands NoStrandsData;
+		NoStrandsData.BulkData.Serialize(Ar, Owner);
+		NoStrandsData.InterpolationBulkData.Serialize(Ar, Owner);
+		if (Ar.CustomVer(FAnimObjectVersion::GUID) >= FAnimObjectVersion::SerializeHairClusterCullingData)
+		{
+			NoStrandsData.ClusterCullingBulkData.Serialize(Ar, Owner);
+		}
 	}
+}
+
+static void InternalSerializePlatformData(FArchive& Ar, UObject* Owner, FHairGroupPlatformData& GroupData)
+{
+	Ar.UsingCustomVersion(FAnimObjectVersion::GUID);
+
+	InternalSerializeGuide(Ar, Owner, GroupData.Guides);
+	InternalSerializeStrand(Ar, Owner, GroupData.Strands);
 
-	if (Ar.CustomVer(FAnimObjectVersion::GUID) >= FAnimObjectVersion::SerializeHairClusterCullingData)
+	if (Ar.CustomVer(FAnimObjectVersion::GUID) >= FAnimObjectVersion::SerializeGroomCardsAndMeshes)
 	{
-		if (!Ar.IsCooking() || !GroupData.bIsCookedOut)
-		{
-			GroupData.Strands.ClusterCullingBulkData.Serialize(Ar, Owner);
-		}
-		else
-		{
-			NoStrandsData.Strands.ClusterCullingBulkData.Serialize(Ar, Owner);
-		}
+		bool bIsCooked = Ar.IsCooking();
+		Ar << bIsCooked;
 
-		if (Ar.CustomVer(FAnimObjectVersion::GUID) >= FAnimObjectVersion::SerializeGroomCardsAndMeshes)
+		if (bIsCooked)
 		{
-			bool bIsCooked = Ar.IsCooking();
-			Ar << bIsCooked;
-
-			if (bIsCooked)
-			{
-				InternalSerialize(Ar, Owner, GroupData.Cards.LODs);
-				InternalSerialize(Ar, Owner, GroupData.Meshes.LODs);
-			}
+			InternalSerializeCards(Ar, Owner, GroupData.Cards.LODs);
+			InternalSerializeMeshes(Ar, Owner, GroupData.Meshes.LODs);
 		}
 	}
 }
@@ -2345,7 +2354,10 @@ bool UGroomAsset::CacheStrandsData(uint32 GroupIndex, FString& OutDerivedDataKey
 		HairGroupsPlatformData[GroupIndex] = FHairGroupPlatformData();
 
 		FMemoryReaderView Ar(Data, /*bIsPersistent*/ true);
-		InternalSerialize(Ar, this, HairGroupsPlatformData[GroupIndex]);
+
+		InternalSerializeGuide(Ar, this, HairGroupsPlatformData[GroupIndex].Guides);
+		InternalSerializeStrand(Ar, this, HairGroupsPlatformData[GroupIndex].Strands);
+
 	}
 	else
 	{
@@ -2382,7 +2394,10 @@ bool UGroomAsset::CacheStrandsData(uint32 GroupIndex, FString& OutDerivedDataKey
 		{
 			TArray<uint8> WriteData;
 			FMemoryWriter Ar(WriteData, /*bIsPersistent*/ true);
-			InternalSerialize(Ar, this, HairGroupsPlatformData[GroupIndex]);
+				
+			InternalSerializeGuide(Ar, this, HairGroupsPlatformData[GroupIndex].Guides);
+			InternalSerializeStrand(Ar, this, HairGroupsPlatformData[GroupIndex].Strands);
+
 			FRequestOwner AsyncOwner(EPriority::Normal);
 			GetCache().PutValue({ {Name, Key, FValue::Compress(MakeSharedBufferFromArray(MoveTemp(WriteData)))} }, AsyncOwner);
 			AsyncOwner.KeepAlive();
@@ -2506,7 +2521,7 @@ bool UGroomAsset::CacheCardsData(uint32 GroupIndex, const FString& StrandsKey)
 		UE_CLOG(IsHairStrandsDDCLogEnable(), LogHairStrands, Log, TEXT("[Groom/DDC] Cards - Found (Groom:%s Group:%d)."), *GetName(), GroupIndex);
 		
 		FMemoryReaderView Ar(Data, /*bIsPersistent*/ true);
-		InternalSerialize(Ar, this, HairGroupData.Cards.LODs);
+		InternalSerializeCards(Ar, this, HairGroupData.Cards.LODs);
 	}
 	else
 	{
@@ -2520,7 +2535,7 @@ bool UGroomAsset::CacheCardsData(uint32 GroupIndex, const FString& StrandsKey)
 
 		TArray<uint8> WriteData;
 		FMemoryWriter Ar(WriteData, /*bIsPersistent*/ true);
-		InternalSerialize(Ar, this, HairGroupData.Cards.LODs);
+		InternalSerializeCards(Ar, this, HairGroupData.Cards.LODs);
 		FRequestOwner AsyncOwner(EPriority::Normal);
 		GetCache().PutValue({ {Name, Key, FValue::Compress(MakeSharedBufferFromArray(MoveTemp(WriteData)))} }, AsyncOwner);
 		AsyncOwner.KeepAlive();
@@ -2617,7 +2632,7 @@ bool UGroomAsset::CacheMeshesData(uint32 GroupIndex)
 		UE_CLOG(IsHairStrandsDDCLogEnable(), LogHairStrands, Log, TEXT("[Groom/DDC] Meshes - Found (Groom:%s Group:%d)."), *GetName(), GroupIndex);
 
 		FMemoryReaderView Ar(Data, /*bIsPersistent*/ true);
-		InternalSerialize(Ar, this, HairGroupData.Meshes.LODs);
+		InternalSerializeMeshes(Ar, this, HairGroupData.Meshes.LODs);
 	}
 	else
 	{
@@ -2631,7 +2646,7 @@ bool UGroomAsset::CacheMeshesData(uint32 GroupIndex)
 
 		TArray<uint8> WriteData;
 		FMemoryWriter Ar(WriteData, /*bIsPersistent*/ true);
-		InternalSerialize(Ar, this, HairGroupData.Meshes.LODs);
+		InternalSerializeMeshes(Ar, this, HairGroupData.Meshes.LODs);
 		FRequestOwner AsyncOwner(EPriority::Normal);
 		GetCache().PutValue({ {Name, Key, FValue::Compress(MakeSharedBufferFromArray(MoveTemp(WriteData)))} }, AsyncOwner);
 		AsyncOwner.KeepAlive();
diff --git a/Engine/Plugins/Runtime/HairStrands/Source/HairStrandsCore/Private/GroomBuilder.cpp b/Engine/Plugins/Runtime/HairStrands/Source/HairStrandsCore/Private/GroomBuilder.cpp
--- a/Engine/Plugins/Runtime/HairStrands/Source/HairStrandsCore/Private/GroomBuilder.cpp
+++ b/Engine/Plugins/Runtime/HairStrands/Source/HairStrandsCore/Private/GroomBuilder.cpp
@@ -40,7 +40,7 @@ static FAutoConsoleVariableRef CVarHairGroupIndexBuilder_MaxVoxelResolution(TEXT
 
 FString FGroomBuilder::GetVersion()
 {
-	return TEXT("v8m");
+	return TEXT("v8n");
 }
 
 namespace FHairStrandsDecimation
diff --git a/Engine/Plugins/Runtime/HairStrands/Source/HairStrandsCore/Public/GroomAsset.h b/Engine/Plugins/Runtime/HairStrands/Source/HairStrandsCore/Public/GroomAsset.h
--- a/Engine/Plugins/Runtime/HairStrands/Source/HairStrandsCore/Public/GroomAsset.h
+++ b/Engine/Plugins/Runtime/HairStrands/Source/HairStrandsCore/Public/GroomAsset.h
@@ -146,7 +146,7 @@ struct HAIRSTRANDSCORE_API FHairGroupPlatformData
 	////////////////////////////////////////////////////////////////////////////
 	// Data
 
-	struct FSim : FBase
+	struct FGuides : FBase
 	{
 		/* Return the memory size for GPU resources */
 		uint32 GetResourcesSize() const
@@ -202,6 +202,7 @@ struct HAIRSTRANDSCORE_API FHairGroupPlatformData
 		FHairStrandsRaytracingResource*		RaytracingResource = nullptr;
 		#endif
 
+		bool bIsCookedOut = false;
 	} Strands;
 
 	struct FCards
@@ -382,8 +383,6 @@ struct HAIRSTRANDSCORE_API FHairGroupPlatformData
 		FHairStrandsDebugDatas Data;
 		FHairStrandsDebugDatas::FResources* Resource = nullptr;
 	} Debug;
-
-	bool bIsCookedOut = false;
 };
 
 struct FHairDescriptionGroup
