Split hair bulk data into header/data.

diff --git a/Engine/Plugins/Runtime/HairStrands/Source/HairStrandsCore/Private/GroomAsset.cpp b/Engine/Plugins/Runtime/HairStrands/Source/HairStrandsCore/Private/GroomAsset.cpp
--- a/Engine/Plugins/Runtime/HairStrands/Source/HairStrandsCore/Private/GroomAsset.cpp
+++ b/Engine/Plugins/Runtime/HairStrands/Source/HairStrandsCore/Private/GroomAsset.cpp
@@ -125,12 +125,12 @@ uint32 FHairGroupPlatformData::FStrands::GetDataSize() const
 	Total += FBaseWithInterpolation::GetDataSize();
 	uint32 CTotal = 0;
 	uint32 BTotal = 0;
-	CTotal += ClusterCullingBulkData.LODVisibility.GetAllocatedSize();
-	CTotal += ClusterCullingBulkData.CPULODScreenSize.GetAllocatedSize();
-	CTotal += ClusterCullingBulkData.ClusterLODInfos.IsBulkDataLoaded()		? ClusterCullingBulkData.ClusterLODInfos.GetBulkDataSize() : 0;
-	BTotal += ClusterCullingBulkData.VertexToClusterIds.IsBulkDataLoaded()	? ClusterCullingBulkData.VertexToClusterIds.GetBulkDataSize() : 0;
-	BTotal += ClusterCullingBulkData.ClusterVertexIds.IsBulkDataLoaded()	? ClusterCullingBulkData.ClusterVertexIds.GetBulkDataSize() : 0;
-	CTotal += ClusterCullingBulkData.PackedClusterInfos.IsBulkDataLoaded()	? ClusterCullingBulkData.PackedClusterInfos.GetBulkDataSize() : 0;
+	CTotal += ClusterCullingBulkData.Header.LODVisibility.GetAllocatedSize();
+	CTotal += ClusterCullingBulkData.Header.CPULODScreenSize.GetAllocatedSize();
+	CTotal += ClusterCullingBulkData.Data.ClusterLODInfos.IsBulkDataLoaded()	? ClusterCullingBulkData.Data.ClusterLODInfos.GetBulkDataSize() : 0;
+	BTotal += ClusterCullingBulkData.Data.VertexToClusterIds.IsBulkDataLoaded()	? ClusterCullingBulkData.Data.VertexToClusterIds.GetBulkDataSize() : 0;
+	BTotal += ClusterCullingBulkData.Data.ClusterVertexIds.IsBulkDataLoaded()	? ClusterCullingBulkData.Data.ClusterVertexIds.GetBulkDataSize() : 0;
+	CTotal += ClusterCullingBulkData.Data.PackedClusterInfos.IsBulkDataLoaded()	? ClusterCullingBulkData.Data.PackedClusterInfos.GetBulkDataSize() : 0;
 
 	return BTotal + Total + CTotal;
 }
@@ -2378,12 +2378,10 @@ bool UGroomAsset::CacheStrandsData(uint32 GroupIndex, FString& OutDerivedDataKey
 			return false;
 		}
 
-		const bool bNeedInterpolationData = NeedsInterpolationData(GroupIndex);
-
 		// Build groom data with the new build settings
 		bSuccess = BuildHairGroup(
 			GroupIndex,
-			bNeedInterpolationData,
+			NeedsInterpolationData(GroupIndex),
 			LocalHairDescriptionGroups,
 			HairGroupsInterpolation,
 			HairGroupsLOD,
@@ -3023,7 +3021,7 @@ bool UGroomAsset::BuildMeshesData(uint32 GroupIndex)
 			{
 				// Build a default box
 				FHairMeshesBuilder::BuildGeometry(
-					GroupData.Strands.BulkData.BoundingBox,
+					GroupData.Strands.BulkData.Header.BoundingBox,
 					LOD.BulkData);
 			}
 
@@ -3182,16 +3180,16 @@ void UGroomAsset::InitStrandsResources()
 				bool bLastVisibility = true;
 				// LOD visibility is not serialized into FHairStrandsClusterCullingBulkData as it does not affect the actual generated data. For consistency we 
 				// patch the LOD visibility with the groom asset value, which might be different from what has been serialized
-				for (uint32 LODIt = 0, LODCount = GroupData.Strands.ClusterCullingBulkData.LODVisibility.Num(); LODIt < LODCount; ++LODIt)
+				for (uint32 LODIt = 0, LODCount = GroupData.Strands.ClusterCullingBulkData.Header.LODVisibility.Num(); LODIt < LODCount; ++LODIt)
 				{
 					if (LODIt < HairGroupLODCount)
 					{
 						bLastVisibility = HairGroupsLOD[GroupIndex].LODs[LODIt].bVisible;
-						GroupData.Strands.ClusterCullingBulkData.LODVisibility[LODIt] = bLastVisibility;
+						GroupData.Strands.ClusterCullingBulkData.Header.LODVisibility[LODIt] = bLastVisibility;
 					}
 					else
 					{
-						GroupData.Strands.ClusterCullingBulkData.LODVisibility[LODIt] = bLastVisibility;
+						GroupData.Strands.ClusterCullingBulkData.Header.LODVisibility[LODIt] = bLastVisibility;
 					}
 				}
 
diff --git a/Engine/Plugins/Runtime/HairStrands/Source/HairStrandsCore/Private/GroomBuilder.cpp b/Engine/Plugins/Runtime/HairStrands/Source/HairStrandsCore/Private/GroomBuilder.cpp
--- a/Engine/Plugins/Runtime/HairStrands/Source/HairStrandsCore/Private/GroomBuilder.cpp
+++ b/Engine/Plugins/Runtime/HairStrands/Source/HairStrandsCore/Private/GroomBuilder.cpp
@@ -40,7 +40,7 @@ static FAutoConsoleVariableRef CVarHairGroupIndexBuilder_MaxVoxelResolution(TEXT
 
 FString FGroomBuilder::GetVersion()
 {
-	return TEXT("v8n");
+	return TEXT("v8o");
 }
 
 namespace FHairStrandsDecimation
@@ -346,14 +346,14 @@ namespace HairStrandsBuilder
 			}
 		}
 
-		OutBulkData.BoundingBox = HairStrands.BoundingBox;
-		OutBulkData.CurveCount = HairStrands.GetNumCurves();
-		OutBulkData.PointCount = HairStrands.GetNumPoints();
-		OutBulkData.MaxLength = MaxLength;
-		OutBulkData.MaxRadius = MaxRadius;
-		OutBulkData.Flags = FHairStrandsBulkData::DataFlags_HasData;
-		OutBulkData.ImportedAttributes = HairStrands.GetAttributes();
-		OutBulkData.ImportedAttributeFlags = HairStrands.GetAttributeFlags();
+		OutBulkData.Header.BoundingBox = HairStrands.BoundingBox;
+		OutBulkData.Header.CurveCount = HairStrands.GetNumCurves();
+		OutBulkData.Header.PointCount = HairStrands.GetNumPoints();
+		OutBulkData.Header.MaxLength = MaxLength;
+		OutBulkData.Header.MaxRadius = MaxRadius;
+		OutBulkData.Header.Flags = FHairStrandsBulkData::DataFlags_HasData;
+		OutBulkData.Header.ImportedAttributes = HairStrands.GetAttributes();
+		OutBulkData.Header.ImportedAttributeFlags = HairStrands.GetAttributeFlags();
 
 		const uint32 UintToByte = 4;
 
@@ -368,40 +368,40 @@ namespace HairStrandsBuilder
 
 			for (uint32 AttributeIt=0; AttributeIt< HAIR_CURVE_ATTRIBUTE_COUNT; ++AttributeIt)
 			{
-				OutBulkData.CurveAttributeOffsets[AttributeIt] = HAIR_ATTRIBUTE_INVALID_OFFSET;
+				OutBulkData.Header.CurveAttributeOffsets[AttributeIt] = HAIR_ATTRIBUTE_INVALID_OFFSET;
 			}
 
 			if (HasHairAttribute(Attributes, EHairAttribute::RootUV))
 			{
-				OutBulkData.CurveAttributeOffsets[HAIR_CURVE_ATTRIBUTE_ROOTUV] = OutPackedAttributes.Num() * UintToByte;
+				OutBulkData.Header.CurveAttributeOffsets[HAIR_CURVE_ATTRIBUTE_ROOTUV] = OutPackedAttributes.Num() * UintToByte;
 				OutPackedAttributes.Append(AttributeRootUV);
 			}
 
 			{
-				OutBulkData.CurveAttributeOffsets[HAIR_CURVE_ATTRIBUTE_SEED]   = OutPackedAttributes.Num() * UintToByte;
+				OutBulkData.Header.CurveAttributeOffsets[HAIR_CURVE_ATTRIBUTE_SEED]   = OutPackedAttributes.Num() * UintToByte;
 				OutPackedAttributes.Append(AttributeSeed);
 			}
 
 			{
-				OutBulkData.CurveAttributeOffsets[HAIR_CURVE_ATTRIBUTE_LENGTH] = OutPackedAttributes.Num() * UintToByte;
+				OutBulkData.Header.CurveAttributeOffsets[HAIR_CURVE_ATTRIBUTE_LENGTH] = OutPackedAttributes.Num() * UintToByte;
 				OutPackedAttributes.Append(AttributeLength);
 			}
 
 			if (HasHairAttribute(Attributes, EHairAttribute::ClumpID))
 			{
-				if (HasHairAttributeFlags(OutBulkData.ImportedAttributeFlags, EHairAttributeFlags::HasMultipleClumpIDs))
+				if (HasHairAttributeFlags(OutBulkData.Header.ImportedAttributeFlags, EHairAttributeFlags::HasMultipleClumpIDs))
 				{
-					OutBulkData.CurveAttributeOffsets[HAIR_CURVE_ATTRIBUTE_CLUMPID3] = OutPackedAttributes.Num() * UintToByte;
+					OutBulkData.Header.CurveAttributeOffsets[HAIR_CURVE_ATTRIBUTE_CLUMPID3] = OutPackedAttributes.Num() * UintToByte;
 					OutPackedAttributes.Append(AttributeClumpIDs);
 				}
 				else
 				{
-					OutBulkData.CurveAttributeOffsets[HAIR_CURVE_ATTRIBUTE_CLUMPID] = OutPackedAttributes.Num() * UintToByte;
+					OutBulkData.Header.CurveAttributeOffsets[HAIR_CURVE_ATTRIBUTE_CLUMPID] = OutPackedAttributes.Num() * UintToByte;
 					OutPackedAttributes.Append(AttributeClumpIDs);
 				}
 			}
 
-			CopyToBulkData<FHairStrandsAttributeFormat>(OutBulkData.CurveAttributes, OutPackedAttributes);
+			CopyToBulkData<FHairStrandsAttributeFormat>(OutBulkData.Data.CurveAttributes, OutPackedAttributes);
 		}
 
 		// Concatenate all point-attributes
@@ -414,39 +414,39 @@ namespace HairStrandsBuilder
 
 			for (uint32 AttributeIt=0; AttributeIt< HAIR_POINT_ATTRIBUTE_COUNT; ++AttributeIt)
 			{
-				OutBulkData.PointAttributeOffsets[AttributeIt] = HAIR_ATTRIBUTE_INVALID_OFFSET;
+				OutBulkData.Header.PointAttributeOffsets[AttributeIt] = HAIR_ATTRIBUTE_INVALID_OFFSET;
 			}
 			
 			if (HasHairAttribute(Attributes, EHairAttribute::Color))
 			{
-				OutBulkData.PointAttributeOffsets[HAIR_POINT_ATTRIBUTE_COLOR] = OutPackedAttributes.Num() * UintToByte;
+				OutBulkData.Header.PointAttributeOffsets[HAIR_POINT_ATTRIBUTE_COLOR] = OutPackedAttributes.Num() * UintToByte;
 				OutPackedAttributes.Append(AttributeColor);
 			}
 
 			if (HasHairAttribute(Attributes, EHairAttribute::Roughness))
 			{
-				OutBulkData.PointAttributeOffsets[HAIR_POINT_ATTRIBUTE_ROUGHNESS] = OutPackedAttributes.Num() * UintToByte;
+				OutBulkData.Header.PointAttributeOffsets[HAIR_POINT_ATTRIBUTE_ROUGHNESS] = OutPackedAttributes.Num() * UintToByte;
 				OutPackedAttributes.Append(AttributeRoughness);
 			}
 
 			if (HasHairAttribute(Attributes, EHairAttribute::AO))
 			{
-				OutBulkData.PointAttributeOffsets[HAIR_POINT_ATTRIBUTE_AO] = OutPackedAttributes.Num() * UintToByte;
+				OutBulkData.Header.PointAttributeOffsets[HAIR_POINT_ATTRIBUTE_AO] = OutPackedAttributes.Num() * UintToByte;
 				OutPackedAttributes.Append(AttributeAO);
 			}
-			CopyToBulkData<FHairStrandsAttributeFormat>(OutBulkData.PointAttributes, OutPackedAttributes);
+			CopyToBulkData<FHairStrandsAttributeFormat>(OutBulkData.Data.PointAttributes, OutPackedAttributes);
 		}
 
-		CopyToBulkData<FHairStrandsPositionFormat>(OutBulkData.Positions, OutPackedPositions);
-		CopyToBulkData<FHairStrandsCurveFormat>(OutBulkData.Curves, OutPackedCurves);
+		CopyToBulkData<FHairStrandsPositionFormat>(OutBulkData.Data.Positions, OutPackedPositions);
+		CopyToBulkData<FHairStrandsCurveFormat>(OutBulkData.Data.Curves, OutPackedCurves);
 		if (bUse16bitsCurveIndex)
 		{
-			OutBulkData.Flags |= FHairStrandsBulkData::DataFlags_Has16bitsCurveIndex;
-			CopyToBulkData<FHairStrandsPointToCurveFormat16>(OutBulkData.PointToCurve, OutPointToCurve16);
+			OutBulkData.Header.Flags |= FHairStrandsBulkData::DataFlags_Has16bitsCurveIndex;
+			CopyToBulkData<FHairStrandsPointToCurveFormat16>(OutBulkData.Data.PointToCurve, OutPointToCurve16);
 		}
 		else
 		{
-			CopyToBulkData<FHairStrandsPointToCurveFormat32>(OutBulkData.PointToCurve, OutPointToCurve32);
+			CopyToBulkData<FHairStrandsPointToCurveFormat32>(OutBulkData.Data.PointToCurve, OutPointToCurve32);
 		}
 	}
 
@@ -1333,12 +1333,12 @@ namespace HairInterpolationBuilder
 		auto UpperPart = [](uint32 Index) -> uint8  { return uint8((Index >> 16) & 0xFF); };
 
 
-		OutBulkData.Flags = FHairStrandsInterpolationBulkData::DataFlags_HasData;
-		OutBulkData.PointCount = PointCount;
+		OutBulkData.Header.Flags = FHairStrandsInterpolationBulkData::DataFlags_HasData;
+		OutBulkData.Header.PointCount = PointCount;
 
 		if (HairInterpolation.bUseUniqueGuide)
 		{
-			OutBulkData.Flags |= FHairStrandsInterpolationBulkData::DataFlags_HasSingleGuideData;
+			OutBulkData.Header.Flags |= FHairStrandsInterpolationBulkData::DataFlags_HasSingleGuideData;
 
 			TArray<FHairStrandsInterpolationFormat::Type> OutPointsInterpolation;
 			OutPointsInterpolation.SetNum(PointCount * FHairStrandsInterpolationFormat::ComponentCount);
@@ -1355,7 +1355,7 @@ namespace HairInterpolationBuilder
 				OutInterp.VertexLerp		= S[0] * 255.f;
 			}
 
-			HairStrandsBuilder::CopyToBulkData<FHairStrandsInterpolationFormat>(OutBulkData.Interpolation, OutPointsInterpolation);
+			HairStrandsBuilder::CopyToBulkData<FHairStrandsInterpolationFormat>(OutBulkData.Data.Interpolation, OutPointsInterpolation);
 		}
 		else
 		{
@@ -1389,16 +1389,16 @@ namespace HairInterpolationBuilder
 				OutInterp1.Pad1			= 0;
 			}
 
-			HairStrandsBuilder::CopyToBulkData<FHairStrandsInterpolation0Format>(OutBulkData.Interpolation0, OutPointsInterpolation0);
-			HairStrandsBuilder::CopyToBulkData<FHairStrandsInterpolation1Format>(OutBulkData.Interpolation1, OutPointsInterpolation1);
+			HairStrandsBuilder::CopyToBulkData<FHairStrandsInterpolation0Format>(OutBulkData.Data.Interpolation0, OutPointsInterpolation0);
+			HairStrandsBuilder::CopyToBulkData<FHairStrandsInterpolation1Format>(OutBulkData.Data.Interpolation1, OutPointsInterpolation1);
 		}
 
 		{
-			OutBulkData.SimPointCount = SimDatas.GetNumPoints();
+			OutBulkData.Header.SimPointCount = SimDatas.GetNumPoints();
 			const uint32 SimCurveCount = SimDatas.GetNumCurves();
 
 			TArray<FHairStrandsRootIndexFormat::Type> SimRootPointIndex;
-			SimRootPointIndex.SetNum(OutBulkData.SimPointCount);
+			SimRootPointIndex.SetNum(OutBulkData.Header.SimPointCount);
 			for (uint32 CurveIndex = 0; CurveIndex < SimCurveCount; ++CurveIndex)
 			{
 				const uint16 SimPointCount = SimDatas.StrandsCurves.CurvesCount[CurveIndex];
@@ -1409,7 +1409,7 @@ namespace HairInterpolationBuilder
 				}
 			}
 
-			HairStrandsBuilder::CopyToBulkData<FHairStrandsRootIndexFormat>(OutBulkData.SimRootPointIndex, SimRootPointIndex);
+			HairStrandsBuilder::CopyToBulkData<FHairStrandsRootIndexFormat>(OutBulkData.Data.SimRootPointIndex, SimRootPointIndex);
 		}
 	}
 
@@ -3125,29 +3125,29 @@ static void BuildClusterBulkData(
 {
 	Out.Reset();
 
-	Out.ClusterCount	= In.ClusterCount;
-	Out.VertexCount		= In.VertexCount;
-	Out.VertexLODCount	= In.ClusterVertexIds.Num();
-	Out.ClusterLODCount = In.ClusterLODInfos.Num();
+	Out.Header.ClusterCount		= In.ClusterCount;
+	Out.Header.VertexCount		= In.VertexCount;
+	Out.Header.VertexLODCount	= In.ClusterVertexIds.Num();
+	Out.Header.ClusterLODCount 	= In.ClusterLODInfos.Num();
 
-	Out.CPULODScreenSize= In.CPULODScreenSize;
-	Out.LODVisibility	= In.LODVisibility;
-	Out.LODInfos		= In.LODInfos;
+	Out.Header.CPULODScreenSize	= In.CPULODScreenSize;
+	Out.Header.LODVisibility	= In.LODVisibility;
+	Out.Header.LODInfos			= In.LODInfos;
 
 	// Sanity check
-	check(Out.ClusterCount			== uint32(In.ClusterInfos.Num()));
-	check(Out.ClusterLODCount		== uint32(In.ClusterLODInfos.Num()));
-	check(Out.VertexCount			== uint32(In.VertexToClusterIds.Num()));
-	check(Out.VertexLODCount		== uint32(In.ClusterVertexIds.Num()));
+	check(Out.Header.ClusterCount		== uint32(In.ClusterInfos.Num()));
+	check(Out.Header.ClusterLODCount	== uint32(In.ClusterLODInfos.Num()));
+	check(Out.Header.VertexCount		== uint32(In.VertexToClusterIds.Num()));
+	check(Out.Header.VertexLODCount		== uint32(In.ClusterVertexIds.Num()));
 	
-	HairStrandsBuilder::CopyToBulkData<FHairClusterLODInfoFormat>(Out.ClusterLODInfos, In.ClusterLODInfos);
-	HairStrandsBuilder::CopyToBulkData<FHairClusterIndexFormat>(Out.VertexToClusterIds, In.VertexToClusterIds);
-	HairStrandsBuilder::CopyToBulkData<FHairClusterIndexFormat>(Out.ClusterVertexIds, In.ClusterVertexIds);
+	HairStrandsBuilder::CopyToBulkData<FHairClusterLODInfoFormat>(Out.Data.ClusterLODInfos, In.ClusterLODInfos);
+	HairStrandsBuilder::CopyToBulkData<FHairClusterIndexFormat>(Out.Data.VertexToClusterIds, In.VertexToClusterIds);
+	HairStrandsBuilder::CopyToBulkData<FHairClusterIndexFormat>(Out.Data.ClusterVertexIds, In.ClusterVertexIds);
 
 	// Pack LODInfo into GPU format
 	{
-		check(uint32(In.ClusterInfos.Num()) == Out.ClusterCount);
-		check(uint32(In.VertexToClusterIds.Num()) == Out.VertexCount);
+		check(uint32(In.ClusterInfos.Num()) == Out.Header.ClusterCount);
+		check(uint32(In.VertexToClusterIds.Num()) == Out.Header.VertexCount);
 	
 		TArray<FHairClusterInfo::Packed> PackedClusterInfos;
 		PackedClusterInfos.Reserve(In.ClusterInfos.Num());
@@ -3183,7 +3183,7 @@ static void BuildClusterBulkData(
 			static_assert(sizeof(FHairClusterInfo::Packed) == sizeof(FHairClusterInfo::BulkType));
 		}
 
-		HairStrandsBuilder::CopyToBulkData<FHairClusterInfoFormat>(Out.PackedClusterInfos, PackedClusterInfos);
+		HairStrandsBuilder::CopyToBulkData<FHairClusterInfoFormat>(Out.Data.PackedClusterInfos, PackedClusterInfos);
 	}
 }
 
diff --git a/Engine/Plugins/Runtime/HairStrands/Source/HairStrandsCore/Private/GroomComponent.cpp b/Engine/Plugins/Runtime/HairStrands/Source/HairStrandsCore/Private/GroomComponent.cpp
--- a/Engine/Plugins/Runtime/HairStrands/Source/HairStrandsCore/Private/GroomComponent.cpp
+++ b/Engine/Plugins/Runtime/HairStrands/Source/HairStrandsCore/Private/GroomComponent.cpp
@@ -135,7 +135,7 @@ static FHairGroupDesc GetGroomGroupsDesc(const UGroomAsset* Asset, UGroomCompone
 	}
 
 	FHairGroupDesc O = Component->GroomGroupsDesc[GroupIndex];
-	O.HairLength = Asset->HairGroupsPlatformData[GroupIndex].Strands.BulkData.MaxLength;
+	O.HairLength = Asset->HairGroupsPlatformData[GroupIndex].Strands.BulkData.Header.MaxLength;
 	O.LODBias 	 = Asset->EffectiveLODBias[GroupIndex] > 0 ? FMath::Max(O.LODBias, Asset->EffectiveLODBias[GroupIndex]) : O.LODBias;
 
 	// Width can be overridden at 3 levels:
@@ -147,7 +147,7 @@ static FHairGroupDesc GetGroomGroupsDesc(const UGroomAsset* Asset, UGroomCompone
 		O.HairWidth = 
 			Asset->HairGroupsRendering[GroupIndex].GeometrySettings.HairWidth_Override ? 
 			Asset->HairGroupsRendering[GroupIndex].GeometrySettings.HairWidth : 
-			Asset->HairGroupsPlatformData[GroupIndex].Strands.BulkData.MaxRadius * 2.0f;
+			Asset->HairGroupsPlatformData[GroupIndex].Strands.BulkData.Header.MaxRadius * 2.0f;
 	}
 
 	if (!O.HairRootScale_Override)				{ O.HairRootScale				= Asset->HairGroupsRendering[GroupIndex].GeometrySettings.HairRootScale;				}
@@ -2718,7 +2718,7 @@ void UGroomComponent::InitResources(bool bIsBindingReloading)
 			{
 				// This codes assumes strands LOD are contigus and the highest (i.e., 0...x). Change this code to something more robust
 				check(HairGroupInstance->HairGroupPublicData);
-				const int32 StrandsLODCount = GroupData.Strands.ClusterCullingResource->BulkData.CPULODScreenSize.Num();
+				const int32 StrandsLODCount = GroupData.Strands.ClusterCullingResource->BulkData.Header.CPULODScreenSize.Num();
 				const TArray<float>& LODScreenSizes = HairGroupInstance->HairGroupPublicData->GetLODScreenSizes();
 				const TArray<bool>& LODVisibilities = HairGroupInstance->HairGroupPublicData->GetLODVisibilities();
 				check(StrandsLODCount <= LODScreenSizes.Num());
@@ -2729,11 +2729,11 @@ void UGroomComponent::InitResources(bool bIsBindingReloading)
 					// Other type of geometry are not serizalized, and so won't match
 					if (GroomAsset->HairGroupsLOD[GroupIt].LODs[LODIt].GeometryType == EGroomGeometryType::Strands)
 					{
-						check(GroupData.Strands.ClusterCullingResource->BulkData.CPULODScreenSize[LODIt] == LODScreenSizes[LODIt]);
-						check(GroupData.Strands.ClusterCullingResource->BulkData.LODVisibility[LODIt] == LODVisibilities[LODIt]);
+						check(GroupData.Strands.ClusterCullingResource->BulkData.Header.CPULODScreenSize[LODIt] == LODScreenSizes[LODIt]);
+						check(GroupData.Strands.ClusterCullingResource->BulkData.Header.LODVisibility[LODIt] == LODVisibilities[LODIt]);
 					}
 				}
-				HairGroupInstance->HairGroupPublicData->ClusterCount = HairGroupInstance->Strands.ClusterCullingResource->BulkData.ClusterCount;
+				HairGroupInstance->HairGroupPublicData->ClusterCount = HairGroupInstance->Strands.ClusterCullingResource->BulkData.Header.ClusterCount;
 				BeginInitResource(HairGroupInstance->HairGroupPublicData);
 			}
 
diff --git a/Engine/Plugins/Runtime/HairStrands/Source/HairStrandsCore/Private/GroomDebug.cpp b/Engine/Plugins/Runtime/HairStrands/Source/HairStrandsCore/Private/GroomDebug.cpp
--- a/Engine/Plugins/Runtime/HairStrands/Source/HairStrandsCore/Private/GroomDebug.cpp
+++ b/Engine/Plugins/Runtime/HairStrands/Source/HairStrandsCore/Private/GroomDebug.cpp
@@ -967,7 +967,7 @@ static void AddHairDebugPrintInstancePass(
 					Data2.Y = Instance->HairGroupPublicData->GetActiveStrandsPointCount();
 					Data2.Z = Instance->HairGroupPublicData->GetActiveStrandsCurveCount();
 
-					Data2.W = Instance->Strands.Data->ImportedAttributes;
+					Data2.W = Instance->Strands.Data->Header.ImportedAttributes;
 				}
 				
 				Data3 = FUintVector4(0);
diff --git a/Engine/Plugins/Runtime/HairStrands/Source/HairStrandsCore/Private/GroomManager.cpp b/Engine/Plugins/Runtime/HairStrands/Source/HairStrandsCore/Private/GroomManager.cpp
--- a/Engine/Plugins/Runtime/HairStrands/Source/HairStrandsCore/Private/GroomManager.cpp
+++ b/Engine/Plugins/Runtime/HairStrands/Source/HairStrandsCore/Private/GroomManager.cpp
@@ -843,7 +843,7 @@ static void RunHairLODSelection(
 			for (const FSceneView* View : Views)
 			{
 				const FVector3d BoundScale = Instance->LocalToWorld.GetScale3D();
-				const FVector3f BoundExtent = Instance->Strands.Data ? FVector3f(Instance->Strands.Data->BoundingBox.GetExtent()) : FVector3f(0,0,0);
+				const FVector3f BoundExtent = Instance->Strands.Data ? FVector3f(Instance->Strands.Data->Header.BoundingBox.GetExtent()) : FVector3f(0,0,0);
 				const float BoundRadius = FMath::Max3(BoundExtent.X, BoundExtent.Y, BoundExtent.Z) * FMath::Max3(BoundScale.X, BoundScale.Y, BoundScale.Z);
 
 				const float ScreenSize_RestBound = FMath::Clamp(ComputeBoundsScreenSize(FVector4(SphereBound.Center, 1), BoundRadius, *View), 0.f, 1.0f);
diff --git a/Engine/Plugins/Runtime/HairStrands/Source/HairStrandsCore/Private/GroomResources.cpp b/Engine/Plugins/Runtime/HairStrands/Source/HairStrandsCore/Private/GroomResources.cpp
--- a/Engine/Plugins/Runtime/HairStrands/Source/HairStrandsCore/Private/GroomResources.cpp
+++ b/Engine/Plugins/Runtime/HairStrands/Source/HairStrandsCore/Private/GroomResources.cpp
@@ -865,7 +865,7 @@ FHairStrandsRestResource::FHairStrandsRestResource(FHairStrandsBulkData& InBulkD
 	PositionBuffer(), PointAttributeBuffer(), CurveAttributeBuffer(), PointToCurveBuffer(), BulkData(InBulkData), CurveType(InCurveType)
 {
 	// Sanity check
-	check(!!(BulkData.Flags & FHairStrandsBulkData::DataFlags_HasData));
+	check(!!(BulkData.Header.Flags & FHairStrandsBulkData::DataFlags_HasData));
 }
 
 bool FHairStrandsRestResource::InternalIsDataLoaded()
@@ -873,11 +873,11 @@ bool FHairStrandsRestResource::InternalIsDataLoaded()
 	if (BulkDataRequest.IsNone())
 	{
 		FBulkDataBatchRequest::FBatchBuilder Batch = FBulkDataBatchRequest::NewBatch(5);
-		Batch.Read(BulkData.Positions);
-		Batch.Read(BulkData.CurveAttributes);
-		Batch.Read(BulkData.PointAttributes);
-		Batch.Read(BulkData.PointToCurve);
-		Batch.Read(BulkData.Curves);
+		Batch.Read(BulkData.Data.Positions);
+		Batch.Read(BulkData.Data.CurveAttributes);
+		Batch.Read(BulkData.Data.PointAttributes);
+		Batch.Read(BulkData.Data.PointToCurve);
+		Batch.Read(BulkData.Data.Curves);
 
 		Batch.Issue(BulkDataRequest);
 	}
@@ -895,17 +895,17 @@ void FHairStrandsRestResource::InternalAllocate(FRDGBuilder& GraphBuilder)
 
 	// 1. Lock data, which force the loading data from files (on non-editor build/cooked data). These data are then uploaded to the GPU
 	// 2. A local copy is done by the buffer uploader. This copy is discarded once the uploading is done.
-	InternalCreateVertexBufferRDG_FromBulkData<FHairStrandsPositionFormat>(GraphBuilder, BulkData.Positions, PointCount, PositionBuffer, ToHairResourceDebugName(HAIRSTRANDS_RESOUCE_NAME(CurveType, Hair.StrandsRest_PositionBuffer), ResourceName), OwnerName, EHairResourceUsageType::Static);
-	InternalCreateByteAddressBufferRDG_FromBulkData(GraphBuilder, BulkData.CurveAttributes, CurveAttributeBuffer, ToHairResourceDebugName(HAIRSTRANDS_RESOUCE_NAME(CurveType, Hair.StrandsRest_CurveAttributeBuffer), ResourceName), OwnerName, EHairResourceUsageType::Static);
-	InternalCreateByteAddressBufferRDG_FromBulkData(GraphBuilder, BulkData.PointAttributes, PointAttributeBuffer, ToHairResourceDebugName(HAIRSTRANDS_RESOUCE_NAME(CurveType, Hair.StrandsRest_PointAttributeBuffer), ResourceName), OwnerName, EHairResourceUsageType::Static);
-	InternalCreateVertexBufferRDG_FromBulkData<FHairStrandsCurveFormat>(GraphBuilder, BulkData.Curves, CurveCount, CurveBuffer, ToHairResourceDebugName(HAIRSTRANDS_RESOUCE_NAME(CurveType, Hair.StrandsRest_CurveBuffer), ResourceName), OwnerName, EHairResourceUsageType::Static);
-	if (!!(BulkData.Flags & BulkData.DataFlags_Has16bitsCurveIndex))
+	InternalCreateVertexBufferRDG_FromBulkData<FHairStrandsPositionFormat>(GraphBuilder, BulkData.Data.Positions, PointCount, PositionBuffer, ToHairResourceDebugName(HAIRSTRANDS_RESOUCE_NAME(CurveType, Hair.StrandsRest_PositionBuffer), ResourceName), OwnerName, EHairResourceUsageType::Static);
+	InternalCreateByteAddressBufferRDG_FromBulkData(GraphBuilder, BulkData.Data.CurveAttributes, CurveAttributeBuffer, ToHairResourceDebugName(HAIRSTRANDS_RESOUCE_NAME(CurveType, Hair.StrandsRest_CurveAttributeBuffer), ResourceName), OwnerName, EHairResourceUsageType::Static);
+	InternalCreateByteAddressBufferRDG_FromBulkData(GraphBuilder, BulkData.Data.PointAttributes, PointAttributeBuffer, ToHairResourceDebugName(HAIRSTRANDS_RESOUCE_NAME(CurveType, Hair.StrandsRest_PointAttributeBuffer), ResourceName), OwnerName, EHairResourceUsageType::Static);
+	InternalCreateVertexBufferRDG_FromBulkData<FHairStrandsCurveFormat>(GraphBuilder, BulkData.Data.Curves, CurveCount, CurveBuffer, ToHairResourceDebugName(HAIRSTRANDS_RESOUCE_NAME(CurveType, Hair.StrandsRest_CurveBuffer), ResourceName), OwnerName, EHairResourceUsageType::Static);
+	if (!!(BulkData.Header.Flags & FHairStrandsBulkData::DataFlags_Has16bitsCurveIndex))
 	{
-		InternalCreateVertexBufferRDG_FromBulkData<FHairStrandsPointToCurveFormat16>(GraphBuilder, BulkData.PointToCurve, PointCount, PointToCurveBuffer, ToHairResourceDebugName(HAIRSTRANDS_RESOUCE_NAME(CurveType, Hair.StrandsRest_PointToCurveBuffer), ResourceName), OwnerName, EHairResourceUsageType::Static);
+		InternalCreateVertexBufferRDG_FromBulkData<FHairStrandsPointToCurveFormat16>(GraphBuilder, BulkData.Data.PointToCurve, PointCount, PointToCurveBuffer, ToHairResourceDebugName(HAIRSTRANDS_RESOUCE_NAME(CurveType, Hair.StrandsRest_PointToCurveBuffer), ResourceName), OwnerName, EHairResourceUsageType::Static);
 	}
 	else
 	{
-		InternalCreateVertexBufferRDG_FromBulkData<FHairStrandsPointToCurveFormat32>(GraphBuilder, BulkData.PointToCurve, PointCount, PointToCurveBuffer, ToHairResourceDebugName(HAIRSTRANDS_RESOUCE_NAME(CurveType, Hair.StrandsRest_PointToCurveBuffer), ResourceName), OwnerName, EHairResourceUsageType::Static);
+		InternalCreateVertexBufferRDG_FromBulkData<FHairStrandsPointToCurveFormat32>(GraphBuilder, BulkData.Data.PointToCurve, PointCount, PointToCurveBuffer, ToHairResourceDebugName(HAIRSTRANDS_RESOUCE_NAME(CurveType, Hair.StrandsRest_PointToCurveBuffer), ResourceName), OwnerName, EHairResourceUsageType::Static);
 	}
 
 	TArray<FVector4f> RestOffset;
@@ -919,9 +919,9 @@ void FHairStrandsRestResource::InternalAllocate(FRDGBuilder& GraphBuilder)
 		CurveData.SetNum(CurveCount);
 		uint32* Data = (uint32*)CurveData.GetData();
 
-		const uint32* ReadOnlyData = (const uint32*)BulkData.Curves.Lock(LOCK_READ_ONLY);
+		const uint32* ReadOnlyData = (const uint32*)BulkData.Data.Curves.Lock(LOCK_READ_ONLY);
 		memcpy(Data, ReadOnlyData, sizeof(uint32) * CurveCount);
-		BulkData.Curves.Unlock();
+		BulkData.Data.Curves.Unlock();
 	}
 }
 
@@ -1032,65 +1032,86 @@ FHairStrandsClusterCullingBulkData::FHairStrandsClusterCullingBulkData()
 
 void FHairStrandsClusterCullingBulkData::Reset()
 {
-	ClusterCount = 0;
-	ClusterLODCount = 0;
-	VertexCount = 0;
-	VertexLODCount = 0;
+	Header.ClusterCount = 0;
+	Header.ClusterLODCount = 0;
+	Header.VertexCount = 0;
+	Header.VertexLODCount = 0;
 
-	LODVisibility.Empty();;
-	CPULODScreenSize.Empty();
+	Header.LODVisibility.Empty();
+	Header.CPULODScreenSize.Empty();
+	Header.LODInfos.Empty();
 
-	ClusterLODInfos.RemoveBulkData();
-	VertexToClusterIds.RemoveBulkData();
-	ClusterVertexIds.RemoveBulkData();
-	PackedClusterInfos.RemoveBulkData();
+	Data.ClusterLODInfos.RemoveBulkData();
+	Data.VertexToClusterIds.RemoveBulkData();
+	Data.ClusterVertexIds.RemoveBulkData();
+	Data.PackedClusterInfos.RemoveBulkData();
 
 	// Reset the bulk byte buffer to ensure the (serialize) data size is reset to 0
-	ClusterLODInfos 	= FByteBulkData();
-	VertexToClusterIds 	= FByteBulkData();
-	ClusterVertexIds 	= FByteBulkData();
-	PackedClusterInfos 	= FByteBulkData();
+	Data.ClusterLODInfos 	= FByteBulkData();
+	Data.VertexToClusterIds = FByteBulkData();
+	Data.ClusterVertexIds 	= FByteBulkData();
+	Data.PackedClusterInfos = FByteBulkData();
 }
 
 void FHairStrandsClusterCullingBulkData::Serialize(FArchive& Ar, UObject* Owner)
 {
-	Ar << ClusterCount;
-	Ar << ClusterLODCount;
-	Ar << VertexCount;
-	Ar << VertexLODCount;
-	Ar << LODVisibility;
-	Ar << CPULODScreenSize;
+	SerializeHeader(Ar, Owner);
+	SerializeData(Ar, Owner);
+}
+
+void FHairStrandsClusterCullingBulkData::SerializeHeader(FArchive& Ar, UObject* Owner)
+{
+	Ar << Header.ClusterCount;
+	Ar << Header.ClusterLODCount;
+	Ar << Header.VertexCount;
+	Ar << Header.VertexLODCount;
+	Ar << Header.LODVisibility;
+	Ar << Header.CPULODScreenSize;
+	uint32 LODInfosCount = Header.LODInfos.Num();
+	Ar << LODInfosCount;
+	if (Ar.IsLoading())
+	{
+		Header.LODInfos.SetNum(LODInfosCount);
+	}
+	for (uint32 It = 0; It < LODInfosCount; ++It)
+	{
+		Ar << Header.LODInfos[It].CurveCount;
+		Ar << Header.LODInfos[It].PointCount;
+	}
+}
 
+void FHairStrandsClusterCullingBulkData::SerializeData(FArchive& Ar, UObject* Owner)
+{
 	const int32 ChunkIndex = 0;
 	bool bAttemptFileMapping = false;
 
 	if (Ar.IsSaving())
 	{
 		const uint32 BulkFlags = BULKDATA_Force_NOT_InlinePayload;
-		ClusterLODInfos.SetBulkDataFlags(BulkFlags);
-		VertexToClusterIds.SetBulkDataFlags(BulkFlags);
-		ClusterVertexIds.SetBulkDataFlags(BulkFlags);
-		PackedClusterInfos.SetBulkDataFlags(BulkFlags);
+		Data.ClusterLODInfos.SetBulkDataFlags(BulkFlags);
+		Data.VertexToClusterIds.SetBulkDataFlags(BulkFlags);
+		Data.ClusterVertexIds.SetBulkDataFlags(BulkFlags);
+		Data.PackedClusterInfos.SetBulkDataFlags(BulkFlags);
 	}
 
-	if (ClusterLODCount)
+	if (Header.ClusterLODCount)
 	{
-		ClusterLODInfos.Serialize(Ar, Owner, ChunkIndex, bAttemptFileMapping);
+		Data.ClusterLODInfos.Serialize(Ar, Owner, ChunkIndex, bAttemptFileMapping);
 	}
 
-	if (VertexCount)
+	if (Header.VertexCount)
 	{
-		VertexToClusterIds.Serialize(Ar, Owner, ChunkIndex, bAttemptFileMapping);
+		Data.VertexToClusterIds.Serialize(Ar, Owner, ChunkIndex, bAttemptFileMapping);
 	}
 
-	if (VertexLODCount)
+	if (Header.VertexLODCount)
 	{
-		ClusterVertexIds.Serialize(Ar, Owner, ChunkIndex, bAttemptFileMapping);
+		Data.ClusterVertexIds.Serialize(Ar, Owner, ChunkIndex, bAttemptFileMapping);
 	}
 
-	if (ClusterCount)
+	if (Header.ClusterCount)
 	{
-		PackedClusterInfos.Serialize(Ar, Owner, ChunkIndex, bAttemptFileMapping);
+		Data.PackedClusterInfos.Serialize(Ar, Owner, ChunkIndex, bAttemptFileMapping);
 	}
 
 	if (GHairStrandsBulkData_Validation > 0 && Ar.IsSaving())
@@ -1101,13 +1122,13 @@ void FHairStrandsClusterCullingBulkData::Serialize(FArchive& Ar, UObject* Owner)
 
 void FHairStrandsClusterCullingBulkData::Validate(bool bIsSaving)
 {
-	if (ClusterCount == 0)
+	if (Header.ClusterCount == 0)
 		return;
 
-	const FHairClusterInfo::Packed* Datas = (const FHairClusterInfo::Packed*)PackedClusterInfos.Lock(LOCK_READ_ONLY);
+	const FHairClusterInfo::Packed* Datas = (const FHairClusterInfo::Packed*)Data.PackedClusterInfos.Lock(LOCK_READ_ONLY);
 	
 	// Simple heuristic to check if the data are valid
-	const uint32 MaxCount = FMath::Min(ClusterCount, 128u);
+	const uint32 MaxCount = FMath::Min(Header.ClusterCount, 128u);
 	bool bIsValid = true;
 	for (uint32 It = 0; It < MaxCount; ++It)
 	{
@@ -1116,11 +1137,11 @@ void FHairStrandsClusterCullingBulkData::Validate(bool bIsSaving)
 	}
 	if (!bIsValid)
 	{
-		FString Name = ClusterLODInfos.GetDebugName();
-		UE_LOG(LogHairStrands, Error, TEXT("[Groom/DDC] Strands - Invalid ClusterCullingBulkData when %s bulk data - %s"), bIsSaving ? TEXT("Saving") : TEXT("Loading"), *Name);
+		FString DebugName = Data.ClusterLODInfos.GetDebugName();
+		UE_LOG(LogHairStrands, Error, TEXT("[Groom/DDC] Strands - Invalid ClusterCullingBulkData when %s bulk data - %s"), bIsSaving ? TEXT("Saving") : TEXT("Loading"), *DebugName);
 	}
 
-	PackedClusterInfos.Unlock();
+	Data.PackedClusterInfos.Unlock();
 }
 
 /////////////////////////////////////////////////////////////////////////////////////////
@@ -1137,10 +1158,10 @@ bool FHairStrandsClusterCullingResource::InternalIsDataLoaded()
 	if (BulkDataRequest.IsNone())
 	{
 		FBulkDataBatchRequest::NewBatch(4)
-			.Read(BulkData.PackedClusterInfos)
-			.Read(BulkData.ClusterLODInfos)
-			.Read(BulkData.VertexToClusterIds)
-			.Read(BulkData.ClusterVertexIds)
+			.Read(BulkData.Data.PackedClusterInfos)
+			.Read(BulkData.Data.ClusterLODInfos)
+			.Read(BulkData.Data.VertexToClusterIds)
+			.Read(BulkData.Data.ClusterVertexIds)
 			.Issue(BulkDataRequest);
 	}
 
@@ -1156,11 +1177,11 @@ void FHairStrandsClusterCullingResource::InternalAllocate(FRDGBuilder& GraphBuil
 		BulkData.Validate(false);
 	}
 
-	InternalCreateStructuredBufferRDG_FromBulkData<FHairClusterInfoFormat>(GraphBuilder, BulkData.PackedClusterInfos, BulkData.ClusterCount, ClusterInfoBuffer, ToHairResourceDebugName(TEXT("Hair.StrandsClusterCulling_ClusterInfoBuffer"), ResourceName), OwnerName, EHairResourceUsageType::Static);
-	InternalCreateStructuredBufferRDG_FromBulkData<FHairClusterLODInfoFormat>(GraphBuilder, BulkData.ClusterLODInfos, BulkData.ClusterLODCount, ClusterLODInfoBuffer, ToHairResourceDebugName(TEXT("Hair.StrandsClusterCulling_ClusterLODInfoBuffer"), ResourceName), OwnerName, EHairResourceUsageType::Static);
+	InternalCreateStructuredBufferRDG_FromBulkData<FHairClusterInfoFormat>(GraphBuilder, BulkData.Data.PackedClusterInfos, BulkData.Header.ClusterCount, ClusterInfoBuffer, ToHairResourceDebugName(TEXT("Hair.StrandsClusterCulling_ClusterInfoBuffer"), ResourceName), OwnerName, EHairResourceUsageType::Static);
+	InternalCreateStructuredBufferRDG_FromBulkData<FHairClusterLODInfoFormat>(GraphBuilder, BulkData.Data.ClusterLODInfos, BulkData.Header.ClusterLODCount, ClusterLODInfoBuffer, ToHairResourceDebugName(TEXT("Hair.StrandsClusterCulling_ClusterLODInfoBuffer"), ResourceName), OwnerName, EHairResourceUsageType::Static);
 
-	InternalCreateVertexBufferRDG_FromBulkData<FHairClusterIndexFormat>(GraphBuilder, BulkData.VertexToClusterIds, BulkData.VertexCount, VertexToClusterIdBuffer, ToHairResourceDebugName(TEXT("Hair.StrandsClusterCulling_VertexToClusterIds"), ResourceName), OwnerName, EHairResourceUsageType::Static);
-	InternalCreateVertexBufferRDG_FromBulkData<FHairClusterIndexFormat>(GraphBuilder, BulkData.ClusterVertexIds, BulkData.VertexLODCount, ClusterVertexIdBuffer, ToHairResourceDebugName(TEXT("Hair.StrandsClusterCulling_ClusterVertexIds"), ResourceName), OwnerName, EHairResourceUsageType::Static);
+	InternalCreateVertexBufferRDG_FromBulkData<FHairClusterIndexFormat>(GraphBuilder, BulkData.Data.VertexToClusterIds, BulkData.Header.VertexCount, VertexToClusterIdBuffer, ToHairResourceDebugName(TEXT("Hair.StrandsClusterCulling_VertexToClusterIds"), ResourceName), OwnerName, EHairResourceUsageType::Static);
+	InternalCreateVertexBufferRDG_FromBulkData<FHairClusterIndexFormat>(GraphBuilder, BulkData.Data.ClusterVertexIds, BulkData.Header.VertexLODCount, ClusterVertexIdBuffer, ToHairResourceDebugName(TEXT("Hair.StrandsClusterCulling_ClusterVertexIds"), ResourceName), OwnerName, EHairResourceUsageType::Static);
 }
 
 void FHairStrandsClusterCullingResource::InternalRelease()
@@ -1590,7 +1611,7 @@ FHairStrandsInterpolationResource::FHairStrandsInterpolationResource(FHairStrand
 	InterpolationBuffer(), Interpolation0Buffer(), Interpolation1Buffer(), BulkData(InBulkData)
 {
 	// Sanity check
-	check(!!(BulkData.Flags & FHairStrandsInterpolationBulkData::DataFlags_HasData));
+	check(!!(BulkData.Header.Flags & FHairStrandsInterpolationBulkData::DataFlags_HasData));
 }
 
 bool FHairStrandsInterpolationResource::InternalIsDataLoaded()
@@ -1598,17 +1619,17 @@ bool FHairStrandsInterpolationResource::InternalIsDataLoaded()
 	if (BulkDataRequest.IsNone())
 	{
 		FBulkDataBatchRequest::FBatchBuilder Batch = FBulkDataBatchRequest::NewBatch(4);
-		const bool bUseSingleGuide = !!(BulkData.Flags & FHairStrandsInterpolationBulkData::DataFlags_HasSingleGuideData);
+		const bool bUseSingleGuide = !!(BulkData.Header.Flags & FHairStrandsInterpolationBulkData::DataFlags_HasSingleGuideData);
 		if (bUseSingleGuide)
 		{
-			Batch.Read(BulkData.Interpolation);
+			Batch.Read(BulkData.Data.Interpolation);
 		}
 		else
 		{
-			Batch.Read(BulkData.Interpolation0);
-			Batch.Read(BulkData.Interpolation1);
+			Batch.Read(BulkData.Data.Interpolation0);
+			Batch.Read(BulkData.Data.Interpolation1);
 		}
-		Batch.Read(BulkData.SimRootPointIndex);
+		Batch.Read(BulkData.Data.SimRootPointIndex);
 
 		Batch.Issue(BulkDataRequest);
 	}
@@ -1620,17 +1641,17 @@ void FHairStrandsInterpolationResource::InternalAllocate(FRDGBuilder& GraphBuild
 {
 	BulkDataRequest = FBulkDataBatchRequest();
 
-	const bool bUseSingleGuide = !!(BulkData.Flags & FHairStrandsInterpolationBulkData::DataFlags_HasSingleGuideData);
+	const bool bUseSingleGuide = !!(BulkData.Header.Flags & FHairStrandsInterpolationBulkData::DataFlags_HasSingleGuideData);
 	if (bUseSingleGuide)
 	{
-		InternalCreateVertexBufferRDG_FromBulkData<FHairStrandsInterpolationFormat>(GraphBuilder, BulkData.Interpolation, BulkData.PointCount, InterpolationBuffer, ToHairResourceDebugName(TEXT("Hair.StrandsInterpolation_InterpolationBuffer"), ResourceName), OwnerName, EHairResourceUsageType::Static);
+		InternalCreateVertexBufferRDG_FromBulkData<FHairStrandsInterpolationFormat>(GraphBuilder, BulkData.Data.Interpolation, BulkData.Header.PointCount, InterpolationBuffer, ToHairResourceDebugName(TEXT("Hair.StrandsInterpolation_InterpolationBuffer"), ResourceName), OwnerName, EHairResourceUsageType::Static);
 	}
 	else
 	{
-		InternalCreateVertexBufferRDG_FromBulkData<FHairStrandsInterpolation0Format>(GraphBuilder, BulkData.Interpolation0, BulkData.PointCount, Interpolation0Buffer, ToHairResourceDebugName(TEXT("Hair.StrandsInterpolation_Interpolation0Buffer"), ResourceName), OwnerName, EHairResourceUsageType::Static);
-		InternalCreateVertexBufferRDG_FromBulkData<FHairStrandsInterpolation1Format>(GraphBuilder, BulkData.Interpolation1, BulkData.PointCount, Interpolation1Buffer, ToHairResourceDebugName(TEXT("Hair.StrandsInterpolation_Interpolation1Buffer"), ResourceName), OwnerName, EHairResourceUsageType::Static);
+		InternalCreateVertexBufferRDG_FromBulkData<FHairStrandsInterpolation0Format>(GraphBuilder, BulkData.Data.Interpolation0, BulkData.Header.PointCount, Interpolation0Buffer, ToHairResourceDebugName(TEXT("Hair.StrandsInterpolation_Interpolation0Buffer"), ResourceName), OwnerName, EHairResourceUsageType::Static);
+		InternalCreateVertexBufferRDG_FromBulkData<FHairStrandsInterpolation1Format>(GraphBuilder, BulkData.Data.Interpolation1, BulkData.Header.PointCount, Interpolation1Buffer, ToHairResourceDebugName(TEXT("Hair.StrandsInterpolation_Interpolation1Buffer"), ResourceName), OwnerName, EHairResourceUsageType::Static);
 	}
-	InternalCreateVertexBufferRDG_FromBulkData<FHairStrandsRootIndexFormat>(GraphBuilder, BulkData.SimRootPointIndex, BulkData.SimPointCount, SimRootPointIndexBuffer, ToHairResourceDebugName(TEXT("Hair.StrandsInterpolation_SimRootPointIndex"), ResourceName), OwnerName, EHairResourceUsageType::Static);
+	InternalCreateVertexBufferRDG_FromBulkData<FHairStrandsRootIndexFormat>(GraphBuilder, BulkData.Data.SimRootPointIndex, BulkData.Header.SimPointCount, SimRootPointIndexBuffer, ToHairResourceDebugName(TEXT("Hair.StrandsInterpolation_SimRootPointIndex"), ResourceName), OwnerName, EHairResourceUsageType::Static);
 }
 
 void FHairStrandsInterpolationResource::InternalRelease()
diff --git a/Engine/Plugins/Runtime/HairStrands/Source/HairStrandsCore/Private/GroomTextureBuilder.cpp b/Engine/Plugins/Runtime/HairStrands/Source/HairStrandsCore/Private/GroomTextureBuilder.cpp
--- a/Engine/Plugins/Runtime/HairStrands/Source/HairStrandsCore/Private/GroomTextureBuilder.cpp
+++ b/Engine/Plugins/Runtime/HairStrands/Source/HairStrandsCore/Private/GroomTextureBuilder.cpp
@@ -1057,10 +1057,10 @@ static bool TraceTextures(
 			Instance.Common.RootScale = 1.0f;
 			Instance.Common.TipScale = 1.0f;
 			Instance.Common.LengthScale = 1.0f;
-			Instance.Common.Length = GroupData.Strands.RestResource->BulkData.MaxLength;
+			Instance.Common.Length = GroupData.Strands.RestResource->BulkData.Header.MaxLength;
 			Instance.Common.PositionOffset = FVector3f(GroupData.Strands.RestResource->GetPositionOffset());
-			PACK_CURVE_HAIR_ATTRIBUTE_OFFSETS(Instance.Common.CurveAttributeOffsets, GroupData.Strands.RestResource->BulkData.CurveAttributeOffsets);
-			PACK_POINT_HAIR_ATTRIBUTE_OFFSETS(Instance.Common.PointAttributeOffsets, GroupData.Strands.RestResource->BulkData.PointAttributeOffsets);
+			PACK_CURVE_HAIR_ATTRIBUTE_OFFSETS(Instance.Common.CurveAttributeOffsets, GroupData.Strands.RestResource->BulkData.Header.CurveAttributeOffsets);
+			PACK_POINT_HAIR_ATTRIBUTE_OFFSETS(Instance.Common.PointAttributeOffsets, GroupData.Strands.RestResource->BulkData.Header.PointAttributeOffsets);
 
 			Instance.Resources.PositionBuffer = GroupData.Strands.RestResource->PositionBuffer.SRV;
 			Instance.Resources.CurveAttributeBuffer = GroupData.Strands.RestResource->CurveAttributeBuffer.SRV;
diff --git a/Engine/Plugins/Runtime/HairStrands/Source/HairStrandsCore/Private/HairStrandsDatas.cpp b/Engine/Plugins/Runtime/HairStrands/Source/HairStrandsCore/Private/HairStrandsDatas.cpp
--- a/Engine/Plugins/Runtime/HairStrands/Source/HairStrandsCore/Private/HairStrandsDatas.cpp
+++ b/Engine/Plugins/Runtime/HairStrands/Source/HairStrandsCore/Private/HairStrandsDatas.cpp
@@ -170,143 +170,157 @@ FArchive& operator<<(FArchive& Ar, FHairInterpolation1Vertex& Vertex)
 
 void FHairStrandsInterpolationBulkData::Reset()
 {
-	Flags = 0;
-	PointCount = 0;
-	SimPointCount = 0;
+	Header.Flags = 0;
+	Header.PointCount = 0;
+	Header.SimPointCount = 0;
 	
 	// Deallocate memory if needed
-	Interpolation.RemoveBulkData();
-	Interpolation0.RemoveBulkData();
-	Interpolation1.RemoveBulkData();
-	SimRootPointIndex.RemoveBulkData();
+	Data.Interpolation.RemoveBulkData();
+	Data.Interpolation0.RemoveBulkData();
+	Data.Interpolation1.RemoveBulkData();
+	Data.SimRootPointIndex.RemoveBulkData();
 
 	// Reset the bulk byte buffer to ensure the (serialize) data size is reset to 0
-	Interpolation		= FByteBulkData();
-	Interpolation0		= FByteBulkData();
-	Interpolation1		= FByteBulkData();
-	SimRootPointIndex	= FByteBulkData();
+	Data.Interpolation		= FByteBulkData();
+	Data.Interpolation0		= FByteBulkData();
+	Data.Interpolation1		= FByteBulkData();
+	Data.SimRootPointIndex	= FByteBulkData();
 }
 
 void FHairStrandsInterpolationBulkData::Serialize(FArchive& Ar, UObject* Owner)
+{
+	SerializeHeader(Ar, Owner);
+	SerializeData(Ar, Owner);
+}
+
+void FHairStrandsInterpolationBulkData::SerializeHeader(FArchive& Ar, UObject* Owner)
+{
+	Ar << Header.Flags;
+	Ar << Header.PointCount;
+	Ar << Header.SimPointCount;
+}
+
+void FHairStrandsInterpolationBulkData::SerializeData(FArchive& Ar, UObject* Owner)
 {
 	static_assert(sizeof(FHairInterpolationVertex::BulkType) == sizeof(FHairInterpolationVertex));
 	static_assert(sizeof(FHairInterpolation0Vertex::BulkType) == sizeof(FHairInterpolation0Vertex));
 	static_assert(sizeof(FHairInterpolation1Vertex::BulkType) == sizeof(FHairInterpolation1Vertex));
 	static_assert(sizeof(FHairStrandsRootIndexFormat::BulkType) == sizeof(FHairStrandsRootIndexFormat::Type));
 
-	if (Ar.IsSaving())
-	{
-		const uint32 BulkFlags = BULKDATA_Force_NOT_InlinePayload;
-		Interpolation.SetBulkDataFlags(BulkFlags);
-		Interpolation0.SetBulkDataFlags(BulkFlags);
-		Interpolation1.SetBulkDataFlags(BulkFlags);
-		SimRootPointIndex.SetBulkDataFlags(BulkFlags);
-	}
-
-	Ar << Flags;
-	Ar << PointCount;
-	Ar << SimPointCount;
-
-	if (!!(Flags & DataFlags_HasData))
+	if (!!(Header.Flags & DataFlags_HasData))
 	{
 		const int32 ChunkIndex = 0;
 		bool bAttemptFileMapping = false;
 
-		if (!!(Flags & DataFlags_HasSingleGuideData))
+		if (!!(Header.Flags & DataFlags_HasSingleGuideData))
 		{
-			Interpolation.Serialize(Ar, Owner, ChunkIndex, bAttemptFileMapping);
+			Data.Interpolation.Serialize(Ar, Owner, ChunkIndex, bAttemptFileMapping);
 		}
 		else
 		{
-			Interpolation0.Serialize(Ar, Owner, ChunkIndex, bAttemptFileMapping);
-			Interpolation1.Serialize(Ar, Owner, ChunkIndex, bAttemptFileMapping);
+			Data.Interpolation0.Serialize(Ar, Owner, ChunkIndex, bAttemptFileMapping);
+			Data.Interpolation1.Serialize(Ar, Owner, ChunkIndex, bAttemptFileMapping);
 		}
-		SimRootPointIndex.Serialize(Ar, Owner, ChunkIndex, bAttemptFileMapping);
+		Data.SimRootPointIndex.Serialize(Ar, Owner, ChunkIndex, bAttemptFileMapping);
 	}
 }
 
 void FHairStrandsBulkData::Serialize(FArchive& Ar, UObject* Owner)
 {
-	static_assert(sizeof(FHairStrandsPositionFormat::BulkType) == sizeof(FHairStrandsPositionFormat::Type));
-	static_assert(sizeof(FHairStrandsAttributeFormat::BulkType) == sizeof(FHairStrandsAttributeFormat::Type));
-	static_assert(sizeof(FHairStrandsPointToCurveFormat16::BulkType) == sizeof(FHairStrandsPointToCurveFormat16::Type));
-	static_assert(sizeof(FHairStrandsPointToCurveFormat32::BulkType) == sizeof(FHairStrandsPointToCurveFormat32::Type));
-	static_assert(sizeof(FHairStrandsRootIndexFormat::BulkType) == sizeof(FHairStrandsRootIndexFormat::Type)); 
+	SerializeHeader(Ar, Owner);
+	SerializeData(Ar, Owner);
+}
 
+void FHairStrandsBulkData::SerializeHeader(FArchive& Ar, UObject* Owner)
+{
 	Ar.UsingCustomVersion(FReleaseObjectVersion::GUID);
 	Ar.UsingCustomVersion(FUE5ReleaseStreamObjectVersion::GUID);
 
-	Ar << CurveCount;
-	Ar << PointCount;
-	Ar << MaxLength;
-	Ar << MaxRadius;
-	Ar << BoundingBox;
-	Ar << Flags;
+	Ar << Header.CurveCount;
+	Ar << Header.PointCount;
+	Ar << Header.MaxLength;
+	Ar << Header.MaxRadius;
+	Ar << Header.BoundingBox;
+	Ar << Header.Flags;
 	for (uint8 AttributeIt = 0; AttributeIt < HAIR_CURVE_ATTRIBUTE_COUNT; ++AttributeIt)
 	{
-		Ar << CurveAttributeOffsets[AttributeIt];
+		Ar << Header.CurveAttributeOffsets[AttributeIt];
 	}
 	for (uint8 AttributeIt = 0; AttributeIt < HAIR_POINT_ATTRIBUTE_COUNT; ++AttributeIt)
 	{
-		Ar << PointAttributeOffsets[AttributeIt];
+		Ar << Header.PointAttributeOffsets[AttributeIt];
 	}
-	Ar << ImportedAttributes;
-	Ar << ImportedAttributeFlags;
+	Ar << Header.ImportedAttributes;
+	Ar << Header.ImportedAttributeFlags;
+}
+
+
+void FHairStrandsBulkData::SerializeData(FArchive& Ar, UObject* Owner)
+{
+	static_assert(sizeof(FHairStrandsPositionFormat::BulkType) == sizeof(FHairStrandsPositionFormat::Type));
+	static_assert(sizeof(FHairStrandsAttributeFormat::BulkType) == sizeof(FHairStrandsAttributeFormat::Type));
+	static_assert(sizeof(FHairStrandsPointToCurveFormat16::BulkType) == sizeof(FHairStrandsPointToCurveFormat16::Type));
+	static_assert(sizeof(FHairStrandsPointToCurveFormat32::BulkType) == sizeof(FHairStrandsPointToCurveFormat32::Type));
+	static_assert(sizeof(FHairStrandsRootIndexFormat::BulkType) == sizeof(FHairStrandsRootIndexFormat::Type)); 
 
+	Ar.UsingCustomVersion(FReleaseObjectVersion::GUID);
+	Ar.UsingCustomVersion(FUE5ReleaseStreamObjectVersion::GUID);
+	  
 	// Forced not inline means the bulk data won't automatically be loaded when we deserialize
 	// but only when we explicitly take action to load it
 	if (Ar.IsSaving())
 	{
 		const uint32 BulkFlags = BULKDATA_Force_NOT_InlinePayload;
-		Positions.SetBulkDataFlags(BulkFlags);
-		CurveAttributes.SetBulkDataFlags(BulkFlags);
-		PointAttributes.SetBulkDataFlags(BulkFlags);
-		PointToCurve.SetBulkDataFlags(BulkFlags);
-		Curves.SetBulkDataFlags(BulkFlags);
+		Data.Positions.SetBulkDataFlags(BulkFlags);
+		Data.CurveAttributes.SetBulkDataFlags(BulkFlags);
+		Data.PointAttributes.SetBulkDataFlags(BulkFlags);
+		Data.PointToCurve.SetBulkDataFlags(BulkFlags);
+		Data.Curves.SetBulkDataFlags(BulkFlags);
 	}
 
-	if (!!(Flags & DataFlags_HasData))
+	if (!!(Header.Flags & DataFlags_HasData))
 	{
 		const int32 ChunkIndex = 0;
 		bool bAttemptFileMapping = false;
 
-		Positions.Serialize(Ar, Owner, ChunkIndex, bAttemptFileMapping);
-		PointAttributes.Serialize(Ar, Owner, ChunkIndex, bAttemptFileMapping);
-		CurveAttributes.Serialize(Ar, Owner, ChunkIndex, bAttemptFileMapping);
-		PointToCurve.Serialize(Ar, Owner, ChunkIndex, bAttemptFileMapping);
-		Curves.Serialize(Ar, Owner, ChunkIndex, bAttemptFileMapping);
+		Data.Positions.Serialize(Ar, Owner, ChunkIndex, bAttemptFileMapping);
+		Data.PointAttributes.Serialize(Ar, Owner, ChunkIndex, bAttemptFileMapping);
+		Data.CurveAttributes.Serialize(Ar, Owner, ChunkIndex, bAttemptFileMapping);
+		Data.PointToCurve.Serialize(Ar, Owner, ChunkIndex, bAttemptFileMapping);
+		Data.Curves.Serialize(Ar, Owner, ChunkIndex, bAttemptFileMapping);
 	}
 }
 
+
 void FHairStrandsBulkData::Reset()
 {
-	CurveCount = 0;
-	PointCount = 0;
-	MaxLength = 0;
-	MaxRadius = 0;
-	BoundingBox = FBox(EForceInit::ForceInit);
-	Flags = 0;
+	Header.CurveCount = 0;
+	Header.PointCount = 0;
+	Header.MaxLength = 0;
+	Header.MaxRadius = 0;
+	Header.BoundingBox = FBox(EForceInit::ForceInit);
+	Header.Flags = 0;
 	for (uint8 AttributeIt = 0; AttributeIt < HAIR_CURVE_ATTRIBUTE_COUNT; ++AttributeIt)
 	{
-		CurveAttributeOffsets[AttributeIt] = 0xFFFFFFFF;
+		Header.CurveAttributeOffsets[AttributeIt] = 0xFFFFFFFF;
 	}
 	for (uint8 AttributeIt = 0; AttributeIt < HAIR_POINT_ATTRIBUTE_COUNT; ++AttributeIt)
 	{
-		PointAttributeOffsets[AttributeIt] = 0xFFFFFFFF;
+		Header.PointAttributeOffsets[AttributeIt] = 0xFFFFFFFF;
 	}
 	// Deallocate memory if needed
-	Positions.RemoveBulkData();
-	CurveAttributes.RemoveBulkData();
-	PointAttributes.RemoveBulkData();
-	PointToCurve.RemoveBulkData();
-	Curves.RemoveBulkData();
+	Data.Positions.RemoveBulkData();
+	Data.CurveAttributes.RemoveBulkData();
+	Data.PointAttributes.RemoveBulkData();
+	Data.PointToCurve.RemoveBulkData();
+	Data.Curves.RemoveBulkData();
 
 	// Reset the bulk byte buffer to ensure the (serialize) data size is reset to 0
-	Positions 		= FByteBulkData();
-	CurveAttributes	= FByteBulkData();
-	PointAttributes	= FByteBulkData();
-	PointToCurve	= FByteBulkData();
-	Curves			= FByteBulkData();
+	Data.Positions 		= FByteBulkData();
+	Data.CurveAttributes= FByteBulkData();
+	Data.PointAttributes= FByteBulkData();
+	Data.PointToCurve	= FByteBulkData();
+	Data.Curves			= FByteBulkData();
 }
 
 void FHairStrandsDatas::Reset()
diff --git a/Engine/Plugins/Runtime/HairStrands/Source/HairStrandsCore/Private/HairStrandsRendering.cpp b/Engine/Plugins/Runtime/HairStrands/Source/HairStrandsCore/Private/HairStrandsRendering.cpp
--- a/Engine/Plugins/Runtime/HairStrands/Source/HairStrandsCore/Private/HairStrandsRendering.cpp
+++ b/Engine/Plugins/Runtime/HairStrands/Source/HairStrandsCore/Private/HairStrandsRendering.cpp
@@ -914,7 +914,7 @@ static void AddHairClusterAABBPass(
 	FRDGBufferSRVRef& DrawIndirectRasterComputeBuffer)
 {
 	// Clusters AABB are only update if the groom is deformed.
-	const FBoxSphereBounds& Bounds = Instance->Strands.Data->BoundingBox;
+	const FBoxSphereBounds& Bounds = Instance->Strands.Data->Header.BoundingBox;
 	FTransform InRenLocalToTranslatedWorld = Instance->LocalToWorld;
 	InRenLocalToTranslatedWorld.AddToTranslation(TranslatedWorldOffset);
 	const FBoxSphereBounds TransformedBounds = Bounds.TransformBy(InRenLocalToTranslatedWorld);
@@ -1733,13 +1733,13 @@ static FHairGroupPublicData::FVertexFactoryInput InternalComputeHairStrandsVerte
 
 	if (bUseGuideAttributeOffsets)
 	{
-		PACK_CURVE_HAIR_ATTRIBUTE_OFFSETS(OutVFInput.Strands.Common.CurveAttributeOffsets, Instance->Guides.Data->CurveAttributeOffsets);
-		PACK_POINT_HAIR_ATTRIBUTE_OFFSETS(OutVFInput.Strands.Common.PointAttributeOffsets, Instance->Guides.Data->PointAttributeOffsets);
+		PACK_CURVE_HAIR_ATTRIBUTE_OFFSETS(OutVFInput.Strands.Common.CurveAttributeOffsets, Instance->Guides.Data->Header.CurveAttributeOffsets);
+		PACK_POINT_HAIR_ATTRIBUTE_OFFSETS(OutVFInput.Strands.Common.PointAttributeOffsets, Instance->Guides.Data->Header.PointAttributeOffsets);
 	}
 	else
 	{
-		PACK_CURVE_HAIR_ATTRIBUTE_OFFSETS(OutVFInput.Strands.Common.CurveAttributeOffsets, Instance->Strands.Data->CurveAttributeOffsets);
-		PACK_POINT_HAIR_ATTRIBUTE_OFFSETS(OutVFInput.Strands.Common.PointAttributeOffsets, Instance->Strands.Data->PointAttributeOffsets);
+		PACK_CURVE_HAIR_ATTRIBUTE_OFFSETS(OutVFInput.Strands.Common.CurveAttributeOffsets, Instance->Strands.Data->Header.CurveAttributeOffsets);
+		PACK_POINT_HAIR_ATTRIBUTE_OFFSETS(OutVFInput.Strands.Common.PointAttributeOffsets, Instance->Strands.Data->Header.PointAttributeOffsets);
 	}
 
 #if RHI_RAYTRACING
@@ -2137,7 +2137,7 @@ void ComputeHairStrandsInterpolation(
 				// Special case for debug mode were the attribute buffer is patch with some custom data to show hair properties (strands belonging to the same cluster, ...)
 				if (Instance->Strands.DebugCurveAttributeBuffer.Buffer == nullptr)
 				{
-					CreateHairStrandsDebugAttributeBuffer(GraphBuilder, &Instance->Strands.DebugCurveAttributeBuffer, Instance->Strands.Data->CurveAttributes.GetBulkDataSize(), FName(Instance->Debug.MeshComponentName));
+					CreateHairStrandsDebugAttributeBuffer(GraphBuilder, &Instance->Strands.DebugCurveAttributeBuffer, Instance->Strands.Data->Data.CurveAttributes.GetBulkDataSize(), FName(Instance->Debug.MeshComponentName));
 				}
 				FRDGImportedBuffer OutRenCurveAttributeBuffer = Register(GraphBuilder, Instance->Strands.DebugCurveAttributeBuffer, ERDGImportedBufferFlags::CreateUAV);
 
@@ -2148,7 +2148,7 @@ void ComputeHairStrandsInterpolation(
 					PatchMode,
 					bValidGuide,
 					bUseSingleGuide,
-					Instance->Strands.Data->CurveAttributeOffsets,
+					Instance->Strands.Data->Header.CurveAttributeOffsets,
 					Register(GraphBuilder, Instance->Strands.RestResource->CurveAttributeBuffer, ERDGImportedBufferFlags::CreateSRV).Buffer,
 					RegisterAsSRV(GraphBuilder, Instance->Strands.RestResource->CurveBuffer),
 					RegisterAsSRV(GraphBuilder, Instance->Strands.ClusterCullingResource->VertexToClusterIdBuffer),
diff --git a/Engine/Plugins/Runtime/HairStrands/Source/HairStrandsCore/Private/Niagara/NiagaraDataInterfaceHairStrands.cpp b/Engine/Plugins/Runtime/HairStrands/Source/HairStrandsCore/Private/Niagara/NiagaraDataInterfaceHairStrands.cpp
--- a/Engine/Plugins/Runtime/HairStrands/Source/HairStrandsCore/Private/Niagara/NiagaraDataInterfaceHairStrands.cpp
+++ b/Engine/Plugins/Runtime/HairStrands/Source/HairStrandsCore/Private/Niagara/NiagaraDataInterfaceHairStrands.cpp
@@ -443,7 +443,7 @@ void FNDIHairStrandsData::Update(UNiagaraDataInterfaceHairStrands* Interface, FN
 				ParamsScale[32 * ThicknessOffset + i] = HairPhysics.StrandsParameters.ThicknessScale.GetRichCurve()->Eval(VertexCoord);
 			}
 
-			const FBox& StrandsBox = StrandsDatas->BoundingBox;
+			const FBox& StrandsBox = StrandsDatas->Header.BoundingBox;
 
 			NumStrands = StrandsDatas->GetNumCurves();
 			LocalSimulation = false;
diff --git a/Engine/Plugins/Runtime/HairStrands/Source/HairStrandsCore/Public/GroomAsset.h b/Engine/Plugins/Runtime/HairStrands/Source/HairStrandsCore/Public/GroomAsset.h
--- a/Engine/Plugins/Runtime/HairStrands/Source/HairStrandsCore/Public/GroomAsset.h
+++ b/Engine/Plugins/Runtime/HairStrands/Source/HairStrandsCore/Public/GroomAsset.h
@@ -114,11 +114,11 @@ struct HAIRSTRANDSCORE_API FHairGroupPlatformData
 		uint32 GetDataSize() const
 		{
 			uint32 Total = 0;
-			Total += BulkData.Positions.IsBulkDataLoaded() 		? BulkData.Positions.GetBulkDataSize()   : 0;
-			Total += BulkData.CurveAttributes.IsBulkDataLoaded()? BulkData.CurveAttributes.GetBulkDataSize() : 0;
-			Total += BulkData.PointAttributes.IsBulkDataLoaded()? BulkData.PointAttributes.GetBulkDataSize() : 0;
-			Total += BulkData.Curves.IsBulkDataLoaded() 		? BulkData.Curves.GetBulkDataSize(): 0;
-			Total += BulkData.PointToCurve.IsBulkDataLoaded()	? BulkData.PointToCurve.GetBulkDataSize() : 0;
+			Total += BulkData.Data.Positions.IsBulkDataLoaded() 		? BulkData.Data.Positions.GetBulkDataSize()   : 0;
+			Total += BulkData.Data.CurveAttributes.IsBulkDataLoaded()	? BulkData.Data.CurveAttributes.GetBulkDataSize() : 0;
+			Total += BulkData.Data.PointAttributes.IsBulkDataLoaded()	? BulkData.Data.PointAttributes.GetBulkDataSize() : 0;
+			Total += BulkData.Data.Curves.IsBulkDataLoaded() 			? BulkData.Data.Curves.GetBulkDataSize(): 0;
+			Total += BulkData.Data.PointToCurve.IsBulkDataLoaded()		? BulkData.Data.PointToCurve.GetBulkDataSize() : 0;
 			return Total;
 		}
 
@@ -132,10 +132,10 @@ struct HAIRSTRANDSCORE_API FHairGroupPlatformData
 		{
 			uint32 Total = 0;
 			Total += FBase::GetDataSize();
-			Total += InterpolationBulkData.Interpolation.IsBulkDataLoaded()		? InterpolationBulkData.Interpolation.GetBulkDataSize() : 0;
-			Total += InterpolationBulkData.Interpolation0.IsBulkDataLoaded()	? InterpolationBulkData.Interpolation0.GetBulkDataSize() : 0;
-			Total += InterpolationBulkData.Interpolation1.IsBulkDataLoaded()	? InterpolationBulkData.Interpolation1.GetBulkDataSize() : 0;
-			Total += InterpolationBulkData.SimRootPointIndex.IsBulkDataLoaded()	? InterpolationBulkData.SimRootPointIndex.GetBulkDataSize() : 0;
+			Total += InterpolationBulkData.Data.Interpolation.IsBulkDataLoaded()	? InterpolationBulkData.Data.Interpolation.GetBulkDataSize() : 0;
+			Total += InterpolationBulkData.Data.Interpolation0.IsBulkDataLoaded()	? InterpolationBulkData.Data.Interpolation0.GetBulkDataSize() : 0;
+			Total += InterpolationBulkData.Data.Interpolation1.IsBulkDataLoaded()	? InterpolationBulkData.Data.Interpolation1.GetBulkDataSize() : 0;
+			Total += InterpolationBulkData.Data.SimRootPointIndex.IsBulkDataLoaded()? InterpolationBulkData.Data.SimRootPointIndex.GetBulkDataSize() : 0;
 			return Total;
 		}
 
@@ -186,12 +186,12 @@ struct HAIRSTRANDSCORE_API FHairGroupPlatformData
 			uint32 Total = 0;
 			Total += FBaseWithInterpolation::GetDataSize();
 
-			Total += ClusterCullingBulkData.LODVisibility.GetAllocatedSize();
-			Total += ClusterCullingBulkData.CPULODScreenSize.GetAllocatedSize();
-			Total += ClusterCullingBulkData.ClusterInfos.IsBulkDataLoaded()			? ClusterCullingBulkData.ClusterInfos.GetBulkDataSize() : 0;
-			Total += ClusterCullingBulkData.ClusterLODInfos.IsBulkDataLoaded()		? ClusterCullingBulkData.ClusterLODInfos.GetBulkDataSize() : 0;
-			Total += ClusterCullingBulkData.VertexToClusterIds.IsBulkDataLoaded()	? ClusterCullingBulkData.VertexToClusterIds.GetBulkDataSize() : 0;
-			Total += ClusterCullingBulkData.ClusterVertexIds.IsBulkDataLoaded()		? ClusterCullingBulkData.ClusterVertexIds.GetBulkDataSize() : 0;
+			Total += ClusterCullingBulkData.Header.LODVisibility.GetAllocatedSize();
+			Total += ClusterCullingBulkData.Header.CPULODScreenSize.GetAllocatedSize();
+			Total += ClusterCullingBulkData.Data.ClusterInfos.IsBulkDataLoaded()		? ClusterCullingBulkData.Data.ClusterInfos.GetBulkDataSize() : 0;
+			Total += ClusterCullingBulkData.Data.ClusterLODInfos.IsBulkDataLoaded()		? ClusterCullingBulkData.Data.ClusterLODInfos.GetBulkDataSize() : 0;
+			Total += ClusterCullingBulkData.Data.VertexToClusterIds.IsBulkDataLoaded()	? ClusterCullingBulkData.Data.VertexToClusterIds.GetBulkDataSize() : 0;
+			Total += ClusterCullingBulkData.Data.ClusterVertexIds.IsBulkDataLoaded()	? ClusterCullingBulkData.Data.ClusterVertexIds.GetBulkDataSize() : 0;
 			return Total;
 		}
 #endif
diff --git a/Engine/Plugins/Runtime/HairStrands/Source/HairStrandsCore/Public/HairStrandsDatas.h b/Engine/Plugins/Runtime/HairStrands/Source/HairStrandsCore/Public/HairStrandsDatas.h
--- a/Engine/Plugins/Runtime/HairStrands/Source/HairStrandsCore/Public/HairStrandsDatas.h
+++ b/Engine/Plugins/Runtime/HairStrands/Source/HairStrandsCore/Public/HairStrandsDatas.h
@@ -287,7 +287,19 @@ struct FHairStrandsRootUtils
 	static float	 PackUVsToFloat(const FVector2f& UV);
 };
 
-struct FHairStrandsInterpolationBulkData;
+
+struct HAIRSTRANDSCORE_API FHairStrandsBulkCommon
+{
+	virtual ~FHairStrandsBulkCommon() { }
+	virtual void SerializeHeader(FArchive& Ar, UObject* Owner) = 0;
+	virtual void SerializeData(FArchive& Ar, UObject* Owner) = 0;
+
+#if WITH_EDITORONLY_DATA
+	// Transient Name/DDCkey for streaming
+	FString Name;
+	FString DerivedDataKey;
+#endif
+};
 
 /** Hair strands points interpolation attributes */
 struct HAIRSTRANDSCORE_API FHairStrandsInterpolationDatas
@@ -319,7 +331,7 @@ struct HAIRSTRANDSCORE_API FHairStrandsInterpolationDatas
 	bool bUseUniqueGuide = false;
 };
 
-struct HAIRSTRANDSCORE_API FHairStrandsInterpolationBulkData
+struct HAIRSTRANDSCORE_API FHairStrandsInterpolationBulkData : FHairStrandsBulkCommon
 {
 	enum EDataFlags
 	{
@@ -329,15 +341,25 @@ struct HAIRSTRANDSCORE_API FHairStrandsInterpolationBulkData
 
 	void Reset();
 	void Serialize(FArchive& Ar, UObject* Owner);
-	uint32 GetPointCount() const { return PointCount; };
+	virtual void SerializeHeader(FArchive& Ar, UObject* Owner) override;
+	virtual void SerializeData(FArchive& Ar, UObject* Owner) override;
+	uint32 GetPointCount() const { return Header.PointCount; };
 
-	uint32 Flags = 0;
-	uint32 PointCount = 0;
-	uint32 SimPointCount = 0;
-	FByteBulkData Interpolation;	// FHairStrandsInterpolationFormat  - Per-rendering-vertex interpolation data (closest guides, weight factors, ...). Data for a single guide
-	FByteBulkData Interpolation0;	// FHairStrandsInterpolation0Format - Per-rendering-vertex interpolation data (closest guides, weight factors, ...). Data for up to 3 guides
-	FByteBulkData Interpolation1;	// FHairStrandsInterpolation1Format - Per-rendering-vertex interpolation data (closest guides, weight factors, ...). Data for up to 3 guides
-	FByteBulkData SimRootPointIndex;// FHairStrandsRootIndexFormat      - Per-rendering-vertex index of the sim-root vertex
+	struct FHeader
+	{
+		uint32 Flags = 0;
+		uint32 PointCount = 0;
+		uint32 SimPointCount = 0;
+	} Header;
+
+	struct FData
+	{
+
+		FByteBulkData Interpolation;	// FHairStrandsInterpolationFormat  - Per-rendering-vertex interpolation data (closest guides, weight factors, ...). Data for a single guide
+		FByteBulkData Interpolation0;	// FHairStrandsInterpolation0Format - Per-rendering-vertex interpolation data (closest guides, weight factors, ...). Data for up to 3 guides
+		FByteBulkData Interpolation1;	// FHairStrandsInterpolation1Format - Per-rendering-vertex interpolation data (closest guides, weight factors, ...). Data for up to 3 guides
+		FByteBulkData SimRootPointIndex;// FHairStrandsRootIndexFormat      - Per-rendering-vertex index of the sim-root vertex
+	} Data;
 };
 
 /** Hair strands points attribute */
@@ -457,7 +479,7 @@ struct HAIRSTRANDSCORE_API FHairStrandsDatas
 float GetHairStrandsMaxLength(const FHairStrandsDatas& In);
 float GetHairStrandsMaxRadius(const FHairStrandsDatas& In);
 
-struct HAIRSTRANDSCORE_API FHairStrandsBulkData
+struct HAIRSTRANDSCORE_API FHairStrandsBulkData : FHairStrandsBulkCommon
 {
 	enum EDataFlags
 	{
@@ -466,35 +488,43 @@ struct HAIRSTRANDSCORE_API FHairStrandsBulkData
 	};
 
 	void Serialize(FArchive& Ar, UObject* Owner);
+	virtual void SerializeHeader(FArchive& Ar, UObject* Owner) override;
+	virtual void SerializeData(FArchive& Ar, UObject* Owner) override;
 
-	bool IsValid() const { return CurveCount > 0 && PointCount > 0; }
+	bool IsValid() const { return Header.CurveCount > 0 && Header.PointCount > 0; }
 	void Reset();
 
-	uint32 GetNumCurves() const { return CurveCount;  };
-	uint32 GetNumPoints() const { return PointCount; };
-	float  GetMaxLength() const	{ return MaxLength; };
-	float  GetMaxRadius() const { return MaxRadius; }
-	FVector GetPositionOffset() const { return BoundingBox.GetCenter(); }
-	const FBox& GetBounds() const { return BoundingBox; }
-
-	uint32 CurveCount = 0;
-	uint32 PointCount = 0;
-	float MaxLength = 0;
-	float MaxRadius = 0;
-	FBox BoundingBox = FBox(EForceInit::ForceInit);
-	uint32 Flags = 0;
-	uint32 CurveAttributeOffsets[HAIR_CURVE_ATTRIBUTE_COUNT] = {0};
-	uint32 PointAttributeOffsets[HAIR_POINT_ATTRIBUTE_COUNT] = {0};
-
-	/** Imported attribute info */
-	uint32 ImportedAttributes = 0;
-	uint32 ImportedAttributeFlags = 0;
+	uint32 GetNumCurves() const { return Header.CurveCount;  };
+	uint32 GetNumPoints() const { return Header.PointCount; };
+	float  GetMaxLength() const	{ return Header.MaxLength; };
+	float  GetMaxRadius() const { return Header.MaxRadius; }
+	FVector GetPositionOffset() const { return Header.BoundingBox.GetCenter(); }
+	const FBox& GetBounds() const { return Header.BoundingBox; }
 
-	FByteBulkData Positions;		// Size = PointCount
-	FByteBulkData CurveAttributes;	// Size = y*CurveCount (depends on the per-curve stored attributes)
-	FByteBulkData PointAttributes;	// Size = x*PointCount (depends on the per-point stored attributes)
-	FByteBulkData PointToCurve; 	// Size = PointCount
-	FByteBulkData Curves;			// Size = CurveCount
+	struct FHeader
+	{
+		uint32 CurveCount = 0;
+		uint32 PointCount = 0;
+		float  MaxLength = 0;
+		float  MaxRadius = 0;
+		FBox   BoundingBox = FBox(EForceInit::ForceInit);
+		uint32 Flags = 0;
+		uint32 CurveAttributeOffsets[HAIR_CURVE_ATTRIBUTE_COUNT] = {0};
+		uint32 PointAttributeOffsets[HAIR_POINT_ATTRIBUTE_COUNT] = {0};
+	
+		/** Imported attribute info */
+		uint32 ImportedAttributes = 0;
+		uint32 ImportedAttributeFlags = 0;
+	} Header;
+
+	struct FData
+	{
+		FByteBulkData Positions;		// Size = PointCount
+		FByteBulkData CurveAttributes;	// Size = y*CurveCount (depends on the per-curve stored attributes)
+		FByteBulkData PointAttributes;	// Size = x*PointCount (depends on the per-point stored attributes)
+		FByteBulkData PointToCurve; 	// Size = PointCount
+		FByteBulkData Curves;			// Size = CurveCount
+	} Data;
 };
 
 /** Hair strands debug data */
@@ -769,32 +799,42 @@ struct HAIRSTRANDSCORE_API FHairStrandsClusterCullingData
 	uint32 VertexCount = 0;
 };
 
-struct HAIRSTRANDSCORE_API FHairStrandsClusterCullingBulkData
+struct HAIRSTRANDSCORE_API FHairStrandsClusterCullingBulkData : FHairStrandsBulkCommon
 {
 	FHairStrandsClusterCullingBulkData();
 	void Reset();
 	void Serialize(FArchive& Ar, UObject* Owner);
-	bool IsValid() const { return ClusterCount > 0 && VertexCount > 0; }
+	virtual void SerializeHeader(FArchive& Ar, UObject* Owner) override;
+	virtual void SerializeData(FArchive& Ar, UObject* Owner) override;
+	bool IsValid() const { return Header.ClusterCount > 0 && Header.VertexCount > 0; }
 	void Validate(bool bIsSaving);
 
-	/* Set LOD visibility, allowing to remove the simulation/rendering of certain LOD */
-	TArray<bool>	LODVisibility;
-
-	/* Screen size at which LOD should switches on CPU */
-	TArray<float>	CPULODScreenSize;
-
-	/* Curve count and Point count per LOD */
-	TArray<FHairLODInfo> LODInfos;
+	struct FHeader
+	{
+		/* Set LOD visibility, allowing to remove the simulation/rendering of certain LOD */
+		TArray<bool> LODVisibility;
+	
+		/* Screen size at which LOD should switches on CPU */
+		TArray<float> CPULODScreenSize;
+	
+		/* Curve count and Point count per LOD */
+		TArray<FHairLODInfo> LODInfos;
+	
+		uint32 ClusterCount = 0;
+		uint32 ClusterLODCount = 0;
+		uint32 VertexCount = 0;
+		uint32 VertexLODCount = 0;
+	} Header;
+
+	struct FData
+	{
 
-	/* LOD info for the various clusters for LOD management on GPU */
-	FByteBulkData	PackedClusterInfos;		// Size - ClusterCount
-	FByteBulkData	ClusterLODInfos;		// Size - ClusterLODCount
-	FByteBulkData	VertexToClusterIds;		// Size - VertexCount
-	FByteBulkData	ClusterVertexIds;		// Size - VertexLODCount
+		/* LOD info for the various clusters for LOD management on GPU */
+		FByteBulkData	PackedClusterInfos;		// Size - ClusterCount
+		FByteBulkData	ClusterLODInfos;		// Size - ClusterLODCount
+		FByteBulkData	VertexToClusterIds;		// Size - VertexCount
+		FByteBulkData	ClusterVertexIds;		// Size - VertexLODCount
+	} Data;
 
-	uint32 ClusterCount = 0;
-	uint32 ClusterLODCount = 0;
-	uint32 VertexCount = 0;
-	uint32 VertexLODCount = 0;
 };
 
diff --git a/Engine/Plugins/Runtime/HairStrands/Source/HairStrandsEditor/Private/GroomAssetDetails.cpp b/Engine/Plugins/Runtime/HairStrands/Source/HairStrandsEditor/Private/GroomAssetDetails.cpp
--- a/Engine/Plugins/Runtime/HairStrands/Source/HairStrandsEditor/Private/GroomAssetDetails.cpp
+++ b/Engine/Plugins/Runtime/HairStrands/Source/HairStrandsEditor/Private/GroomAssetDetails.cpp
@@ -1301,9 +1301,9 @@ void FGroomRenderingDetails::OnGenerateElementForLODs(TSharedRef<IPropertyHandle
 
 	// LOD Stats
 	const FHairStrandsClusterCullingBulkData& ClusterCullingBulkData = GroomAsset->HairGroupsPlatformData[GroupIndex].Strands.ClusterCullingBulkData;
-	if (ClusterCullingBulkData.IsValid() && LODIndex < ClusterCullingBulkData.LODInfos.Num())
+	if (ClusterCullingBulkData.IsValid() && LODIndex < ClusterCullingBulkData.Header.LODInfos.Num())
 	{
-		const FHairLODInfo& LODInfo = ClusterCullingBulkData.LODInfos[LODIndex];
+		const FHairLODInfo& LODInfo = ClusterCullingBulkData.Header.LODInfos[LODIndex];
 
 		ChildrenBuilder.AddCustomRow(LOCTEXT("HairStrandsLODInfo_Array", "HairStrandsLODInfo"))
 		.ValueContent()
@@ -1480,13 +1480,13 @@ void FGroomRenderingDetails::OnGenerateElementForHairGroup(TSharedRef<IPropertyH
 		.ValueContent()
 		.HAlign(HAlign_Fill)
 		[
-			MakeHairStrandsInfoGrid(DetailFontInfo, GroomAsset->HairGroupsInfo[GroupIndex], GroomAsset->HairGroupsPlatformData[GroupIndex].Strands.BulkData.MaxRadius)
+			MakeHairStrandsInfoGrid(DetailFontInfo, GroomAsset->HairGroupsInfo[GroupIndex], GroomAsset->HairGroupsPlatformData[GroupIndex].Strands.BulkData.Header.MaxRadius)
 		];
 		ChildrenBuilder.AddCustomRow(LOCTEXT("HairStrandsAttributeInfo_Array", "HairStrandsAttributeInfo"))
 		.ValueContent()
 		.HAlign(HAlign_Fill)
 		[
-			MakeHairStrandsAttributeInfoGrid(DetailFontInfo, GroomAsset->HairGroupsPlatformData[GroupIndex].Strands.BulkData.ImportedAttributes, GroomAsset->HairGroupsPlatformData[GroupIndex].Strands.BulkData.ImportedAttributeFlags)
+			MakeHairStrandsAttributeInfoGrid(DetailFontInfo, GroomAsset->HairGroupsPlatformData[GroupIndex].Strands.BulkData.Header.ImportedAttributes, GroomAsset->HairGroupsPlatformData[GroupIndex].Strands.BulkData.Header.ImportedAttributeFlags)
 		];
 	}
 
