Fix invalid serialization of empty point attribute bulk data when a groom does not have any optional point attribute data.

diff --git a/Engine/Plugins/Runtime/HairStrands/Source/HairStrandsCore/Private/GroomBuilder.cpp b/Engine/Plugins/Runtime/HairStrands/Source/HairStrandsCore/Private/GroomBuilder.cpp
--- a/Engine/Plugins/Runtime/HairStrands/Source/HairStrandsCore/Private/GroomBuilder.cpp
+++ b/Engine/Plugins/Runtime/HairStrands/Source/HairStrandsCore/Private/GroomBuilder.cpp
@@ -40,7 +40,7 @@ static FAutoConsoleVariableRef CVarHairGroupIndexBuilder_MaxVoxelResolution(TEXT
 
 FString FGroomBuilder::GetVersion()
 {
-	return TEXT("v8o");
+	return TEXT("v8p1");
 }
 
 namespace FHairStrandsDecimation
@@ -434,7 +434,13 @@ namespace HairStrandsBuilder
 				OutBulkData.Header.PointAttributeOffsets[HAIR_POINT_ATTRIBUTE_AO] = OutPackedAttributes.Num() * UintToByte;
 				OutPackedAttributes.Append(AttributeAO);
 			}
-			CopyToBulkData<FHairStrandsAttributeFormat>(OutBulkData.Data.PointAttributes, OutPackedAttributes);
+
+			const bool bHasPointAttribute = OutPackedAttributes.Num() > 0;
+			if (bHasPointAttribute)
+			{
+				OutBulkData.Header.Flags |= FHairStrandsBulkData::DataFlags_HasPointAttribute;
+				CopyToBulkData<FHairStrandsAttributeFormat>(OutBulkData.Data.PointAttributes, OutPackedAttributes);
+			}
 		}
 
 		CopyToBulkData<FHairStrandsPositionFormat>(OutBulkData.Data.Positions, OutPackedPositions);
diff --git a/Engine/Plugins/Runtime/HairStrands/Source/HairStrandsCore/Private/GroomResources.cpp b/Engine/Plugins/Runtime/HairStrands/Source/HairStrandsCore/Private/GroomResources.cpp
--- a/Engine/Plugins/Runtime/HairStrands/Source/HairStrandsCore/Private/GroomResources.cpp
+++ b/Engine/Plugins/Runtime/HairStrands/Source/HairStrandsCore/Private/GroomResources.cpp
@@ -297,7 +297,7 @@ void InternalCreateByteAddressBufferRDG_FromBulkData(FRDGBuilder& GraphBuilder,
 
 
 template<typename FormatType>
-void InternalCreateVertexBufferRDG(FRDGBuilder& GraphBuilder, const typename FormatType::Type* InData, uint32 InDataCount, FRDGExternalBuffer& Out, const TCHAR* DebugName, const FName& OwnerName, EHairResourceUsageType UsageType, ERDGInitialDataFlags InitialDataFlags)
+void InternalCreateVertexBufferRDG(FRDGBuilder& GraphBuilder, const typename FormatType::Type* InData, uint32 InDataCount, FRDGExternalBuffer& Out, const TCHAR* DebugName, const FName& OwnerName, EHairResourceUsageType UsageType, ERDGInitialDataFlags InitialDataFlags, bool bIsByteAddressBuffer=false)
 {
 	FRDGBufferRef Buffer = nullptr;
 
@@ -320,13 +320,14 @@ void InternalCreateVertexBufferRDG(FRDGBuilder& GraphBuilder, const typename For
 		InData,
 		DataSizeInBytes,
 		InitialDataFlags,
-		OwnerName);
+		OwnerName, 
+		bIsByteAddressBuffer);
 
 	ConvertToExternalBufferWithViews(GraphBuilder, Buffer, Out, FormatType::Format);
 }
 
 template<typename FormatType>
-void InternalCreateVertexBufferRDG(FRDGBuilder& GraphBuilder, const TArray<typename FormatType::Type>& InData, FRDGExternalBuffer& Out, const TCHAR* DebugName, const FName& OwnerName, EHairResourceUsageType UsageType, ERDGInitialDataFlags InitialDataFlags= ERDGInitialDataFlags::NoCopy)
+void InternalCreateVertexBufferRDG(FRDGBuilder& GraphBuilder, const TArray<typename FormatType::Type>& InData, FRDGExternalBuffer& Out, const TCHAR* DebugName, const FName& OwnerName, EHairResourceUsageType UsageType, ERDGInitialDataFlags InitialDataFlags= ERDGInitialDataFlags::NoCopy, bool bIsByteAddressBuffer=false)
 {
 	FRDGBufferRef Buffer = nullptr;
 
@@ -350,13 +351,14 @@ void InternalCreateVertexBufferRDG(FRDGBuilder& GraphBuilder, const TArray<typen
 		InData.GetData(),
 		DataSizeInBytes,
 		InitialDataFlags,
-		OwnerName);
+		OwnerName,
+		bIsByteAddressBuffer);
 
 	ConvertToExternalBufferWithViews(GraphBuilder, Buffer, Out, FormatType::Format);
 }
 
 template<typename DataType>
-void InternalCreateVertexBufferRDG(FRDGBuilder& GraphBuilder, const TArray<DataType>& InData, EPixelFormat Format, FRDGExternalBuffer& Out, const TCHAR* DebugName, const FName& OwnerName, EHairResourceUsageType UsageType)
+void InternalCreateVertexBufferRDG(FRDGBuilder& GraphBuilder, const TArray<DataType>& InData, EPixelFormat Format, FRDGExternalBuffer& Out, const TCHAR* DebugName, const FName& OwnerName, EHairResourceUsageType UsageType, bool bIsByteAddressBuffer=false)
 {
 	FRDGBufferRef Buffer = nullptr;
 
@@ -380,7 +382,8 @@ void InternalCreateVertexBufferRDG(FRDGBuilder& GraphBuilder, const TArray<DataT
 		InData.GetData(),
 		DataSizeInBytes,
 		ERDGInitialDataFlags::NoCopy,
-		OwnerName);
+		OwnerName,
+		bIsByteAddressBuffer);
 
 	ConvertToExternalBufferWithViews(GraphBuilder, Buffer, Out, Format);
 }
@@ -875,7 +878,10 @@ bool FHairStrandsRestResource::InternalIsDataLoaded()
 		FBulkDataBatchRequest::FBatchBuilder Batch = FBulkDataBatchRequest::NewBatch(5);
 		Batch.Read(BulkData.Data.Positions);
 		Batch.Read(BulkData.Data.CurveAttributes);
-		Batch.Read(BulkData.Data.PointAttributes);
+		if (BulkData.Header.Flags & FHairStrandsBulkData::DataFlags_HasPointAttribute)
+		{
+			Batch.Read(BulkData.Data.PointAttributes);
+		}
 		Batch.Read(BulkData.Data.PointToCurve);
 		Batch.Read(BulkData.Data.Curves);
 
@@ -897,7 +903,17 @@ void FHairStrandsRestResource::InternalAllocate(FRDGBuilder& GraphBuilder)
 	// 2. A local copy is done by the buffer uploader. This copy is discarded once the uploading is done.
 	InternalCreateVertexBufferRDG_FromBulkData<FHairStrandsPositionFormat>(GraphBuilder, BulkData.Data.Positions, PointCount, PositionBuffer, ToHairResourceDebugName(HAIRSTRANDS_RESOUCE_NAME(CurveType, Hair.StrandsRest_PositionBuffer), ResourceName), OwnerName, EHairResourceUsageType::Static);
 	InternalCreateByteAddressBufferRDG_FromBulkData(GraphBuilder, BulkData.Data.CurveAttributes, CurveAttributeBuffer, ToHairResourceDebugName(HAIRSTRANDS_RESOUCE_NAME(CurveType, Hair.StrandsRest_CurveAttributeBuffer), ResourceName), OwnerName, EHairResourceUsageType::Static);
-	InternalCreateByteAddressBufferRDG_FromBulkData(GraphBuilder, BulkData.Data.PointAttributes, PointAttributeBuffer, ToHairResourceDebugName(HAIRSTRANDS_RESOUCE_NAME(CurveType, Hair.StrandsRest_PointAttributeBuffer), ResourceName), OwnerName, EHairResourceUsageType::Static);
+	if (!!(BulkData.Header.Flags & FHairStrandsBulkData::DataFlags_HasPointAttribute))
+	{
+		InternalCreateByteAddressBufferRDG_FromBulkData(GraphBuilder, BulkData.Data.PointAttributes, PointAttributeBuffer, ToHairResourceDebugName(HAIRSTRANDS_RESOUCE_NAME(CurveType, Hair.StrandsRest_PointAttributeBuffer), ResourceName), OwnerName, EHairResourceUsageType::Static);
+	}
+	else
+	{
+		TArray<uint32> DummyAttribute;
+		DummyAttribute.Add(0u);
+		InternalCreateVertexBufferRDG(GraphBuilder, DummyAttribute, EPixelFormat::PF_R32_UINT, PointAttributeBuffer, ToHairResourceDebugName(HAIRSTRANDS_RESOUCE_NAME(CurveType, Hair.StrandsRest_PointAttributeBuffer), ResourceName), OwnerName, EHairResourceUsageType::Static, true /*bIsByteAddressBuffer*/);
+		GraphBuilder.UseExternalAccessMode(Register(GraphBuilder, PointAttributeBuffer, ERDGImportedBufferFlags::CreateSRV).Buffer, ERHIAccess::SRVMask);
+	}
 	InternalCreateVertexBufferRDG_FromBulkData<FHairStrandsCurveFormat>(GraphBuilder, BulkData.Data.Curves, CurveCount, CurveBuffer, ToHairResourceDebugName(HAIRSTRANDS_RESOUCE_NAME(CurveType, Hair.StrandsRest_CurveBuffer), ResourceName), OwnerName, EHairResourceUsageType::Static);
 	if (!!(BulkData.Header.Flags & FHairStrandsBulkData::DataFlags_Has16bitsCurveIndex))
 	{
diff --git a/Engine/Plugins/Runtime/HairStrands/Source/HairStrandsCore/Private/HairStrandsDatas.cpp b/Engine/Plugins/Runtime/HairStrands/Source/HairStrandsCore/Private/HairStrandsDatas.cpp
--- a/Engine/Plugins/Runtime/HairStrands/Source/HairStrandsCore/Private/HairStrandsDatas.cpp
+++ b/Engine/Plugins/Runtime/HairStrands/Source/HairStrandsCore/Private/HairStrandsDatas.cpp
@@ -273,7 +273,10 @@ void FHairStrandsBulkData::SerializeData(FArchive& Ar, UObject* Owner)
 		const uint32 BulkFlags = BULKDATA_Force_NOT_InlinePayload;
 		Data.Positions.SetBulkDataFlags(BulkFlags);
 		Data.CurveAttributes.SetBulkDataFlags(BulkFlags);
-		Data.PointAttributes.SetBulkDataFlags(BulkFlags);
+		if (Header.Flags & DataFlags_HasPointAttribute)
+		{
+			Data.PointAttributes.SetBulkDataFlags(BulkFlags);
+		}
 		Data.PointToCurve.SetBulkDataFlags(BulkFlags);
 		Data.Curves.SetBulkDataFlags(BulkFlags);
 	}
@@ -284,7 +287,10 @@ void FHairStrandsBulkData::SerializeData(FArchive& Ar, UObject* Owner)
 		bool bAttemptFileMapping = false;
 
 		Data.Positions.Serialize(Ar, Owner, ChunkIndex, bAttemptFileMapping);
-		Data.PointAttributes.Serialize(Ar, Owner, ChunkIndex, bAttemptFileMapping);
+		if (Header.Flags & DataFlags_HasPointAttribute)
+		{
+			Data.PointAttributes.Serialize(Ar, Owner, ChunkIndex, bAttemptFileMapping);
+		}
 		Data.CurveAttributes.Serialize(Ar, Owner, ChunkIndex, bAttemptFileMapping);
 		Data.PointToCurve.Serialize(Ar, Owner, ChunkIndex, bAttemptFileMapping);
 		Data.Curves.Serialize(Ar, Owner, ChunkIndex, bAttemptFileMapping);
diff --git a/Engine/Plugins/Runtime/HairStrands/Source/HairStrandsCore/Public/HairStrandsDatas.h b/Engine/Plugins/Runtime/HairStrands/Source/HairStrandsCore/Public/HairStrandsDatas.h
--- a/Engine/Plugins/Runtime/HairStrands/Source/HairStrandsCore/Public/HairStrandsDatas.h
+++ b/Engine/Plugins/Runtime/HairStrands/Source/HairStrandsCore/Public/HairStrandsDatas.h
@@ -485,6 +485,7 @@ struct HAIRSTRANDSCORE_API FHairStrandsBulkData : FHairStrandsBulkCommon
 	{
 		DataFlags_HasData = 1,				// Contains valid data. Otherwise: Position, Attributes, ... are all empty
 		DataFlags_Has16bitsCurveIndex = 2,	// Use 16bits index for vertex to curve mapping
+		DataFlags_HasPointAttribute = 4,	// Contains point attribute data.
 	};
 
 	void Serialize(FArchive& Ar, UObject* Owner);
