Add DDC streaming for groom resource.

diff --git a/Engine/Plugins/Runtime/HairStrands/Source/HairStrandsCore/Private/GroomAsset.cpp b/Engine/Plugins/Runtime/HairStrands/Source/HairStrandsCore/Private/GroomAsset.cpp
--- a/Engine/Plugins/Runtime/HairStrands/Source/HairStrandsCore/Private/GroomAsset.cpp
+++ b/Engine/Plugins/Runtime/HairStrands/Source/HairStrandsCore/Private/GroomAsset.cpp
@@ -708,7 +708,7 @@ static void InternalSerializePlatformDatas(FArchive& Ar, UObject* Owner, TArray<
 static void InternalSerializeCard(FArchive& Ar, UObject* Owner, FHairGroupPlatformData::FCards::FLOD& CardLODData);
 static void InternalSerializeMesh(FArchive& Ar, UObject* Owner, FHairGroupPlatformData::FMeshes::FLOD& MeshLODData);
 static void InternalSerializeGuide(FArchive& Ar, UObject* Owner, FHairGroupPlatformData::FGuides& GuideData);
-static void InternalSerializeStrand(FArchive& Ar, UObject* Owner, FHairGroupPlatformData::FStrands& StrandData);
+static void InternalSerializeStrand(FArchive& Ar, UObject* Owner, FHairGroupPlatformData::FStrands& StrandData, bool bHeader, bool bData);
 static void InternalSerializePlatformData(FArchive& Ar, UObject* Owner, FHairGroupPlatformData& GroupData);
 
 // Serialize is only called for cooking??
@@ -1673,17 +1673,28 @@ static void InternalSerializeGuide(FArchive& Ar, UObject* Owner, FHairGroupPlatf
 	GuideData.BulkData.Serialize(Ar, Owner);
 }
 
-static void InternalSerializeStrand(FArchive& Ar, UObject* Owner, FHairGroupPlatformData::FStrands& StrandData)
+static void InternalSerializeStrand(FArchive& Ar, UObject* Owner, FHairGroupPlatformData::FStrands& StrandData, bool bHeader, bool bData)
 {
 	Ar.UsingCustomVersion(FAnimObjectVersion::GUID);
 
+	// When cooking data, force loading of bulk data prior to saving them
+	if (Ar.IsCooking() && Ar.IsSaving())
+	{
+		{ FHairResourceRequest R; R.Request(StrandData.BulkData, true /*bWait*/); }
+		{ FHairResourceRequest R; R.Request(StrandData.InterpolationBulkData, true /*bWait*/); }
+		{ FHairResourceRequest R; R.Request(StrandData.ClusterCullingBulkData, true /*bWait*/); }
+	}
+
 	if (!Ar.IsCooking() || !StrandData.bIsCookedOut)
 	{
-		StrandData.BulkData.Serialize(Ar, Owner);
-		StrandData.InterpolationBulkData.Serialize(Ar, Owner);
+		if (bHeader)	{ StrandData.BulkData.SerializeHeader(Ar, Owner); }
+		if (bData)		{ StrandData.BulkData.SerializeData(Ar, Owner); }
+		if (bHeader)	{ StrandData.InterpolationBulkData.SerializeHeader(Ar, Owner); }
+		if (bData)		{ StrandData.InterpolationBulkData.SerializeData(Ar, Owner); }
 		if (Ar.CustomVer(FAnimObjectVersion::GUID) >= FAnimObjectVersion::SerializeHairClusterCullingData)
 		{
-			StrandData.ClusterCullingBulkData.Serialize(Ar, Owner);
+			if (bHeader){ StrandData.ClusterCullingBulkData.SerializeHeader(Ar, Owner); }
+			if (bData)	{ StrandData.ClusterCullingBulkData.SerializeData(Ar, Owner); }
 		}
 	}
 	else
@@ -1692,11 +1703,14 @@ static void InternalSerializeStrand(FArchive& Ar, UObject* Owner, FHairGroupPlat
 		// Theoritically, we should have something to detect if we are going to use or not guide (for 
 		// simulation or RBF deformation) on the target platform
 		FHairGroupPlatformData::FStrands NoStrandsData;
-		NoStrandsData.BulkData.Serialize(Ar, Owner);
-		NoStrandsData.InterpolationBulkData.Serialize(Ar, Owner);
+		if (bHeader) 	{ NoStrandsData.BulkData.SerializeHeader(Ar, Owner); }
+		if (bData)		{ NoStrandsData.BulkData.SerializeData(Ar, Owner); }
+		if (bHeader) 	{ NoStrandsData.InterpolationBulkData.SerializeHeader(Ar, Owner); }
+		if (bData) 		{ NoStrandsData.InterpolationBulkData.SerializeData(Ar, Owner); }
 		if (Ar.CustomVer(FAnimObjectVersion::GUID) >= FAnimObjectVersion::SerializeHairClusterCullingData)
 		{
-			NoStrandsData.ClusterCullingBulkData.Serialize(Ar, Owner);
+			if (bHeader){ NoStrandsData.ClusterCullingBulkData.SerializeHeader(Ar, Owner); }
+			if (bData)	{ NoStrandsData.ClusterCullingBulkData.SerializeData(Ar, Owner); }
 		}
 	}
 }
@@ -1706,7 +1720,7 @@ static void InternalSerializePlatformData(FArchive& Ar, UObject* Owner, FHairGro
 	Ar.UsingCustomVersion(FAnimObjectVersion::GUID);
 
 	InternalSerializeGuide(Ar, Owner, GroupData.Guides);
-	InternalSerializeStrand(Ar, Owner, GroupData.Strands);
+	InternalSerializeStrand(Ar, Owner, GroupData.Strands, true, true);
 
 	if (Ar.CustomVer(FAnimObjectVersion::GUID) >= FAnimObjectVersion::SerializeGroomCardsAndMeshes)
 	{
@@ -2332,32 +2346,47 @@ bool UGroomAsset::CacheStrandsData(uint32 GroupIndex, FString& OutDerivedDataKey
 
 	using namespace UE::DerivedData;
 
-	const FString DerivedDataKey = UGroomAsset::GetDerivedDataKeyForStrands(GroupIndex);
-	const FCacheKey Key = ConvertLegacyCacheKey(DerivedDataKey);
+	const FString DerivedDataKey= UGroomAsset::GetDerivedDataKeyForStrands(GroupIndex);
+	const FCacheKey HeaderKey = ConvertLegacyCacheKey(DerivedDataKey + FString(TEXT("_Header")));
 	const FSharedString Name = MakeStringView(GetPathName());
 	FSharedBuffer Data;
 	{
 		FRequestOwner Owner(EPriority::Blocking);
-		GetCache().GetValue({ {Name, Key} }, Owner, [&Data](FCacheGetValueResponse&& Response)
+		GetCache().GetValue({ {Name, HeaderKey} }, Owner, [&Data](FCacheGetValueResponse&& Response)
 		{
 			Data = Response.Value.GetData().Decompress();
 		});
 		Owner.Wait();
 	}
 
+	// Populate key/name for streaming data request
+	auto FillDrivedDataKey = [&DerivedDataKey, &Name](FHairGroupPlatformData& In)
+	{
+		In.Strands.BulkData.DerivedDataKey = DerivedDataKey + FString(TEXT("_RestData"));
+		In.Strands.BulkData.Name = Name;
+	
+		In.Strands.InterpolationBulkData.DerivedDataKey = DerivedDataKey + FString(TEXT("_InterpolationData"));
+		In.Strands.InterpolationBulkData.Name = Name;
+	
+		In.Strands.ClusterCullingBulkData.DerivedDataKey = DerivedDataKey + FString(TEXT("_ClusterCullingData"));
+		In.Strands.ClusterCullingBulkData.Name = Name;
+	};
+
+	FHairGroupPlatformData& PlatformData = HairGroupsPlatformData[GroupIndex];
+
 	bool bSuccess = true;
 	if (Data)
 	{
 		UE_CLOG(IsHairStrandsDDCLogEnable(), LogHairStrands, Log, TEXT("[Groom/DDC] Strands - Found (Groom:%s Group6:%d)."), *GetName(), GroupIndex);
 
-		// Reset hair group data to ensure previously loaded bulk data are cleared/cleaned priori to load new data.
-		HairGroupsPlatformData[GroupIndex] = FHairGroupPlatformData();
+		// Reset hair group data to ensure previously loaded bulk data are cleared/cleaned prior to load new data.
+		PlatformData = FHairGroupPlatformData();
+		FillDrivedDataKey(PlatformData);
 
 		FMemoryReaderView Ar(Data, /*bIsPersistent*/ true);
 
-		InternalSerializeGuide(Ar, this, HairGroupsPlatformData[GroupIndex].Guides);
-		InternalSerializeStrand(Ar, this, HairGroupsPlatformData[GroupIndex].Strands);
-
+		InternalSerializeGuide(Ar, this, PlatformData.Guides);
+		InternalSerializeStrand(Ar, this, PlatformData.Strands, true/*Header*/, false/*Data*/);
 	}
 	else
 	{
@@ -2390,15 +2419,63 @@ bool UGroomAsset::CacheStrandsData(uint32 GroupIndex, FString& OutDerivedDataKey
 		
 		if (bSuccess)
 		{
-			TArray<uint8> WriteData;
-			FMemoryWriter Ar(WriteData, /*bIsPersistent*/ true);
+			FillDrivedDataKey(PlatformData);
+
+			// Header
+			{
+				TArray<uint8> WriteData;
+				FMemoryWriter Ar(WriteData, /*bIsPersistent*/ true);
+				InternalSerializeGuide(Ar, this, PlatformData.Guides);
+				InternalSerializeStrand(Ar, this, PlatformData.Strands, true/*Header*/, false/*Data*/);
+
+
+				FRequestOwner AsyncOwner(EPriority::Normal);
+				GetCache().PutValue({ {Name, HeaderKey, FValue::Compress(MakeSharedBufferFromArray(MoveTemp(WriteData)))} }, AsyncOwner);
+				AsyncOwner.KeepAlive();
+			}
+
+			// Data (Rest)
+			{
+				TArray<uint8> WriteData;
+				FMemoryWriter Ar(WriteData, /*bIsPersistent*/ true);
+
+				PlatformData.Strands.BulkData.SerializeData(Ar, this);
+				
+				const FCacheKey DataKey = ConvertLegacyCacheKey(PlatformData.Strands.BulkData.DerivedDataKey);
+
+				FRequestOwner AsyncOwner(EPriority::Normal);
+				GetCache().PutValue({ {Name, DataKey, FValue::Compress(MakeSharedBufferFromArray(MoveTemp(WriteData)))} }, AsyncOwner);
+				AsyncOwner.KeepAlive();
+			}
+
+			// Data (Interpolation)
+			if (PlatformData.Strands.InterpolationBulkData.HasData())
+			{
+				TArray<uint8> WriteData;
+				FMemoryWriter Ar(WriteData, /*bIsPersistent*/ true);
+
+				PlatformData.Strands.InterpolationBulkData.SerializeData(Ar, this);
+				
+				const FCacheKey DataKey = ConvertLegacyCacheKey(PlatformData.Strands.InterpolationBulkData.DerivedDataKey);
+
+				FRequestOwner AsyncOwner(EPriority::Normal);
+				GetCache().PutValue({ {Name, DataKey, FValue::Compress(MakeSharedBufferFromArray(MoveTemp(WriteData)))} }, AsyncOwner);
+				AsyncOwner.KeepAlive();
+			}
+
+			// Data (cluster Culling)
+			{
+				TArray<uint8> WriteData;
+				FMemoryWriter Ar(WriteData, /*bIsPersistent*/ true);
+
+				PlatformData.Strands.ClusterCullingBulkData.SerializeData(Ar, this);
 				
-			InternalSerializeGuide(Ar, this, HairGroupsPlatformData[GroupIndex].Guides);
-			InternalSerializeStrand(Ar, this, HairGroupsPlatformData[GroupIndex].Strands);
+				const FCacheKey DataKey = ConvertLegacyCacheKey(PlatformData.Strands.ClusterCullingBulkData.DerivedDataKey);
 
-			FRequestOwner AsyncOwner(EPriority::Normal);
-			GetCache().PutValue({ {Name, Key, FValue::Compress(MakeSharedBufferFromArray(MoveTemp(WriteData)))} }, AsyncOwner);
-			AsyncOwner.KeepAlive();
+				FRequestOwner AsyncOwner(EPriority::Normal);
+				GetCache().PutValue({ {Name, DataKey, FValue::Compress(MakeSharedBufferFromArray(MoveTemp(WriteData)))} }, AsyncOwner);
+				AsyncOwner.KeepAlive();
+			}
 		}
 	}
 
diff --git a/Engine/Plugins/Runtime/HairStrands/Source/HairStrandsCore/Private/GroomBuilder.cpp b/Engine/Plugins/Runtime/HairStrands/Source/HairStrandsCore/Private/GroomBuilder.cpp
--- a/Engine/Plugins/Runtime/HairStrands/Source/HairStrandsCore/Private/GroomBuilder.cpp
+++ b/Engine/Plugins/Runtime/HairStrands/Source/HairStrandsCore/Private/GroomBuilder.cpp
@@ -40,7 +40,7 @@ static FAutoConsoleVariableRef CVarHairGroupIndexBuilder_MaxVoxelResolution(TEXT
 
 FString FGroomBuilder::GetVersion()
 {
-	return TEXT("v8p1");
+	return TEXT("v8p3");
 }
 
 namespace FHairStrandsDecimation
diff --git a/Engine/Plugins/Runtime/HairStrands/Source/HairStrandsCore/Private/GroomManager.cpp b/Engine/Plugins/Runtime/HairStrands/Source/HairStrandsCore/Private/GroomManager.cpp
--- a/Engine/Plugins/Runtime/HairStrands/Source/HairStrandsCore/Private/GroomManager.cpp
+++ b/Engine/Plugins/Runtime/HairStrands/Source/HairStrandsCore/Private/GroomManager.cpp
@@ -1000,9 +1000,9 @@ static void RunHairLODSelection(
 				Instance->HairGroupPublicData->Allocate(GraphBuilder);
 
 				if (Instance->Strands.RestRootResource)			{ Instance->Strands.RestRootResource->Allocate(GraphBuilder, LoadingType, ResourceStatus); Instance->Strands.RestRootResource->AllocateLOD(GraphBuilder, MeshLODIndex, LoadingType, ResourceStatus); }
-				if (Instance->Strands.RestResource)				{ Instance->Strands.RestResource->Allocate(GraphBuilder, LoadingType, ResourceStatus); }
-				if (Instance->Strands.ClusterCullingResource)	{ Instance->Strands.ClusterCullingResource->Allocate(GraphBuilder, LoadingType, ResourceStatus); }
-				if (Instance->Strands.InterpolationResource)	{ Instance->Strands.InterpolationResource->Allocate(GraphBuilder, LoadingType, ResourceStatus); }
+				if (Instance->Strands.RestResource)				{ Instance->Strands.RestResource->Allocate(GraphBuilder, EHairResourceLoadingType::Async, ResourceStatus); }
+				if (Instance->Strands.ClusterCullingResource)	{ Instance->Strands.ClusterCullingResource->Allocate(GraphBuilder, EHairResourceLoadingType::Async, ResourceStatus); }
+				if (Instance->Strands.InterpolationResource)	{ Instance->Strands.InterpolationResource->Allocate(GraphBuilder, EHairResourceLoadingType::Async, ResourceStatus); }
 
 				if (Instance->Strands.DeformedRootResource)		{ Instance->Strands.DeformedRootResource->Allocate(GraphBuilder, LoadingType, ResourceStatus); Instance->Strands.DeformedRootResource->AllocateLOD(GraphBuilder, MeshLODIndex, LoadingType, ResourceStatus); }
 				if (Instance->Strands.DeformedResource)			{ Instance->Strands.DeformedResource->Allocate(GraphBuilder, LoadingType, ResourceStatus); }
diff --git a/Engine/Plugins/Runtime/HairStrands/Source/HairStrandsCore/Private/GroomResources.cpp b/Engine/Plugins/Runtime/HairStrands/Source/HairStrandsCore/Private/GroomResources.cpp
--- a/Engine/Plugins/Runtime/HairStrands/Source/HairStrandsCore/Private/GroomResources.cpp
+++ b/Engine/Plugins/Runtime/HairStrands/Source/HairStrandsCore/Private/GroomResources.cpp
@@ -27,6 +27,11 @@
 #include "RenderGraphUtils.h"
 #include "RenderTargetPool.h"
 #include "GroomBindingBuilder.h"
+#if WITH_EDITORONLY_DATA
+#include "Containers/StringView.h"
+#include "DerivedDataCache.h"
+#include "DerivedDataRequestOwner.h"
+#endif
 
 static int32 GHairStrandsBulkData_ReleaseAfterUse = 0;
 static FAutoConsoleVariableRef CVarHairStrandsBulkData_ReleaseAfterUse(TEXT("r.HairStrands.Strands.BulkData.ReleaseAfterUse"), GHairStrandsBulkData_ReleaseAfterUse, TEXT("Release CPU bulk data once hair groom/groom binding asset GPU resources are created. This saves memory"));
@@ -90,6 +95,55 @@ const TCHAR* ToHairResourceDebugName(const TCHAR* In, FHairResourceName& InDebug
 #endif
 }
 
+/////////////////////////////////////////////////////////////////////////////////////////
+// IO/DDC request
+
+#if !WITH_EDITORONLY_DATA
+bool FHairResourceRequest::IsNone() const 									{ return IORequest.IsNone(); }
+bool FHairResourceRequest::IsCompleted() const								{ return IORequest.IsCompleted(); }
+void FHairResourceRequest::Request(FHairStrandsBulkCommon& In, bool bWait)	{ In.Request(IORequest); }
+#else
+bool FHairResourceRequest::IsNone() const 									{ return Status == None; }
+bool FHairResourceRequest::IsCompleted() const								{ return Status == Completed; }
+void FHairResourceRequest::Request(FHairStrandsBulkCommon& In, bool bWait)
+{
+	if (!In.HasData())
+	{
+		Status = Completed;
+		return;
+	}
+
+	Status = Pending;
+
+	using namespace UE::DerivedData;
+	DDCRequestOwner = MakeUnique<FRequestOwner>(bWait ? UE::DerivedData::EPriority::Blocking : UE::DerivedData::EPriority::Normal);
+	const FCacheKey DataKey = ConvertLegacyCacheKey(In.DerivedDataKey);
+	const FSharedString Name = MakeStringView(In.Name);
+
+	UE::DerivedData::GetCache().GetValue({ {Name, DataKey} }, *DDCRequestOwner,
+	[this, &In](FCacheGetValueResponse&& Response)
+	{
+		if (Response.Status == UE::DerivedData::EStatus::Ok)
+		{
+			FSharedBuffer Data = Response.Value.GetData().Decompress();
+			FMemoryReaderView Reader(Data, true /*bIsPersistent*/);
+			In.SerializeData(Reader, nullptr /*Owner*/);
+			Status = Completed;
+		}
+		else
+		{
+			checkNoEntry();
+		}
+	});
+
+	if (bWait)
+	{
+		DDCRequestOwner->Wait();
+		check(Status == Completed);
+	}
+}
+#endif
+
 /////////////////////////////////////////////////////////////////////////////////////////
 // FRWBuffer utils 
 void UploadDataToBuffer(FReadBuffer& OutBuffer, uint32 DataSizeInBytes, const void* InCpuData)
@@ -875,26 +929,15 @@ bool FHairStrandsRestResource::InternalIsDataLoaded()
 {
 	if (BulkDataRequest.IsNone())
 	{
-		FBulkDataBatchRequest::FBatchBuilder Batch = FBulkDataBatchRequest::NewBatch(5);
-		Batch.Read(BulkData.Data.Positions);
-		Batch.Read(BulkData.Data.CurveAttributes);
-		if (BulkData.Header.Flags & FHairStrandsBulkData::DataFlags_HasPointAttribute)
-		{
-			Batch.Read(BulkData.Data.PointAttributes);
-		}
-		Batch.Read(BulkData.Data.PointToCurve);
-		Batch.Read(BulkData.Data.Curves);
-
-		Batch.Issue(BulkDataRequest);
+		BulkDataRequest.Request(BulkData);
 	}
-	
 	return BulkDataRequest.IsCompleted();
 }
 
 bool IsHairStrandsContinousLODEnabled();
 void FHairStrandsRestResource::InternalAllocate(FRDGBuilder& GraphBuilder)
 {
-	BulkDataRequest = FBulkDataBatchRequest();
+	BulkDataRequest = FHairResourceRequest();
 
 	const uint32 PointCount = BulkData.GetNumPoints();
 	const uint32 CurveCount = BulkData.GetNumCurves();
@@ -1136,6 +1179,31 @@ void FHairStrandsClusterCullingBulkData::SerializeData(FArchive& Ar, UObject* Ow
 	}
 }
 
+void FHairStrandsClusterCullingBulkData::Request(FBulkDataBatchRequest& In)
+{
+	FBulkDataBatchRequest::FBatchBuilder Batch = FBulkDataBatchRequest::NewBatch(4);
+	if (Header.ClusterLODCount)
+	{
+		Batch.Read(Data.ClusterLODInfos);
+	}
+
+	if (Header.VertexCount)
+	{
+		Batch.Read(Data.VertexToClusterIds);
+	}
+
+	if (Header.VertexLODCount)
+	{
+		Batch.Read(Data.ClusterVertexIds);
+	}
+
+	if (Header.ClusterCount)
+	{
+		Batch.Read(Data.PackedClusterInfos);
+	}
+	Batch.Issue(In);
+}
+
 void FHairStrandsClusterCullingBulkData::Validate(bool bIsSaving)
 {
 	if (Header.ClusterCount == 0)
@@ -1173,20 +1241,14 @@ bool FHairStrandsClusterCullingResource::InternalIsDataLoaded()
 {
 	if (BulkDataRequest.IsNone())
 	{
-		FBulkDataBatchRequest::NewBatch(4)
-			.Read(BulkData.Data.PackedClusterInfos)
-			.Read(BulkData.Data.ClusterLODInfos)
-			.Read(BulkData.Data.VertexToClusterIds)
-			.Read(BulkData.Data.ClusterVertexIds)
-			.Issue(BulkDataRequest);
+		BulkDataRequest.Request(BulkData);
 	}
-
 	return BulkDataRequest.IsCompleted();
 }
 
 void FHairStrandsClusterCullingResource::InternalAllocate(FRDGBuilder& GraphBuilder)
 {
-	BulkDataRequest = FBulkDataBatchRequest();
+	BulkDataRequest = FHairResourceRequest();
 
 	if (GHairStrandsBulkData_Validation > 0)
 	{
@@ -1206,7 +1268,7 @@ void FHairStrandsClusterCullingResource::InternalRelease()
 	ClusterLODInfoBuffer.Release();
 	ClusterVertexIdBuffer.Release();
 	VertexToClusterIdBuffer.Release();
-	BulkDataRequest = FBulkDataBatchRequest();
+	BulkDataRequest = FHairResourceRequest();
 }
 
 /////////////////////////////////////////////////////////////////////////////////////////
@@ -1634,28 +1696,14 @@ bool FHairStrandsInterpolationResource::InternalIsDataLoaded()
 {
 	if (BulkDataRequest.IsNone())
 	{
-		FBulkDataBatchRequest::FBatchBuilder Batch = FBulkDataBatchRequest::NewBatch(4);
-		const bool bUseSingleGuide = !!(BulkData.Header.Flags & FHairStrandsInterpolationBulkData::DataFlags_HasSingleGuideData);
-		if (bUseSingleGuide)
-		{
-			Batch.Read(BulkData.Data.Interpolation);
-		}
-		else
-		{
-			Batch.Read(BulkData.Data.Interpolation0);
-			Batch.Read(BulkData.Data.Interpolation1);
-		}
-		Batch.Read(BulkData.Data.SimRootPointIndex);
-
-		Batch.Issue(BulkDataRequest);
+		BulkDataRequest.Request(BulkData);
 	}
-
 	return BulkDataRequest.IsCompleted();
 }
 
 void FHairStrandsInterpolationResource::InternalAllocate(FRDGBuilder& GraphBuilder)
 {
-	BulkDataRequest = FBulkDataBatchRequest();
+	BulkDataRequest = FHairResourceRequest();
 
 	const bool bUseSingleGuide = !!(BulkData.Header.Flags & FHairStrandsInterpolationBulkData::DataFlags_HasSingleGuideData);
 	if (bUseSingleGuide)
@@ -1676,7 +1724,7 @@ void FHairStrandsInterpolationResource::InternalRelease()
 	Interpolation0Buffer.Release();
 	Interpolation1Buffer.Release();
 	SimRootPointIndexBuffer.Release();
-	BulkDataRequest = FBulkDataBatchRequest();
+	BulkDataRequest = FHairResourceRequest();
 }
 
 /////////////////////////////////////////////////////////////////////////////////////////
diff --git a/Engine/Plugins/Runtime/HairStrands/Source/HairStrandsCore/Private/HairStrandsDatas.cpp b/Engine/Plugins/Runtime/HairStrands/Source/HairStrandsCore/Private/HairStrandsDatas.cpp
--- a/Engine/Plugins/Runtime/HairStrands/Source/HairStrandsCore/Private/HairStrandsDatas.cpp
+++ b/Engine/Plugins/Runtime/HairStrands/Source/HairStrandsCore/Private/HairStrandsDatas.cpp
@@ -225,6 +225,33 @@ void FHairStrandsInterpolationBulkData::SerializeData(FArchive& Ar, UObject* Own
 	}
 }
 
+bool FHairStrandsInterpolationBulkData::HasData() const
+{
+	return Header.Flags & DataFlags_HasData;
+}
+
+void FHairStrandsInterpolationBulkData::Request(FBulkDataBatchRequest& InRequest)
+{
+	if (!(Header.Flags & DataFlags_HasData))
+	{
+		return;
+	}
+
+	check(InRequest.IsNone());
+	FBulkDataBatchRequest::FBatchBuilder Batch = FBulkDataBatchRequest::NewBatch(3);
+	if (!!(Header.Flags & DataFlags_HasSingleGuideData))
+	{
+		Batch.Read(Data.Interpolation);
+	}
+	else
+	{
+		Batch.Read(Data.Interpolation0);
+		Batch.Read(Data.Interpolation1);
+	}
+	Batch.Read(Data.SimRootPointIndex);
+	Batch.Issue(InRequest);
+}
+
 void FHairStrandsBulkData::Serialize(FArchive& Ar, UObject* Owner)
 {
 	SerializeHeader(Ar, Owner);
@@ -297,6 +324,24 @@ void FHairStrandsBulkData::SerializeData(FArchive& Ar, UObject* Owner)
 	}
 }
 
+void FHairStrandsBulkData::Request(FBulkDataBatchRequest& In)
+{
+	if (!!(Header.Flags & DataFlags_HasData))
+	{
+		check(In.IsNone());
+
+		FBulkDataBatchRequest::FBatchBuilder Batch = FBulkDataBatchRequest::NewBatch(5);
+		Batch.Read(Data.Positions);
+		if (Header.Flags & DataFlags_HasPointAttribute)
+		{
+			Batch.Read(Data.PointAttributes);
+		}
+		Batch.Read(Data.CurveAttributes);
+		Batch.Read(Data.PointToCurve);
+		Batch.Read(Data.Curves);
+		Batch.Issue(In);
+	}
+}
 
 void FHairStrandsBulkData::Reset()
 {
diff --git a/Engine/Plugins/Runtime/HairStrands/Source/HairStrandsCore/Public/GroomResources.h b/Engine/Plugins/Runtime/HairStrands/Source/HairStrandsCore/Public/GroomResources.h
--- a/Engine/Plugins/Runtime/HairStrands/Source/HairStrandsCore/Public/GroomResources.h
+++ b/Engine/Plugins/Runtime/HairStrands/Source/HairStrandsCore/Public/GroomResources.h
@@ -12,6 +12,8 @@
 #include "HairStrandsInterface.h"
 #include "HairStrandsMeshProjection.h"
 
+namespace UE::DerivedData { class FRequestOwner; }
+
 #define STRANDS_PROCEDURAL_INTERSECTOR_MAX_SPLITS 4
 
 inline uint32 GetBufferTotalNumBytes(const FRDGExternalBuffer& In) 
@@ -73,6 +75,24 @@ FORCEINLINE bool operator! (EHairResourceStatus A) { return static_cast<uint8>(A
 
 EHairResourceLoadingType GetHairResourceLoadingType(EHairGeometryType InGeometryType, int32 InLODIndex);
 
+struct FHairResourceRequest
+{
+	void Request(FHairStrandsBulkCommon& In, bool bWait=false);
+	bool IsNone() const ;
+	bool IsCompleted() const;
+
+#if !WITH_EDITORONLY_DATA
+	// IO
+	FBulkDataBatchRequest IORequest;
+#else
+	// DDC
+	enum EStatus { None, Pending, Completed };
+	EStatus Status = EStatus::None;
+	FString PathName;
+	TUniquePtr<UE::DerivedData::FRequestOwner> DDCRequestOwner;
+#endif
+};
+
 /* Hair resouces which whom allocation can be deferred */
 struct FHairCommonResource : public FRenderResource
 {
@@ -345,7 +365,7 @@ struct FHairStrandsRestResource : public FHairCommonResource
 	FHairStrandsBulkData& BulkData;
 
 	/* Handle to bulk data request */
-	FBulkDataBatchRequest BulkDataRequest;
+	FHairResourceRequest BulkDataRequest;
 
 	/* Type of curves */
 	const EHairStrandsResourcesType CurveType;
@@ -468,7 +488,7 @@ struct FHairStrandsClusterCullingResource : public FHairCommonResource
 	FHairStrandsClusterCullingBulkData& BulkData;
 
 	/* Handle to bulk data request */
-	FBulkDataBatchRequest BulkDataRequest;
+	FHairResourceRequest BulkDataRequest;
 };
 
 struct FHairStrandsInterpolationResource : public FHairCommonResource
@@ -506,7 +526,7 @@ struct FHairStrandsInterpolationResource : public FHairCommonResource
 	FHairStrandsInterpolationBulkData& BulkData;
 
 	/* Handle to bulk data request */
-	FBulkDataBatchRequest BulkDataRequest;
+	FHairResourceRequest BulkDataRequest;
 };
 
 #if RHI_RAYTRACING
diff --git a/Engine/Plugins/Runtime/HairStrands/Source/HairStrandsCore/Public/HairStrandsDatas.h b/Engine/Plugins/Runtime/HairStrands/Source/HairStrandsCore/Public/HairStrandsDatas.h
--- a/Engine/Plugins/Runtime/HairStrands/Source/HairStrandsCore/Public/HairStrandsDatas.h
+++ b/Engine/Plugins/Runtime/HairStrands/Source/HairStrandsCore/Public/HairStrandsDatas.h
@@ -293,6 +293,8 @@ struct HAIRSTRANDSCORE_API FHairStrandsBulkCommon
 	virtual ~FHairStrandsBulkCommon() { }
 	virtual void SerializeHeader(FArchive& Ar, UObject* Owner) = 0;
 	virtual void SerializeData(FArchive& Ar, UObject* Owner) = 0;
+	virtual void Request(FBulkDataBatchRequest& In) = 0;
+	virtual bool HasData() const { return true; }
 
 #if WITH_EDITORONLY_DATA
 	// Transient Name/DDCkey for streaming
@@ -343,6 +345,8 @@ struct HAIRSTRANDSCORE_API FHairStrandsInterpolationBulkData : FHairStrandsBulkC
 	void Serialize(FArchive& Ar, UObject* Owner);
 	virtual void SerializeHeader(FArchive& Ar, UObject* Owner) override;
 	virtual void SerializeData(FArchive& Ar, UObject* Owner) override;
+	virtual void Request(FBulkDataBatchRequest& In) override;
+	virtual bool HasData() const override;
 	uint32 GetPointCount() const { return Header.PointCount; };
 
 	struct FHeader
@@ -491,6 +495,7 @@ struct HAIRSTRANDSCORE_API FHairStrandsBulkData : FHairStrandsBulkCommon
 	void Serialize(FArchive& Ar, UObject* Owner);
 	virtual void SerializeHeader(FArchive& Ar, UObject* Owner) override;
 	virtual void SerializeData(FArchive& Ar, UObject* Owner) override;
+	virtual void Request(FBulkDataBatchRequest& In) override;
 
 	bool IsValid() const { return Header.CurveCount > 0 && Header.PointCount > 0; }
 	void Reset();
@@ -807,6 +812,7 @@ struct HAIRSTRANDSCORE_API FHairStrandsClusterCullingBulkData : FHairStrandsBulk
 	void Serialize(FArchive& Ar, UObject* Owner);
 	virtual void SerializeHeader(FArchive& Ar, UObject* Owner) override;
 	virtual void SerializeData(FArchive& Ar, UObject* Owner) override;
+	virtual void Request(FBulkDataBatchRequest& In) override;
 	bool IsValid() const { return Header.ClusterCount > 0 && Header.VertexCount > 0; }
 	void Validate(bool bIsSaving);
 
