Fix shader compilation on groom deformer.

diff --git a/Engine/Plugins/Runtime/HairStrands/Shaders/Private/DeformerDataInterfaceGroom.ush b/Engine/Plugins/Runtime/HairStrands/Shaders/Private/DeformerDataInterfaceGroom.ush
--- a/Engine/Plugins/Runtime/HairStrands/Shaders/Private/DeformerDataInterfaceGroom.ush
+++ b/Engine/Plugins/Runtime/HairStrands/Shaders/Private/DeformerDataInterfaceGroom.ush
@@ -71,18 +71,18 @@ float ReadCoordU_{DataInterfaceName}(uint ControlPointIndex)
 float ReadLength_{DataInterfaceName}(uint ControlPointIndex)
 {
 	const float CoordU      = ReadControlPointData(ControlPointIndex).UCoord;
-	const float CurveLength = UnpackHairLength(InternalGetHairStrandsAttributesPerCurve(ControlPointIndex, HAIR_ATTRIBUTE_OFFSET_LENGTH({DataInterfaceName}_AttributeOffsets), HAIR_ATTRIBUTE_STRIDE_LENGTH, {DataInterfaceName}_AttributeBuffer, {DataInterfaceName}_PointToCurveBuffer));
+	const float CurveLength = UnpackHairLength(InternalGetHairStrandsAttributesPerCurve(ControlPointIndex, HAIR_CURVE_ATTRIBUTE_OFFSET_LENGTH({DataInterfaceName}_CurveAttributeOffsets), HAIR_CURVE_ATTRIBUTE_STRIDE_LENGTH, {DataInterfaceName}_CurveAttributeBuffer, {DataInterfaceName}_PointToCurveBuffer));
 	return CurveLength * CoordU;
 }
 
 float2 ReadRootUV_{DataInterfaceName}(uint ControlPointIndex)
 {
-	return UnpackHairRootUV(InternalGetHairStrandsAttributesPerCurve(ControlPointIndex, HAIR_ATTRIBUTE_OFFSET_ROOTUV({DataInterfaceName}_AttributeOffsets), HAIR_ATTRIBUTE_STRIDE_ROOTUV, {DataInterfaceName}_AttributeBuffer, {DataInterfaceName}_PointToCurveBuffer));
+	return UnpackHairRootUV(InternalGetHairStrandsAttributesPerCurve(ControlPointIndex, HAIR_CURVE_ATTRIBUTE_OFFSET_ROOTUV({DataInterfaceName}_CurveAttributeOffsets), HAIR_CURVE_ATTRIBUTE_STRIDE_ROOTUV, {DataInterfaceName}_CurveAttributeBuffer, {DataInterfaceName}_PointToCurveBuffer));
 }
 
 float ReadSeed_{DataInterfaceName}(uint ControlPointIndex)
 {
-	return UnpackHairSeed(InternalGetHairStrandsAttributesPerCurve(ControlPointIndex, HAIR_ATTRIBUTE_OFFSET_SEED({DataInterfaceName}_AttributeOffsets), HAIR_ATTRIBUTE_STRIDE_SEED, {DataInterfaceName}_AttributeBuffer, {DataInterfaceName}_PointToCurveBuffer));
+	return UnpackHairSeed(InternalGetHairStrandsAttributesPerCurve(ControlPointIndex, HAIR_CURVE_ATTRIBUTE_OFFSET_SEED({DataInterfaceName}_CurveAttributeOffsets), HAIR_CURVE_ATTRIBUTE_STRIDE_SEED, {DataInterfaceName}_CurveAttributeBuffer, {DataInterfaceName}_PointToCurveBuffer));
 }
 
 float ReadCurveOffsetPoint_{DataInterfaceName}(uint CurveIndex)
