* Prepare groom for partial/chunk request.
* Unify resource gathering

diff --git a/Engine/Plugins/Runtime/HairStrands/Source/HairStrandsCore/Private/GroomAsset.cpp b/Engine/Plugins/Runtime/HairStrands/Source/HairStrandsCore/Private/GroomAsset.cpp
--- a/Engine/Plugins/Runtime/HairStrands/Source/HairStrandsCore/Private/GroomAsset.cpp
+++ b/Engine/Plugins/Runtime/HairStrands/Source/HairStrandsCore/Private/GroomAsset.cpp
@@ -1680,9 +1680,9 @@ static void InternalSerializeStrand(FArchive& Ar, UObject* Owner, FHairGroupPlat
 	// When cooking data, force loading of bulk data prior to saving them
 	if (Ar.IsCooking() && Ar.IsSaving())
 	{
-		{ FHairResourceRequest R; R.Request(StrandData.BulkData, true /*bWait*/); }
-		{ FHairResourceRequest R; R.Request(StrandData.InterpolationBulkData, true /*bWait*/); }
-		{ FHairResourceRequest R; R.Request(StrandData.ClusterCullingBulkData, true /*bWait*/); }
+		{ FHairResourceRequest R; R.Request(StrandData.BulkData, true /*bWait*/, true /*bFillBulkdata*/); }
+		{ FHairResourceRequest R; R.Request(StrandData.InterpolationBulkData, true /*bWait*/, true /*bFillBulkdata*/); }
+		{ FHairResourceRequest R; R.Request(StrandData.ClusterCullingBulkData, true /*bWait*/, true /*bFillBulkdata*/); }
 	}
 
 	if (!Ar.IsCooking() || !StrandData.bIsCookedOut)
@@ -2363,13 +2363,8 @@ bool UGroomAsset::CacheStrandsData(uint32 GroupIndex, FString& OutDerivedDataKey
 	auto FillDrivedDataKey = [&DerivedDataKey, &Name](FHairGroupPlatformData& In)
 	{
 		In.Strands.BulkData.DerivedDataKey = DerivedDataKey + FString(TEXT("_RestData"));
-		In.Strands.BulkData.Name = Name;
-	
 		In.Strands.InterpolationBulkData.DerivedDataKey = DerivedDataKey + FString(TEXT("_InterpolationData"));
-		In.Strands.InterpolationBulkData.Name = Name;
-	
 		In.Strands.ClusterCullingBulkData.DerivedDataKey = DerivedDataKey + FString(TEXT("_ClusterCullingData"));
-		In.Strands.ClusterCullingBulkData.Name = Name;
 	};
 
 	FHairGroupPlatformData& PlatformData = HairGroupsPlatformData[GroupIndex];
@@ -2428,7 +2423,6 @@ bool UGroomAsset::CacheStrandsData(uint32 GroupIndex, FString& OutDerivedDataKey
 				InternalSerializeGuide(Ar, this, PlatformData.Guides);
 				InternalSerializeStrand(Ar, this, PlatformData.Strands, true/*Header*/, false/*Data*/);
 
-
 				FRequestOwner AsyncOwner(EPriority::Normal);
 				GetCache().PutValue({ {Name, HeaderKey, FValue::Compress(MakeSharedBufferFromArray(MoveTemp(WriteData)))} }, AsyncOwner);
 				AsyncOwner.KeepAlive();
@@ -2436,44 +2430,32 @@ bool UGroomAsset::CacheStrandsData(uint32 GroupIndex, FString& OutDerivedDataKey
 
 			// Data (Rest)
 			{
-				TArray<uint8> WriteData;
-				FMemoryWriter Ar(WriteData, /*bIsPersistent*/ true);
-
-				PlatformData.Strands.BulkData.SerializeData(Ar, this);
-				
-				const FCacheKey DataKey = ConvertLegacyCacheKey(PlatformData.Strands.BulkData.DerivedDataKey);
+				TArray<FCachePutValueRequest> Out;
+				PlatformData.Strands.BulkData.Write_DDC(this, Out);
 
 				FRequestOwner AsyncOwner(EPriority::Normal);
-				GetCache().PutValue({ {Name, DataKey, FValue::Compress(MakeSharedBufferFromArray(MoveTemp(WriteData)))} }, AsyncOwner);
+				GetCache().PutValue(Out, AsyncOwner);
 				AsyncOwner.KeepAlive();
 			}
 
 			// Data (Interpolation)
-			if (PlatformData.Strands.InterpolationBulkData.HasData())
+			if (PlatformData.Strands.InterpolationBulkData.GetResourceCount() > 0)
 			{
-				TArray<uint8> WriteData;
-				FMemoryWriter Ar(WriteData, /*bIsPersistent*/ true);
-
-				PlatformData.Strands.InterpolationBulkData.SerializeData(Ar, this);
-				
-				const FCacheKey DataKey = ConvertLegacyCacheKey(PlatformData.Strands.InterpolationBulkData.DerivedDataKey);
+				TArray<FCachePutValueRequest> Out;
+				PlatformData.Strands.InterpolationBulkData.Write_DDC(this, Out);
 
 				FRequestOwner AsyncOwner(EPriority::Normal);
-				GetCache().PutValue({ {Name, DataKey, FValue::Compress(MakeSharedBufferFromArray(MoveTemp(WriteData)))} }, AsyncOwner);
+				GetCache().PutValue(Out, AsyncOwner);
 				AsyncOwner.KeepAlive();
 			}
 
 			// Data (cluster Culling)
 			{
-				TArray<uint8> WriteData;
-				FMemoryWriter Ar(WriteData, /*bIsPersistent*/ true);
-
-				PlatformData.Strands.ClusterCullingBulkData.SerializeData(Ar, this);
-				
-				const FCacheKey DataKey = ConvertLegacyCacheKey(PlatformData.Strands.ClusterCullingBulkData.DerivedDataKey);
+				TArray<FCachePutValueRequest> Out;
+				PlatformData.Strands.ClusterCullingBulkData.Write_DDC(this, Out);
 
 				FRequestOwner AsyncOwner(EPriority::Normal);
-				GetCache().PutValue({ {Name, DataKey, FValue::Compress(MakeSharedBufferFromArray(MoveTemp(WriteData)))} }, AsyncOwner);
+				GetCache().PutValue(Out, AsyncOwner);
 				AsyncOwner.KeepAlive();
 			}
 		}
diff --git a/Engine/Plugins/Runtime/HairStrands/Source/HairStrandsCore/Private/GroomBuilder.cpp b/Engine/Plugins/Runtime/HairStrands/Source/HairStrandsCore/Private/GroomBuilder.cpp
--- a/Engine/Plugins/Runtime/HairStrands/Source/HairStrandsCore/Private/GroomBuilder.cpp
+++ b/Engine/Plugins/Runtime/HairStrands/Source/HairStrandsCore/Private/GroomBuilder.cpp
@@ -40,7 +40,7 @@ static FAutoConsoleVariableRef CVarHairGroupIndexBuilder_MaxVoxelResolution(TEXT
 
 FString FGroomBuilder::GetVersion()
 {
-	return TEXT("v8q");
+	return TEXT("v8r9");
 }
 
 namespace FHairStrandsDecimation
@@ -150,6 +150,12 @@ namespace HairStrandsBuilder
 		Out.Unlock();
 	}
 
+	template<typename TFormatType>
+	void CopyToBulkData(FHairBulkContainer& Out, const TArray<typename TFormatType::Type>& Data)
+	{
+		CopyToBulkData<TFormatType>(Out.Data, Data);
+	}
+
 	/** Build the packed datas for gpu rendering/simulation */
 	void BuildRenderData(const FHairStrandsDatas& HairStrands, const TArray<uint8>& RandomSeeds, FHairStrandsBulkData& OutBulkData)
 	{
diff --git a/Engine/Plugins/Runtime/HairStrands/Source/HairStrandsCore/Private/GroomResources.cpp b/Engine/Plugins/Runtime/HairStrands/Source/HairStrandsCore/Private/GroomResources.cpp
--- a/Engine/Plugins/Runtime/HairStrands/Source/HairStrandsCore/Private/GroomResources.cpp
+++ b/Engine/Plugins/Runtime/HairStrands/Source/HairStrandsCore/Private/GroomResources.cpp
@@ -99,36 +99,73 @@ const TCHAR* ToHairResourceDebugName(const TCHAR* In, FHairResourceName& InDebug
 // IO/DDC request
 
 #if !WITH_EDITORONLY_DATA
-bool FHairResourceRequest::IsNone() const 									{ return IORequest.IsNone(); }
-bool FHairResourceRequest::IsCompleted() const								{ return IORequest.IsCompleted(); }
-void FHairResourceRequest::Request(FHairStrandsBulkCommon& In, bool bWait)	{ In.Request(IORequest); }
+bool FHairResourceRequest::IsNone() const 		{ return IORequest.IsNone(); }
+bool FHairResourceRequest::IsCompleted() const	{ return IORequest.IsCompleted(); }
+void FHairResourceRequest::Request(FHairStrandsBulkCommon& In, bool bWait, bool bFillBulkdata)
+{ 
+	if (In.GetResourceCount() != 0)
+	{
+		In.Read_IO(IORequest);
+		if (bWait)
+		{
+			IORequest.Wait();
+		}
+	}
+}
 #else
-bool FHairResourceRequest::IsNone() const 									{ return Status == None; }
-bool FHairResourceRequest::IsCompleted() const								{ return Status == Completed; }
-void FHairResourceRequest::Request(FHairStrandsBulkCommon& In, bool bWait)
+bool FHairResourceRequest::IsNone() const 		{ return Chunks.Num() == 0; }
+bool FHairResourceRequest::IsCompleted() const
+{ 
+	for (const FHairChunkRequest& Chunk : Chunks)
+	{
+		if (Chunk.Status != FHairChunkRequest::Completed)
+		{
+			return false;
+		}
+	}
+	return true; 
+}
+
+// Request fullfil 2 use cases:
+// * Load DDC data and upload them to GPU
+// * Load DDC data and store them into bulkdata for serialization
+void FHairResourceRequest::Request(FHairStrandsBulkCommon& In, bool bWait, bool bFillBulkdata)
 {
-	if (!In.HasData())
+	if (In.GetResourceCount() == 0)
 	{
-		Status = Completed;
 		return;
 	}
 
-	Status = Pending;
-
 	using namespace UE::DerivedData;
-	DDCRequestOwner = MakeUnique<FRequestOwner>(bWait ? UE::DerivedData::EPriority::Blocking : UE::DerivedData::EPriority::Normal);
-	const FCacheKey DataKey = ConvertLegacyCacheKey(In.DerivedDataKey);
-	const FSharedString Name = MakeStringView(In.Name);
+	TArray<FCacheGetChunkRequest> Requests;
+	In.Read_DDC(Requests);
 
-	UE::DerivedData::GetCache().GetValue({ {Name, DataKey} }, *DDCRequestOwner,
-	[this, &In](FCacheGetValueResponse&& Response)
+	// Add trackin data
+	Chunks.SetNum(0, false);
+	Chunks.Reserve(Requests.Num());
+	for (FCacheGetChunkRequest& R : Requests)
+	{
+		FHairChunkRequest& Chunk = Chunks.AddDefaulted_GetRef();
+		Chunk.Offset = 0;
+		Chunk.Size = 0;
+		Chunk.Status = FHairChunkRequest::EStatus::Pending;
+		Chunk.Container = (FHairBulkContainer*)R.UserData;
+		R.UserData = (uint64)&Chunk;
+	}
+
+	DDCRequestOwner = MakeUnique<FRequestOwner>(bWait ? UE::DerivedData::EPriority::Blocking : UE::DerivedData::EPriority::Normal); // <= Move this onto the resource (Buffer, in order to ensure no race condition)
+
+	//FRequestBarrier Barrier(*DDCRequestOwner);	// This is a critical section on the owner. It does not constrain ordering
+	GetCache().GetChunks(Requests, *DDCRequestOwner,
+	[this, &In](FCacheGetChunkResponse&& Response)
 	{
 		if (Response.Status == UE::DerivedData::EStatus::Ok)
 		{
-			FSharedBuffer Data = Response.Value.GetData().Decompress();
-			FMemoryReaderView Reader(Data, true /*bIsPersistent*/);
-			In.SerializeData(Reader, nullptr /*Owner*/);
-			Status = Completed;
+			FHairChunkRequest& Chunk = *(FHairChunkRequest*)Response.UserData;
+			Chunk.Data 	 = MoveTemp(Response.RawData);
+			Chunk.Offset = Response.RawOffset;
+			Chunk.Size 	 = Response.RawSize;
+			Chunk.Status = FHairChunkRequest::Completed;
 		}
 		else
 		{
@@ -136,10 +173,33 @@ void FHairResourceRequest::Request(FHairStrandsBulkCommon& In, bool bWait)
 		}
 	});
 
+	// TODO remove this
+	DDCRequestOwner->Wait();
+
 	if (bWait)
 	{
 		DDCRequestOwner->Wait();
-		check(Status == Completed);
+		check(IsCompleted());
+	}
+
+	// Fill in bytebulkdata
+	if (bFillBulkdata)
+	{		
+		for (FHairChunkRequest& Chunk : Chunks)
+		{
+			check(Chunk.Status == FHairChunkRequest::Completed);
+			if (Chunk.Container)
+			{
+				FByteBulkData& BulkData = Chunk.Container->Data;
+
+				// The buffer is then stored into bulk data
+				BulkData.Lock(LOCK_READ_WRITE);
+				void* DstData = BulkData.Realloc(Chunk.Size);
+				FMemory::Memcpy(DstData, Chunk.GetData(), Chunk.Size);
+				BulkData.Unlock();
+				Chunk.Release();
+			}
+		}
 	}
 }
 #endif
@@ -212,6 +272,29 @@ void CreateBuffer(const TArray<typename FormatType::Type>& InData, FHairCardsVer
 /////////////////////////////////////////////////////////////////////////////////////////
 // RDG buffers utils 
 
+static FRDGBufferDesc ApplyUsage(FRDGBufferDesc In, EHairResourceUsageType InUsage)
+{
+	if (InUsage != EHairResourceUsageType::Dynamic)
+	{
+		In.Usage = In.Usage & (~BUF_UnorderedAccess);
+	}
+	return In;
+}
+
+inline bool ReleaseAfterUse() { return GHairStrandsBulkData_ReleaseAfterUse > 0; }
+
+inline void InternalSetBulkDataFlags(FByteBulkData& In)
+{
+	// Unloading of the bulk data is only supported on cooked build, as we can reload the data from the file/archieve
+#if !WITH_EDITORONLY_DATA
+	const bool bReleaseCPUData = ReleaseAfterUse();
+	if (bReleaseCPUData)
+	{
+		In.SetBulkDataFlags(BULKDATA_SingleUse);
+	}
+#endif
+}
+
 static FRDGBufferRef InternalCreateVertexBuffer(
 	FRDGBuilder& GraphBuilder,
 	const TCHAR* Name,
@@ -239,7 +322,7 @@ template<typename FormatType>
 void InternalCreateVertexBufferRDG_FromBulkData(FRDGBuilder& GraphBuilder, FByteBulkData& InBulkData, uint32 InDataCount, FRDGExternalBuffer& Out, const TCHAR* DebugName, const FName& OwnerName, EHairResourceUsageType UsageType)
 {
 	const uint32 InDataCount_Check = InBulkData.GetBulkDataSize() / sizeof(typename FormatType::BulkType);
-	check(InDataCount_Check == InDataCount);
+	check(InDataCount_Check >= InDataCount);
 
 	const uint32 DataSizeInBytes = FormatType::SizeInByte * InDataCount;
 	if (DataSizeInBytes == 0)
@@ -248,25 +331,14 @@ void InternalCreateVertexBufferRDG_FromBulkData(FRDGBuilder& GraphBuilder, FByte
 		return;
 	}
 
-	FRDGBufferDesc Desc = FRDGBufferDesc::CreateBufferDesc(FormatType::SizeInByte, InDataCount);
-	if (UsageType != EHairResourceUsageType::Dynamic)
-	{
-		Desc.Usage = Desc.Usage & (~BUF_UnorderedAccess);
-	}
+	const FRDGBufferDesc Desc = ApplyUsage(FRDGBufferDesc::CreateBufferDesc(FormatType::SizeInByte, InDataCount), UsageType);
 
-	// Unloading of the bulk data is only supported on cooked build, as we can reload the data from the file/archieve
-	#if !WITH_EDITORONLY_DATA
-	const bool bReleaseCPUData = GHairStrandsBulkData_ReleaseAfterUse > 0;
-	if (bReleaseCPUData)
-	{
-		InBulkData.SetBulkDataFlags(BULKDATA_SingleUse);
-	}
-	#endif
+	InternalSetBulkDataFlags(InBulkData);
 
 	const typename FormatType::BulkType* BulkData = (const typename FormatType::BulkType*)InBulkData.Lock(LOCK_READ_ONLY);
 	FRDGBufferRef Buffer = nullptr;
 	#if !WITH_EDITORONLY_DATA
-	if (bReleaseCPUData)
+	if (ReleaseAfterUse())
 	{
 		checkf(EnumHasAnyFlags(Desc.Usage, EBufferUsageFlags::VertexBuffer), TEXT("CreateVertexBuffer called with an FRDGBufferDesc underlying type that is not 'VertexBuffer'. Buffer: %s"), DebugName);
 		Buffer = GraphBuilder.CreateBuffer(Desc, DebugName, ERDGBufferFlags::MultiFrame);
@@ -302,25 +374,14 @@ void InternalCreateByteAddressBufferRDG_FromBulkData(FRDGBuilder& GraphBuilder,
 		return;
 	}
 
-	FRDGBufferDesc Desc = FRDGBufferDesc::CreateByteAddressDesc(DataSizeInBytes);
-	if (UsageType != EHairResourceUsageType::Dynamic)
-	{
-		Desc.Usage = Desc.Usage & (~BUF_UnorderedAccess);
-	}
+	const FRDGBufferDesc Desc = ApplyUsage(FRDGBufferDesc::CreateByteAddressDesc(DataSizeInBytes), UsageType);
 
-	// Unloading of the bulk data is only supported on cooked build, as we can reload the data from the file/archieve
-#if !WITH_EDITORONLY_DATA
-	const bool bReleaseCPUData = GHairStrandsBulkData_ReleaseAfterUse > 0;
-	if (bReleaseCPUData)
-	{
-		InBulkData.SetBulkDataFlags(BULKDATA_SingleUse);
-	}
-#endif
+	InternalSetBulkDataFlags(InBulkData);
 
 	const uint8* BulkData = (const uint8*)InBulkData.Lock(LOCK_READ_ONLY);
 	FRDGBufferRef Buffer = nullptr;
 #if !WITH_EDITORONLY_DATA
-	if (bReleaseCPUData)
+	if (ReleaseAfterUse())
 	{		
 		Buffer = GraphBuilder.CreateBuffer(Desc, DebugName, ERDGBufferFlags::MultiFrame);
 		Buffer->SetOwnerName(OwnerName);
@@ -362,11 +423,7 @@ void InternalCreateVertexBufferRDG(FRDGBuilder& GraphBuilder, const typename For
 		return;
 	}
 
-	FRDGBufferDesc Desc = FRDGBufferDesc::CreateBufferDesc(FormatType::SizeInByte, InDataCount);
-	if (UsageType != EHairResourceUsageType::Dynamic)
-	{
-		Desc.Usage = Desc.Usage & (~BUF_UnorderedAccess);
-	}
+	const FRDGBufferDesc Desc = ApplyUsage(FRDGBufferDesc::CreateBufferDesc(FormatType::SizeInByte, InDataCount), UsageType);
 	Buffer = InternalCreateVertexBuffer(
 		GraphBuilder,
 		DebugName,
@@ -393,11 +450,7 @@ void InternalCreateVertexBufferRDG(FRDGBuilder& GraphBuilder, const TArray<typen
 		return;
 	}
 
-	FRDGBufferDesc Desc = FRDGBufferDesc::CreateBufferDesc(FormatType::SizeInByte, InData.Num());
-	if (UsageType != EHairResourceUsageType::Dynamic)
-	{
-		Desc.Usage = Desc.Usage & (~BUF_UnorderedAccess);
-	}
+	const FRDGBufferDesc Desc = ApplyUsage(FRDGBufferDesc::CreateBufferDesc(FormatType::SizeInByte, InData.Num()), UsageType);
 	Buffer = InternalCreateVertexBuffer(
 		GraphBuilder,
 		DebugName,
@@ -424,11 +477,8 @@ void InternalCreateVertexBufferRDG(FRDGBuilder& GraphBuilder, const TArray<DataT
 		return;
 	}
 
-	FRDGBufferDesc Desc = FRDGBufferDesc::CreateBufferDesc(sizeof(DataType), InData.Num());
-	if (UsageType != EHairResourceUsageType::Dynamic)
-	{
-		Desc.Usage = Desc.Usage & (~BUF_UnorderedAccess);
-	}
+	const FRDGBufferDesc Desc = ApplyUsage(FRDGBufferDesc::CreateBufferDesc(sizeof(DataType), InData.Num()), UsageType);
+
 	Buffer = InternalCreateVertexBuffer(
 		GraphBuilder,
 		DebugName,
@@ -485,7 +535,7 @@ void InternalCreateStructuredBufferRDG(FRDGBuilder& GraphBuilder, uint32 DataCou
 		return;
 	}
 
-	FRDGBufferDesc Desc = FRDGBufferDesc::CreateStructuredDesc(FormatType::SizeInByte, DataCount);
+	FRDGBufferDesc Desc = ApplyUsage(FRDGBufferDesc::CreateStructuredDesc(FormatType::SizeInByte, DataCount), Usage);
 	FRDGBufferRef Buffer = GraphBuilder.CreateBuffer(Desc, DebugName, ERDGBufferFlags::MultiFrame);
 	Buffer->SetOwnerName(OwnerName);
 	AddClearUAVPass(GraphBuilder, GraphBuilder.CreateUAV(Buffer, FormatType::Format), 0u);
@@ -498,7 +548,7 @@ void InternalCreateStructuredBufferRDG_FromBulkData(FRDGBuilder& GraphBuilder, F
 {
 	const uint32 InSizeInByte = sizeof(typename FormatType::Type);
 	const uint32 DataSizeInBytes = InSizeInByte * InDataCount;
-	check(InBulkData.GetBulkDataSize() == DataSizeInBytes);
+	check(InBulkData.GetBulkDataSize() >= DataSizeInBytes);
 
 	if (DataSizeInBytes == 0)
 	{
@@ -506,20 +556,9 @@ void InternalCreateStructuredBufferRDG_FromBulkData(FRDGBuilder& GraphBuilder, F
 		return;
 	}
 
-	FRDGBufferDesc Desc = FRDGBufferDesc::CreateStructuredDesc(InSizeInByte, InDataCount);
-	if (UsageType != EHairResourceUsageType::Dynamic)
-	{
-		Desc.Usage = Desc.Usage & (~BUF_UnorderedAccess);
-	}
+	const FRDGBufferDesc Desc = ApplyUsage(FRDGBufferDesc::CreateStructuredDesc(InSizeInByte, InDataCount), UsageType);
 
-	// Unloading of the bulk data is only supported on cooked build, as we can reload the data from the file/archieve
-	#if !WITH_EDITORONLY_DATA
-	const bool bReleaseCPUData = GHairStrandsBulkData_ReleaseAfterUse > 0;
-	if (bReleaseCPUData)
-	{
-		InBulkData.SetBulkDataFlags(BULKDATA_SingleUse);
-	}
-	#endif
+	InternalSetBulkDataFlags(InBulkData);
 
 	const typename FormatType::Type* Data = (const typename FormatType::Type*)InBulkData.Lock(LOCK_READ_ONLY);
 	FRDGBufferRef Buffer = GraphBuilder.CreateBuffer(Desc, DebugName, ERDGBufferFlags::MultiFrame);
@@ -527,7 +566,7 @@ void InternalCreateStructuredBufferRDG_FromBulkData(FRDGBuilder& GraphBuilder, F
 	if (Data && DataSizeInBytes)
 	{
 		#if !WITH_EDITORONLY_DATA
-		if (bReleaseCPUData)
+		if (ReleaseAfterUse())
 		{
 			GraphBuilder.QueueBufferUpload(Buffer, Data, DataSizeInBytes, [&InBulkData](const void* Ptr) { InBulkData.Unlock(); });
 		}
@@ -929,7 +968,7 @@ bool FHairStrandsRestResource::InternalIsDataLoaded()
 {
 	if (BulkDataRequest.IsNone())
 	{
-		BulkDataRequest.Request(BulkData);
+		BulkDataRequest.Request(BulkData, true, true);
 	}
 	return BulkDataRequest.IsCompleted();
 }
@@ -937,6 +976,9 @@ bool FHairStrandsRestResource::InternalIsDataLoaded()
 bool IsHairStrandsContinousLODEnabled();
 void FHairStrandsRestResource::InternalAllocate(FRDGBuilder& GraphBuilder)
 {
+	// When the request is delete, the DDC request becomes cancelled. 
+	// So the request need to be completed by now
+	//check(BulkDataRequest.IsCompleted());
 	BulkDataRequest = FHairResourceRequest();
 
 	const uint32 PointCount = BulkData.GetNumPoints();
@@ -944,11 +986,11 @@ void FHairStrandsRestResource::InternalAllocate(FRDGBuilder& GraphBuilder)
 
 	// 1. Lock data, which force the loading data from files (on non-editor build/cooked data). These data are then uploaded to the GPU
 	// 2. A local copy is done by the buffer uploader. This copy is discarded once the uploading is done.
-	InternalCreateVertexBufferRDG_FromBulkData<FHairStrandsPositionFormat>(GraphBuilder, BulkData.Data.Positions, PointCount, PositionBuffer, ToHairResourceDebugName(HAIRSTRANDS_RESOUCE_NAME(CurveType, Hair.StrandsRest_PositionBuffer), ResourceName), OwnerName, EHairResourceUsageType::Static);
-	InternalCreateByteAddressBufferRDG_FromBulkData(GraphBuilder, BulkData.Data.CurveAttributes, CurveAttributeBuffer, ToHairResourceDebugName(HAIRSTRANDS_RESOUCE_NAME(CurveType, Hair.StrandsRest_CurveAttributeBuffer), ResourceName), OwnerName, EHairResourceUsageType::Static);
+	InternalCreateVertexBufferRDG_FromBulkData<FHairStrandsPositionFormat>(GraphBuilder, BulkData.Data.Positions.Data, PointCount, PositionBuffer, ToHairResourceDebugName(HAIRSTRANDS_RESOUCE_NAME(CurveType, Hair.StrandsRest_PositionBuffer), ResourceName), OwnerName, EHairResourceUsageType::Static);
+	InternalCreateByteAddressBufferRDG_FromBulkData(GraphBuilder, BulkData.Data.CurveAttributes.Data, CurveAttributeBuffer, ToHairResourceDebugName(HAIRSTRANDS_RESOUCE_NAME(CurveType, Hair.StrandsRest_CurveAttributeBuffer), ResourceName), OwnerName, EHairResourceUsageType::Static);
 	if (!!(BulkData.Header.Flags & FHairStrandsBulkData::DataFlags_HasPointAttribute))
 	{
-		InternalCreateByteAddressBufferRDG_FromBulkData(GraphBuilder, BulkData.Data.PointAttributes, PointAttributeBuffer, ToHairResourceDebugName(HAIRSTRANDS_RESOUCE_NAME(CurveType, Hair.StrandsRest_PointAttributeBuffer), ResourceName), OwnerName, EHairResourceUsageType::Static);
+		InternalCreateByteAddressBufferRDG_FromBulkData(GraphBuilder, BulkData.Data.PointAttributes.Data, PointAttributeBuffer, ToHairResourceDebugName(HAIRSTRANDS_RESOUCE_NAME(CurveType, Hair.StrandsRest_PointAttributeBuffer), ResourceName), OwnerName, EHairResourceUsageType::Static);
 	}
 	else
 	{
@@ -957,14 +999,14 @@ void FHairStrandsRestResource::InternalAllocate(FRDGBuilder& GraphBuilder)
 		InternalCreateVertexBufferRDG(GraphBuilder, DummyAttribute, EPixelFormat::PF_R32_UINT, PointAttributeBuffer, ToHairResourceDebugName(HAIRSTRANDS_RESOUCE_NAME(CurveType, Hair.StrandsRest_PointAttributeBuffer), ResourceName), OwnerName, EHairResourceUsageType::Static, true /*bIsByteAddressBuffer*/);
 		GraphBuilder.UseExternalAccessMode(Register(GraphBuilder, PointAttributeBuffer, ERDGImportedBufferFlags::CreateSRV).Buffer, ERHIAccess::SRVMask);
 	}
-	InternalCreateVertexBufferRDG_FromBulkData<FHairStrandsCurveFormat>(GraphBuilder, BulkData.Data.Curves, CurveCount, CurveBuffer, ToHairResourceDebugName(HAIRSTRANDS_RESOUCE_NAME(CurveType, Hair.StrandsRest_CurveBuffer), ResourceName), OwnerName, EHairResourceUsageType::Static);
+	InternalCreateVertexBufferRDG_FromBulkData<FHairStrandsCurveFormat>(GraphBuilder, BulkData.Data.Curves.Data, CurveCount, CurveBuffer, ToHairResourceDebugName(HAIRSTRANDS_RESOUCE_NAME(CurveType, Hair.StrandsRest_CurveBuffer), ResourceName), OwnerName, EHairResourceUsageType::Static);
 	if (!!(BulkData.Header.Flags & FHairStrandsBulkData::DataFlags_Has16bitsCurveIndex))
 	{
-		InternalCreateVertexBufferRDG_FromBulkData<FHairStrandsPointToCurveFormat16>(GraphBuilder, BulkData.Data.PointToCurve, PointCount, PointToCurveBuffer, ToHairResourceDebugName(HAIRSTRANDS_RESOUCE_NAME(CurveType, Hair.StrandsRest_PointToCurveBuffer), ResourceName), OwnerName, EHairResourceUsageType::Static);
+		InternalCreateVertexBufferRDG_FromBulkData<FHairStrandsPointToCurveFormat16>(GraphBuilder, BulkData.Data.PointToCurve.Data, PointCount, PointToCurveBuffer, ToHairResourceDebugName(HAIRSTRANDS_RESOUCE_NAME(CurveType, Hair.StrandsRest_PointToCurveBuffer), ResourceName), OwnerName, EHairResourceUsageType::Static);
 	}
 	else
 	{
-		InternalCreateVertexBufferRDG_FromBulkData<FHairStrandsPointToCurveFormat32>(GraphBuilder, BulkData.Data.PointToCurve, PointCount, PointToCurveBuffer, ToHairResourceDebugName(HAIRSTRANDS_RESOUCE_NAME(CurveType, Hair.StrandsRest_PointToCurveBuffer), ResourceName), OwnerName, EHairResourceUsageType::Static);
+		InternalCreateVertexBufferRDG_FromBulkData<FHairStrandsPointToCurveFormat32>(GraphBuilder, BulkData.Data.PointToCurve.Data, PointCount, PointToCurveBuffer, ToHairResourceDebugName(HAIRSTRANDS_RESOUCE_NAME(CurveType, Hair.StrandsRest_PointToCurveBuffer), ResourceName), OwnerName, EHairResourceUsageType::Static);
 	}
 
 	TArray<FVector4f> RestOffset;
@@ -978,9 +1020,9 @@ void FHairStrandsRestResource::InternalAllocate(FRDGBuilder& GraphBuilder)
 		CurveData.SetNum(CurveCount);
 		uint32* Data = (uint32*)CurveData.GetData();
 
-		const uint32* ReadOnlyData = (const uint32*)BulkData.Data.Curves.Lock(LOCK_READ_ONLY);
+		const uint32* ReadOnlyData = (const uint32*)BulkData.Data.Curves.Data.Lock(LOCK_READ_ONLY);
 		memcpy(Data, ReadOnlyData, sizeof(uint32) * CurveCount);
-		BulkData.Data.Curves.Unlock();
+		BulkData.Data.Curves.Data.Unlock();
 	}
 }
 
@@ -1106,16 +1148,10 @@ void FHairStrandsClusterCullingBulkData::Reset()
 	Data.PackedClusterInfos.RemoveBulkData();
 
 	// Reset the bulk byte buffer to ensure the (serialize) data size is reset to 0
-	Data.ClusterLODInfos 	= FByteBulkData();
-	Data.VertexToClusterIds = FByteBulkData();
-	Data.ClusterVertexIds 	= FByteBulkData();
-	Data.PackedClusterInfos = FByteBulkData();
-}
-
-void FHairStrandsClusterCullingBulkData::Serialize(FArchive& Ar, UObject* Owner)
-{
-	SerializeHeader(Ar, Owner);
-	SerializeData(Ar, Owner);
+	Data.ClusterLODInfos 	= FHairBulkContainer();
+	Data.VertexToClusterIds = FHairBulkContainer();
+	Data.ClusterVertexIds 	= FHairBulkContainer();
+	Data.PackedClusterInfos = FHairBulkContainer();
 }
 
 void FHairStrandsClusterCullingBulkData::SerializeHeader(FArchive& Ar, UObject* Owner)
@@ -1139,77 +1175,47 @@ void FHairStrandsClusterCullingBulkData::SerializeHeader(FArchive& Ar, UObject*
 	}
 }
 
-void FHairStrandsClusterCullingBulkData::SerializeData(FArchive& Ar, UObject* Owner)
+uint32 FHairStrandsClusterCullingBulkData::GetResourceCount() const
 {
-	const int32 ChunkIndex = 0;
-	bool bAttemptFileMapping = false;
-
-	if (Ar.IsSaving())
-	{
-		const uint32 BulkFlags = BULKDATA_Force_NOT_InlinePayload;
-		Data.ClusterLODInfos.SetBulkDataFlags(BulkFlags);
-		Data.VertexToClusterIds.SetBulkDataFlags(BulkFlags);
-		Data.ClusterVertexIds.SetBulkDataFlags(BulkFlags);
-		Data.PackedClusterInfos.SetBulkDataFlags(BulkFlags);
-	}
+	return 4;
+}
 
+void FHairStrandsClusterCullingBulkData::GetResources(FHairStrandsBulkCommon::FQuery & Out)
+{
 	if (Header.ClusterLODCount)
 	{
-		Data.ClusterLODInfos.Serialize(Ar, Owner, ChunkIndex, bAttemptFileMapping);
+		Out.Add(Data.ClusterLODInfos, TEXT("_ClusterLODInfos"));
 	}
 
 	if (Header.VertexCount)
 	{
-		Data.VertexToClusterIds.Serialize(Ar, Owner, ChunkIndex, bAttemptFileMapping);
+		Out.Add(Data.VertexToClusterIds, TEXT("_VertexToClusterIds"));
 	}
 
 	if (Header.VertexLODCount)
 	{
-		Data.ClusterVertexIds.Serialize(Ar, Owner, ChunkIndex, bAttemptFileMapping);
+		Out.Add(Data.ClusterVertexIds, TEXT("_ClusterVertexIds"));
 	}
 
 	if (Header.ClusterCount)
 	{
-		Data.PackedClusterInfos.Serialize(Ar, Owner, ChunkIndex, bAttemptFileMapping);
+		Out.Add(Data.PackedClusterInfos, TEXT("_PackedClusterInfos"));
 	}
 
-	if (GHairStrandsBulkData_Validation > 0 && Ar.IsSaving())
+	if (GHairStrandsBulkData_Validation > 0 && Out.Type == FHairStrandsBulkCommon::FQuery::WriteDDC || Out.Type == FHairStrandsBulkCommon::FQuery::ReadWriteIO)
 	{
 		Validate(true);
 	}
 }
 
-void FHairStrandsClusterCullingBulkData::Request(FBulkDataBatchRequest& In)
-{
-	FBulkDataBatchRequest::FBatchBuilder Batch = FBulkDataBatchRequest::NewBatch(4);
-	if (Header.ClusterLODCount)
-	{
-		Batch.Read(Data.ClusterLODInfos);
-	}
-
-	if (Header.VertexCount)
-	{
-		Batch.Read(Data.VertexToClusterIds);
-	}
-
-	if (Header.VertexLODCount)
-	{
-		Batch.Read(Data.ClusterVertexIds);
-	}
-
-	if (Header.ClusterCount)
-	{
-		Batch.Read(Data.PackedClusterInfos);
-	}
-	Batch.Issue(In);
-}
-
 void FHairStrandsClusterCullingBulkData::Validate(bool bIsSaving)
 {
 	if (Header.ClusterCount == 0)
+	{
 		return;
+	}
 
-	const FHairClusterInfo::Packed* Datas = (const FHairClusterInfo::Packed*)Data.PackedClusterInfos.Lock(LOCK_READ_ONLY);
+	const FHairClusterInfo::Packed* Datas = (const FHairClusterInfo::Packed*)Data.PackedClusterInfos.Data.Lock(LOCK_READ_ONLY);
 	
 	// Simple heuristic to check if the data are valid
 	const uint32 MaxCount = FMath::Min(Header.ClusterCount, 128u);
@@ -1225,7 +1231,7 @@ void FHairStrandsClusterCullingBulkData::Validate(bool bIsSaving)
 		UE_LOG(LogHairStrands, Error, TEXT("[Groom/DDC] Strands - Invalid ClusterCullingBulkData when %s bulk data - %s"), bIsSaving ? TEXT("Saving") : TEXT("Loading"), *DebugName);
 	}
 
-	Data.PackedClusterInfos.Unlock();
+	Data.PackedClusterInfos.Data.Unlock();
 }
 
 /////////////////////////////////////////////////////////////////////////////////////////
@@ -1241,13 +1247,14 @@ bool FHairStrandsClusterCullingResource::InternalIsDataLoaded()
 {
 	if (BulkDataRequest.IsNone())
 	{
-		BulkDataRequest.Request(BulkData);
+		BulkDataRequest.Request(BulkData, true, true);
 	}
 	return BulkDataRequest.IsCompleted();
 }
 
 void FHairStrandsClusterCullingResource::InternalAllocate(FRDGBuilder& GraphBuilder)
 {
+	check(BulkDataRequest.IsCompleted());
 	BulkDataRequest = FHairResourceRequest();
 
 	if (GHairStrandsBulkData_Validation > 0)
@@ -1255,11 +1262,11 @@ void FHairStrandsClusterCullingResource::InternalAllocate(FRDGBuilder& GraphBuil
 		BulkData.Validate(false);
 	}
 
-	InternalCreateStructuredBufferRDG_FromBulkData<FHairClusterInfoFormat>(GraphBuilder, BulkData.Data.PackedClusterInfos, BulkData.Header.ClusterCount, ClusterInfoBuffer, ToHairResourceDebugName(TEXT("Hair.StrandsClusterCulling_ClusterInfoBuffer"), ResourceName), OwnerName, EHairResourceUsageType::Static);
-	InternalCreateStructuredBufferRDG_FromBulkData<FHairClusterLODInfoFormat>(GraphBuilder, BulkData.Data.ClusterLODInfos, BulkData.Header.ClusterLODCount, ClusterLODInfoBuffer, ToHairResourceDebugName(TEXT("Hair.StrandsClusterCulling_ClusterLODInfoBuffer"), ResourceName), OwnerName, EHairResourceUsageType::Static);
+	InternalCreateStructuredBufferRDG_FromBulkData<FHairClusterInfoFormat>(GraphBuilder, BulkData.Data.PackedClusterInfos.Data, BulkData.Header.ClusterCount, ClusterInfoBuffer, ToHairResourceDebugName(TEXT("Hair.StrandsClusterCulling_ClusterInfoBuffer"), ResourceName), OwnerName, EHairResourceUsageType::Static);
+	InternalCreateStructuredBufferRDG_FromBulkData<FHairClusterLODInfoFormat>(GraphBuilder, BulkData.Data.ClusterLODInfos.Data, BulkData.Header.ClusterLODCount, ClusterLODInfoBuffer, ToHairResourceDebugName(TEXT("Hair.StrandsClusterCulling_ClusterLODInfoBuffer"), ResourceName), OwnerName, EHairResourceUsageType::Static);
 
-	InternalCreateVertexBufferRDG_FromBulkData<FHairClusterIndexFormat>(GraphBuilder, BulkData.Data.VertexToClusterIds, BulkData.Header.VertexCount, VertexToClusterIdBuffer, ToHairResourceDebugName(TEXT("Hair.StrandsClusterCulling_VertexToClusterIds"), ResourceName), OwnerName, EHairResourceUsageType::Static);
-	InternalCreateVertexBufferRDG_FromBulkData<FHairClusterIndexFormat>(GraphBuilder, BulkData.Data.ClusterVertexIds, BulkData.Header.VertexLODCount, ClusterVertexIdBuffer, ToHairResourceDebugName(TEXT("Hair.StrandsClusterCulling_ClusterVertexIds"), ResourceName), OwnerName, EHairResourceUsageType::Static);
+	InternalCreateVertexBufferRDG_FromBulkData<FHairClusterIndexFormat>(GraphBuilder, BulkData.Data.VertexToClusterIds.Data, BulkData.Header.VertexCount, VertexToClusterIdBuffer, ToHairResourceDebugName(TEXT("Hair.StrandsClusterCulling_VertexToClusterIds"), ResourceName), OwnerName, EHairResourceUsageType::Static);
+	InternalCreateVertexBufferRDG_FromBulkData<FHairClusterIndexFormat>(GraphBuilder, BulkData.Data.ClusterVertexIds.Data, BulkData.Header.VertexLODCount, ClusterVertexIdBuffer, ToHairResourceDebugName(TEXT("Hair.StrandsClusterCulling_ClusterVertexIds"), ResourceName), OwnerName, EHairResourceUsageType::Static);
 }
 
 void FHairStrandsClusterCullingResource::InternalRelease()
@@ -1696,7 +1703,7 @@ bool FHairStrandsInterpolationResource::InternalIsDataLoaded()
 {
 	if (BulkDataRequest.IsNone())
 	{
-		BulkDataRequest.Request(BulkData);
+		BulkDataRequest.Request(BulkData, true, true);
 	}
 	return BulkDataRequest.IsCompleted();
 }
@@ -1704,8 +1711,8 @@ bool FHairStrandsInterpolationResource::InternalIsDataLoaded()
 void FHairStrandsInterpolationResource::InternalAllocate(FRDGBuilder& GraphBuilder)
 {
 	BulkDataRequest = FHairResourceRequest();
-	InternalCreateByteAddressBufferRDG_FromBulkData(GraphBuilder, BulkData.Data.Interpolation, InterpolationBuffer, ToHairResourceDebugName(TEXT("Hair.StrandsInterpolation_InterpolationBuffer"), ResourceName), OwnerName, EHairResourceUsageType::Static);
-	InternalCreateVertexBufferRDG_FromBulkData<FHairStrandsRootIndexFormat>(GraphBuilder, BulkData.Data.SimRootPointIndex, BulkData.Header.SimPointCount, SimRootPointIndexBuffer, ToHairResourceDebugName(TEXT("Hair.StrandsInterpolation_SimRootPointIndex"), ResourceName), OwnerName, EHairResourceUsageType::Static);
+	InternalCreateByteAddressBufferRDG_FromBulkData(GraphBuilder, BulkData.Data.Interpolation.Data, InterpolationBuffer, ToHairResourceDebugName(TEXT("Hair.StrandsInterpolation_InterpolationBuffer"), ResourceName), OwnerName, EHairResourceUsageType::Static);
+	InternalCreateVertexBufferRDG_FromBulkData<FHairStrandsRootIndexFormat>(GraphBuilder, BulkData.Data.SimRootPointIndex.Data, BulkData.Header.SimPointCount, SimRootPointIndexBuffer, ToHairResourceDebugName(TEXT("Hair.StrandsInterpolation_SimRootPointIndex"), ResourceName), OwnerName, EHairResourceUsageType::Static);
 }
 
 void FHairStrandsInterpolationResource::InternalRelease()
diff --git a/Engine/Plugins/Runtime/HairStrands/Source/HairStrandsCore/Private/HairStrandsDatas.cpp b/Engine/Plugins/Runtime/HairStrands/Source/HairStrandsCore/Private/HairStrandsDatas.cpp
--- a/Engine/Plugins/Runtime/HairStrands/Source/HairStrandsCore/Private/HairStrandsDatas.cpp
+++ b/Engine/Plugins/Runtime/HairStrands/Source/HairStrandsCore/Private/HairStrandsDatas.cpp
@@ -5,6 +5,11 @@
 #include "UObject/UE5ReleaseStreamObjectVersion.h"
 #include "HairAttributes.h"
 
+#if WITH_EDITOR
+#include "DerivedDataCache.h"
+#include "DerivedDataRequestOwner.h"
+#endif
+
 void FHairStrandsInterpolationDatas::SetNum(const uint32 NumCurves)
 {
 	PointsSimCurvesVertexWeights.SetNum(NumCurves);
@@ -111,6 +116,121 @@ FArchive& operator<<(FArchive& Ar, FPackedHairVertex& Vertex)
 	return Ar;
 }
 
+namespace HairStrands
+{
+#if WITH_EDITOR
+	const UE::DerivedData::FValueId HairStrandsValueId = UE::DerivedData::FValueId::FromName("HairStrandsStreamingData");
+#endif
+}
+
+void FHairStrandsBulkCommon::Write_DDC(UObject* Owner, TArray<UE::DerivedData::FCachePutValueRequest>& Out)
+{
+#if WITH_EDITORONLY_DATA
+	FHairStrandsBulkCommon::FQuery Q;
+	Q.Type = FHairStrandsBulkCommon::FQuery::WriteDDC;
+	Q.OutWriteDDC = &Out;
+	Q.DerivedDataKey = &DerivedDataKey;
+	Q.Owner = Owner;
+	GetResources(Q);
+#endif
+}
+
+void FHairStrandsBulkCommon::Read_DDC(TArray<UE::DerivedData::FCacheGetChunkRequest>& Out)
+{
+#if WITH_EDITORONLY_DATA
+	FHairStrandsBulkCommon::FQuery Q;
+	Q.Type = FHairStrandsBulkCommon::FQuery::ReadDDC;
+	Q.OutReadDDC = &Out;
+	Q.DerivedDataKey = &DerivedDataKey;
+	GetResources(Q);
+#endif
+}
+
+void FHairStrandsBulkCommon::Read_IO(FBulkDataBatchRequest& Out)
+{
+	FBulkDataBatchRequest::FBatchBuilder Batch = Out.NewBatch(GetResourceCount());
+
+	FHairStrandsBulkCommon::FQuery Q;
+	Q.Type = FHairStrandsBulkCommon::FQuery::ReadIO;
+	Q.OutReadIO = &Batch;
+	GetResources(Q);
+	Q.OutReadIO->Issue(Out);
+}
+void FHairStrandsBulkCommon::Write_IO(FArchive& Ar, UObject* Owner)
+{
+	GetResourceVersion(Ar);
+
+	FHairStrandsBulkCommon::FQuery Q;
+	Q.Type = FHairStrandsBulkCommon::FQuery::ReadWriteIO;
+	Q.OutWriteIO = &Ar;
+	Q.Owner = Owner;
+	GetResources(Q);
+}
+
+void FHairStrandsBulkCommon::Serialize(FArchive& Ar, UObject* Owner)
+{
+	SerializeHeader(Ar, Owner);
+	SerializeData(Ar, Owner);
+}
+
+void FHairStrandsBulkCommon::SerializeData(FArchive& Ar, UObject* Owner)
+{
+	Write_IO(Ar, Owner);
+}
+
+void FHairStrandsBulkCommon::FQuery::Add(FHairBulkContainer& In, const TCHAR* InSuffix) 
+{
+	check(Type != None);
+#if WITH_EDITORONLY_DATA
+	if (Type == WriteDDC)
+	{
+		const int64 DataSizeInByte = In.Data.GetBulkDataSize();
+		TArray<uint8> WriteData;
+		WriteData.SetNum(DataSizeInByte);
+		FMemory::Memcpy(WriteData.GetData(), In.Data.Lock(LOCK_READ_ONLY), DataSizeInByte);
+		In.Data.Unlock();
+
+		using namespace UE::DerivedData;
+		FCachePutValueRequest& Out = OutWriteDDC->AddDefaulted_GetRef();
+		Out.Name 	= Owner->GetPathName();
+		Out.Key 	= ConvertLegacyCacheKey(*DerivedDataKey + InSuffix);
+		Out.Value 	= FValue::Compress(MakeSharedBufferFromArray(MoveTemp(WriteData)));
+		Out.Policy 	= ECachePolicy::Default;
+		Out.UserData= 0;
+	}
+	else if (Type == ReadDDC)
+	{
+		check(OutReadDDC);
+		using namespace UE::DerivedData;
+		FCacheGetChunkRequest& Out = OutReadDDC->AddDefaulted_GetRef();
+		Out.Id			= FValueId::Null; // HairStrands::HairStrandsValueId : This is only needed for cache record, not cache value.
+		Out.Key			= ConvertLegacyCacheKey(*DerivedDataKey + InSuffix);
+		Out.RawOffset	= 0; 			//In.Upload.Offset;
+		Out.RawSize		= MAX_uint64; 	//In.Upload.Size;
+		Out.RawHash		= FIoHash();
+		Out.UserData	= (uint64)&In;
+	}
+	else 
+#endif
+	if (Type == ReadIO)
+	{
+		check(OutReadIO);
+		OutReadIO->Read(In.Data);
+	}
+	else
+	{
+		check(Type == ReadWriteIO)
+		check(OutWriteIO);
+
+		if (OutWriteIO->IsSaving())
+		{
+			const uint32 BulkFlags = BULKDATA_Force_NOT_InlinePayload;
+			In.Data.SetBulkDataFlags(BulkFlags);
+		}
+		In.Data.Serialize(*OutWriteIO, Owner, 0/*ChunkIndex*/, false /*bAttemptFileMapping*/);
+	}
+}
+
 void FHairStrandsInterpolationBulkData::Reset()
 {
 	Header.Flags = 0;
@@ -122,14 +242,8 @@ void FHairStrandsInterpolationBulkData::Reset()
 	Data.SimRootPointIndex.RemoveBulkData();
 
 	// Reset the bulk byte buffer to ensure the (serialize) data size is reset to 0
-	Data.Interpolation		= FByteBulkData();
-	Data.SimRootPointIndex	= FByteBulkData();
-}
-
-void FHairStrandsInterpolationBulkData::Serialize(FArchive& Ar, UObject* Owner)
-{
-	SerializeHeader(Ar, Owner);
-	SerializeData(Ar, Owner);
+	Data.Interpolation		= FHairBulkContainer();
+	Data.SimRootPointIndex	= FHairBulkContainer();
 }
 
 void FHairStrandsInterpolationBulkData::SerializeHeader(FArchive& Ar, UObject* Owner)
@@ -139,48 +253,25 @@ void FHairStrandsInterpolationBulkData::SerializeHeader(FArchive& Ar, UObject* O
 	Ar << Header.SimPointCount;
 }
 
-void FHairStrandsInterpolationBulkData::SerializeData(FArchive& Ar, UObject* Owner)
+uint32 FHairStrandsInterpolationBulkData::GetResourceCount() const
 {
-	static_assert(sizeof(FHairStrandsRootIndexFormat::BulkType) == sizeof(FHairStrandsRootIndexFormat::Type));
-
-	if (!!(Header.Flags & DataFlags_HasData))
-	{
-		const int32 ChunkIndex = 0;
-		bool bAttemptFileMapping = false;
-		Data.Interpolation.Serialize(Ar, Owner, ChunkIndex, bAttemptFileMapping);
-		Data.SimRootPointIndex.Serialize(Ar, Owner, ChunkIndex, bAttemptFileMapping);
-	}
+	return (Header.Flags & DataFlags_HasData) ? 2 : 0;
 }
 
-bool FHairStrandsInterpolationBulkData::HasData() const
+void FHairStrandsInterpolationBulkData::GetResources(FHairStrandsBulkCommon::FQuery& Out)
 {
-	return Header.Flags & DataFlags_HasData;
-}
+	static_assert(sizeof(FHairStrandsRootIndexFormat::BulkType) == sizeof(FHairStrandsRootIndexFormat::Type));
 
-void FHairStrandsInterpolationBulkData::Request(FBulkDataBatchRequest& InRequest)
-{
-	if (!(Header.Flags & DataFlags_HasData))
+	if (Header.Flags & DataFlags_HasData)
 	{
-		return;
+		Out.Add(Data.Interpolation, TEXT("_Interpolation"));
+		Out.Add(Data.SimRootPointIndex, TEXT("_SimRootPointIndex"));
 	}
-
-	check(InRequest.IsNone());
-	FBulkDataBatchRequest::FBatchBuilder Batch = FBulkDataBatchRequest::NewBatch(2);
-	Batch.Read(Data.Interpolation);
-	Batch.Read(Data.SimRootPointIndex);
-	Batch.Issue(InRequest);
-}
-
-void FHairStrandsBulkData::Serialize(FArchive& Ar, UObject* Owner)
-{
-	SerializeHeader(Ar, Owner);
-	SerializeData(Ar, Owner);
 }
 
 void FHairStrandsBulkData::SerializeHeader(FArchive& Ar, UObject* Owner)
 {
-	Ar.UsingCustomVersion(FReleaseObjectVersion::GUID);
-	Ar.UsingCustomVersion(FUE5ReleaseStreamObjectVersion::GUID);
+	GetResourceVersion(Ar);
 
 	Ar << Header.CurveCount;
 	Ar << Header.PointCount;
@@ -200,8 +291,18 @@ void FHairStrandsBulkData::SerializeHeader(FArchive& Ar, UObject* Owner)
 	Ar << Header.ImportedAttributeFlags;
 }
 
+void FHairStrandsBulkData::GetResourceVersion(FArchive& Ar) const
+{
+	Ar.UsingCustomVersion(FReleaseObjectVersion::GUID);
+	Ar.UsingCustomVersion(FUE5ReleaseStreamObjectVersion::GUID);
+}
 
-void FHairStrandsBulkData::SerializeData(FArchive& Ar, UObject* Owner)
+uint32 FHairStrandsBulkData::GetResourceCount() const
+{
+	return 5;
+}
+
+void FHairStrandsBulkData::GetResources(FHairStrandsBulkCommon::FQuery& Out)
 {
 	static_assert(sizeof(FHairStrandsPositionFormat::BulkType) == sizeof(FHairStrandsPositionFormat::Type));
 	static_assert(sizeof(FHairStrandsAttributeFormat::BulkType) == sizeof(FHairStrandsAttributeFormat::Type));
@@ -209,56 +310,16 @@ void FHairStrandsBulkData::SerializeData(FArchive& Ar, UObject* Owner)
 	static_assert(sizeof(FHairStrandsPointToCurveFormat32::BulkType) == sizeof(FHairStrandsPointToCurveFormat32::Type));
 	static_assert(sizeof(FHairStrandsRootIndexFormat::BulkType) == sizeof(FHairStrandsRootIndexFormat::Type)); 
 
-	Ar.UsingCustomVersion(FReleaseObjectVersion::GUID);
-	Ar.UsingCustomVersion(FUE5ReleaseStreamObjectVersion::GUID);
-	  
-	// Forced not inline means the bulk data won't automatically be loaded when we deserialize
-	// but only when we explicitly take action to load it
-	if (Ar.IsSaving())
-	{
-		const uint32 BulkFlags = BULKDATA_Force_NOT_InlinePayload;
-		Data.Positions.SetBulkDataFlags(BulkFlags);
-		Data.CurveAttributes.SetBulkDataFlags(BulkFlags);
-		if (Header.Flags & DataFlags_HasPointAttribute)
-		{
-			Data.PointAttributes.SetBulkDataFlags(BulkFlags);
-		}
-		Data.PointToCurve.SetBulkDataFlags(BulkFlags);
-		Data.Curves.SetBulkDataFlags(BulkFlags);
-	}
-
 	if (!!(Header.Flags & DataFlags_HasData))
 	{
-		const int32 ChunkIndex = 0;
-		bool bAttemptFileMapping = false;
-
-		Data.Positions.Serialize(Ar, Owner, ChunkIndex, bAttemptFileMapping);
-		if (Header.Flags & DataFlags_HasPointAttribute)
-		{
-			Data.PointAttributes.Serialize(Ar, Owner, ChunkIndex, bAttemptFileMapping);
-		}
-		Data.CurveAttributes.Serialize(Ar, Owner, ChunkIndex, bAttemptFileMapping);
-		Data.PointToCurve.Serialize(Ar, Owner, ChunkIndex, bAttemptFileMapping);
-		Data.Curves.Serialize(Ar, Owner, ChunkIndex, bAttemptFileMapping);
-	}
-}
-
-void FHairStrandsBulkData::Request(FBulkDataBatchRequest& In)
-{
-	if (!!(Header.Flags & DataFlags_HasData))
-	{
-		check(In.IsNone());
-
-		FBulkDataBatchRequest::FBatchBuilder Batch = FBulkDataBatchRequest::NewBatch(5);
-		Batch.Read(Data.Positions);
+		Out.Add(Data.Positions, TEXT("_Positions"));
+		Out.Add(Data.CurveAttributes, TEXT("_CurveAttributes"));
 		if (Header.Flags & DataFlags_HasPointAttribute)
 		{
-			Batch.Read(Data.PointAttributes);
+			Out.Add(Data.PointAttributes, TEXT("_PointAttributes"));
 		}
-		Batch.Read(Data.CurveAttributes);
-		Batch.Read(Data.PointToCurve);
-		Batch.Read(Data.Curves);
-		Batch.Issue(In);
+		Out.Add(Data.PointToCurve, TEXT("_PointToCurve"));
+		Out.Add(Data.Curves, TEXT("_Curves"));
 	}
 }
 
@@ -286,11 +347,11 @@ void FHairStrandsBulkData::Reset()
 	Data.Curves.RemoveBulkData();
 
 	// Reset the bulk byte buffer to ensure the (serialize) data size is reset to 0
-	Data.Positions 		= FByteBulkData();
-	Data.CurveAttributes= FByteBulkData();
-	Data.PointAttributes= FByteBulkData();
-	Data.PointToCurve	= FByteBulkData();
-	Data.Curves			= FByteBulkData();
+	Data.Positions 		= FHairBulkContainer();
+	Data.CurveAttributes= FHairBulkContainer();
+	Data.PointAttributes= FHairBulkContainer();
+	Data.PointToCurve	= FHairBulkContainer();
+	Data.Curves			= FHairBulkContainer();
 }
 
 void FHairStrandsDatas::Reset()
diff --git a/Engine/Plugins/Runtime/HairStrands/Source/HairStrandsCore/Public/GroomResources.h b/Engine/Plugins/Runtime/HairStrands/Source/HairStrandsCore/Public/GroomResources.h
--- a/Engine/Plugins/Runtime/HairStrands/Source/HairStrandsCore/Public/GroomResources.h
+++ b/Engine/Plugins/Runtime/HairStrands/Source/HairStrandsCore/Public/GroomResources.h
@@ -75,9 +75,22 @@ FORCEINLINE bool operator! (EHairResourceStatus A) { return static_cast<uint8>(A
 
 EHairResourceLoadingType GetHairResourceLoadingType(EHairGeometryType InGeometryType, int32 InLODIndex);
 
+struct FHairChunkRequest
+{
+	enum EStatus { None, Pending, Completed, Failed };
+	FSharedBuffer Data;
+	uint32 Offset = 0;
+	uint32 Size = 0;
+	EStatus Status = EStatus::None;
+	FHairBulkContainer* Container = nullptr;
+
+	const uint8* GetData() const { check(Status == Completed); return (const uint8*)Data.GetData(); }
+	void Release() { Data.Reset(); }
+};
+
 struct FHairResourceRequest
 {
-	void Request(FHairStrandsBulkCommon& In, bool bWait=false);
+	void Request(FHairStrandsBulkCommon& In, bool bWait=false, bool bFillBulkData=false);
 	bool IsNone() const ;
 	bool IsCompleted() const;
 
@@ -86,9 +99,8 @@ struct FHairResourceRequest
 	FBulkDataBatchRequest IORequest;
 #else
 	// DDC
-	enum EStatus { None, Pending, Completed };
-	EStatus Status = EStatus::None;
 	FString PathName;
+	TArray<FHairChunkRequest> Chunks;
 	TUniquePtr<UE::DerivedData::FRequestOwner> DDCRequestOwner;
 #endif
 };
diff --git a/Engine/Plugins/Runtime/HairStrands/Source/HairStrandsCore/Public/HairStrandsDatas.h b/Engine/Plugins/Runtime/HairStrands/Source/HairStrandsCore/Public/HairStrandsDatas.h
--- a/Engine/Plugins/Runtime/HairStrands/Source/HairStrandsCore/Public/HairStrandsDatas.h
+++ b/Engine/Plugins/Runtime/HairStrands/Source/HairStrandsCore/Public/HairStrandsDatas.h
@@ -13,6 +13,11 @@
 DECLARE_LOG_CATEGORY_EXTERN(LogHairStrands, Log, All);
 
 enum class EHairAttribute : uint8;
+namespace UE::DerivedData
+{
+	struct FCacheGetChunkRequest;
+	struct FCachePutValueRequest;
+}
 
 struct FPackedHairVertex
 {
@@ -227,18 +232,55 @@ struct FHairStrandsRootUtils
 	static float	 PackUVsToFloat(const FVector2f& UV);
 };
 
+struct FHairBulkContainer
+{
+	// Streaming
+	TBitArray<> AvailableChunks;
+	FByteBulkData Data;
+
+	// Forward BulkData functions
+	bool IsBulkDataLoaded() const 		{ return Data.IsBulkDataLoaded(); }
+	int64 GetBulkDataSize() const 		{ return Data.GetBulkDataSize(); }
+	FString GetDebugName() const 		{ return Data.GetDebugName(); }
+	void SetBulkDataFlags(uint32 Flags) { Data.SetBulkDataFlags(Flags); }
+	void RemoveBulkData() 				{ Data.RemoveBulkData();  }
+	void Serialize(FArchive& Ar, UObject* Owner, int32 ChunkIndex, bool bAttemptFileMapping) { Data.Serialize(Ar, Owner, ChunkIndex, bAttemptFileMapping); }
+};
 
 struct HAIRSTRANDSCORE_API FHairStrandsBulkCommon
 {
 	virtual ~FHairStrandsBulkCommon() { }
+	void Serialize(FArchive& Ar, UObject* Owner);
 	virtual void SerializeHeader(FArchive& Ar, UObject* Owner) = 0;
-	virtual void SerializeData(FArchive& Ar, UObject* Owner) = 0;
-	virtual void Request(FBulkDataBatchRequest& In) = 0;
-	virtual bool HasData() const { return true; }
+	void SerializeData(FArchive& Ar, UObject* Owner);
+
+	void Write_DDC(UObject* Owner, TArray<UE::DerivedData::FCachePutValueRequest>& Out);
+	void Read_DDC(TArray<UE::DerivedData::FCacheGetChunkRequest>& Out);
+	void Write_IO(FArchive& Ar, UObject* Owner);
+	void Read_IO(FBulkDataBatchRequest& Out);
+
+	struct FQuery
+	{
+		void Add(FHairBulkContainer& In, TCHAR* InSuffix);
+
+		enum EQueryType { None, ReadDDC, WriteDDC, ReadIO, ReadWriteIO /* aka regular Serialize() */};
+		EQueryType Type = None;
+		FBulkDataBatchRequest::FBatchBuilder* 			OutReadIO = nullptr;
+		FArchive*										OutWriteIO = nullptr;
+	#if WITH_EDITORONLY_DATA
+		TArray<UE::DerivedData::FCacheGetChunkRequest>*	OutReadDDC = nullptr;
+		TArray<UE::DerivedData::FCachePutValueRequest>*	OutWriteDDC = nullptr;
+		FString* DerivedDataKey = nullptr;
+	#endif
+		UObject* Owner = nullptr; 
+	};
+
+	virtual uint32 GetResourceCount() const = 0;
+	virtual void GetResources(FQuery& Out) = 0;
+	virtual void GetResourceVersion(FArchive& Ar) const {}
 
 #if WITH_EDITORONLY_DATA
 	// Transient Name/DDCkey for streaming
-	FString Name;
 	FString DerivedDataKey;
 #endif
 };
@@ -282,11 +324,10 @@ struct HAIRSTRANDSCORE_API FHairStrandsInterpolationBulkData : FHairStrandsBulkC
 	};
 
 	void Reset();
-	void Serialize(FArchive& Ar, UObject* Owner);
 	virtual void SerializeHeader(FArchive& Ar, UObject* Owner) override;
-	virtual void SerializeData(FArchive& Ar, UObject* Owner) override;
-	virtual void Request(FBulkDataBatchRequest& In) override;
-	virtual bool HasData() const override;
+
+	virtual uint32 GetResourceCount() const override;
+	virtual void GetResources(FQuery& Out) override;
 	uint32 GetPointCount() const { return Header.PointCount; };
 
 	struct FHeader
@@ -298,8 +339,8 @@ struct HAIRSTRANDSCORE_API FHairStrandsInterpolationBulkData : FHairStrandsBulkC
 
 	struct FData
 	{
-		FByteBulkData Interpolation;	// FHairStrandsInterpolationFormat  - Per-rendering-vertex interpolation data (closest guides, weight factors, ...). Data for a 1 or 3 guide(s))
-		FByteBulkData SimRootPointIndex;// FHairStrandsRootIndexFormat      - Per-rendering-vertex index of the sim-root vertex
+		FHairBulkContainer Interpolation;		// FHairStrandsInterpolationFormat  - Per-rendering-vertex interpolation data (closest guides, weight factors, ...). Data for a 1 or 3 guide(s))
+		FHairBulkContainer SimRootPointIndex;	// FHairStrandsRootIndexFormat      - Per-rendering-vertex index of the sim-root vertex
 	} Data;
 };
 
@@ -429,10 +470,10 @@ struct HAIRSTRANDSCORE_API FHairStrandsBulkData : FHairStrandsBulkCommon
 		DataFlags_HasPointAttribute = 4,	// Contains point attribute data.
 	};
 
-	void Serialize(FArchive& Ar, UObject* Owner);
 	virtual void SerializeHeader(FArchive& Ar, UObject* Owner) override;
-	virtual void SerializeData(FArchive& Ar, UObject* Owner) override;
-	virtual void Request(FBulkDataBatchRequest& In) override;
+	virtual uint32 GetResourceCount() const override;
+	virtual void GetResources(FQuery& Out) override;
+	virtual void GetResourceVersion(FArchive& Ar) const override;
 
 	bool IsValid() const { return Header.CurveCount > 0 && Header.PointCount > 0; }
 	void Reset();
@@ -462,11 +503,11 @@ struct HAIRSTRANDSCORE_API FHairStrandsBulkData : FHairStrandsBulkCommon
 
 	struct FData
 	{
-		FByteBulkData Positions;		// Size = PointCount
-		FByteBulkData CurveAttributes;	// Size = y*CurveCount (depends on the per-curve stored attributes)
-		FByteBulkData PointAttributes;	// Size = x*PointCount (depends on the per-point stored attributes)
-		FByteBulkData PointToCurve; 	// Size = PointCount
-		FByteBulkData Curves;			// Size = CurveCount
+		FHairBulkContainer Positions;		// Size = PointCount
+		FHairBulkContainer CurveAttributes;	// Size = y*CurveCount (depends on the per-curve stored attributes)
+		FHairBulkContainer PointAttributes;	// Size = x*PointCount (depends on the per-point stored attributes)
+		FHairBulkContainer PointToCurve; 	// Size = PointCount
+		FHairBulkContainer Curves;			// Size = CurveCount
 	} Data;
 };
 
@@ -746,10 +787,11 @@ struct HAIRSTRANDSCORE_API FHairStrandsClusterCullingBulkData : FHairStrandsBulk
 {
 	FHairStrandsClusterCullingBulkData();
 	void Reset();
-	void Serialize(FArchive& Ar, UObject* Owner);
+
 	virtual void SerializeHeader(FArchive& Ar, UObject* Owner) override;
-	virtual void SerializeData(FArchive& Ar, UObject* Owner) override;
-	virtual void Request(FBulkDataBatchRequest& In) override;
+	virtual uint32 GetResourceCount() const override;
+	virtual void GetResources(FQuery& Out) override;
+
 	bool IsValid() const { return Header.ClusterCount > 0 && Header.VertexCount > 0; }
 	void Validate(bool bIsSaving);
 
@@ -772,13 +814,11 @@ struct HAIRSTRANDSCORE_API FHairStrandsClusterCullingBulkData : FHairStrandsBulk
 
 	struct FData
 	{
-
 		/* LOD info for the various clusters for LOD management on GPU */
-		FByteBulkData	PackedClusterInfos;		// Size - ClusterCount
-		FByteBulkData	ClusterLODInfos;		// Size - ClusterLODCount
-		FByteBulkData	VertexToClusterIds;		// Size - VertexCount
-		FByteBulkData	ClusterVertexIds;		// Size - VertexLODCount
+		FHairBulkContainer	PackedClusterInfos;		// Size - ClusterCount
+		FHairBulkContainer	ClusterLODInfos;		// Size - ClusterLODCount
+		FHairBulkContainer	VertexToClusterIds;		// Size - VertexCount
+		FHairBulkContainer	ClusterVertexIds;		// Size - VertexLODCount
 	} Data;
-
 };
 
