Fix CIS issue.

diff --git a/Engine/Plugins/Runtime/HairStrands/Source/HairStrandsCore/Public/HairStrandsDatas.h b/Engine/Plugins/Runtime/HairStrands/Source/HairStrandsCore/Public/HairStrandsDatas.h
--- a/Engine/Plugins/Runtime/HairStrands/Source/HairStrandsCore/Public/HairStrandsDatas.h
+++ b/Engine/Plugins/Runtime/HairStrands/Source/HairStrandsCore/Public/HairStrandsDatas.h
@@ -261,7 +261,7 @@ struct HAIRSTRANDSCORE_API FHairStrandsBulkCommon
 
 	struct FQuery
 	{
-		void Add(FHairBulkContainer& In, TCHAR* InSuffix);
+		void Add(FHairBulkContainer& In, const TCHAR* InSuffix);
 
 		enum EQueryType { None, ReadDDC, WriteDDC, ReadIO, ReadWriteIO /* aka regular Serialize() */};
 		EQueryType Type = None;
