Minor clean up & fix.

diff --git a/Engine/Plugins/Runtime/HairStrands/Source/HairStrandsCore/Private/GroomResources.cpp b/Engine/Plugins/Runtime/HairStrands/Source/HairStrandsCore/Private/GroomResources.cpp
--- a/Engine/Plugins/Runtime/HairStrands/Source/HairStrandsCore/Private/GroomResources.cpp
+++ b/Engine/Plugins/Runtime/HairStrands/Source/HairStrandsCore/Private/GroomResources.cpp
@@ -50,6 +50,11 @@ static int32 GHairStrandsDebugVoxel_MaxSegmentPerVoxel = 2048;
 static FAutoConsoleVariableRef CVarHairStrandsDebugVoxel_WorldSize(TEXT("r.HairStrands.DebugData.VoxelSize"), GHairStrandsDebugVoxel_WorldSize, TEXT("Voxel size use for creating debug data."));
 static FAutoConsoleVariableRef CVarHairStrandsDebugVoxel_MaxSegmentPerVoxel(TEXT("r.HairStrands.DebugData.MaxSegmentPerVoxel"), GHairStrandsDebugVoxel_MaxSegmentPerVoxel, TEXT("Max number of segments per Voxel size when creating debug data."));
 
+bool ValidateHairBulkData()
+{
+	return GHairStrandsBulkData_Validation > 0;
+}
+
 ////////////////////////////////////////////////////////////////////////////////////
 
 EHairResourceLoadingType GetHairResourceLoadingType(EHairGeometryType InGeometryType, int32 InLODIndex)
@@ -408,9 +413,6 @@ void InternalCreateByteAddressBufferRDG_FromBulkData(FRDGBuilder& GraphBuilder,
 	ConvertToExternalBufferWithViews(GraphBuilder, Buffer, Out, PF_Unknown);
 }
 
-
-
-
 template<typename FormatType>
 void InternalCreateVertexBufferRDG(FRDGBuilder& GraphBuilder, const typename FormatType::Type* InData, uint32 InDataCount, FRDGExternalBuffer& Out, const TCHAR* DebugName, const FName& OwnerName, EHairResourceUsageType UsageType, ERDGInitialDataFlags InitialDataFlags, bool bIsByteAddressBuffer=false)
 {
@@ -542,7 +544,6 @@ void InternalCreateStructuredBufferRDG(FRDGBuilder& GraphBuilder, uint32 DataCou
 	ConvertToExternalBufferWithViews(GraphBuilder, Buffer, Out, FormatType::Format);
 }
 
-
 template<typename FormatType>
 void InternalCreateStructuredBufferRDG_FromBulkData(FRDGBuilder& GraphBuilder, FByteBulkData& InBulkData, uint32 InDataCount, FRDGExternalBuffer& Out, const TCHAR* DebugName, const FName& OwnerName, EHairResourceUsageType UsageType)
 {
@@ -1111,129 +1112,6 @@ void FHairStrandsDeformedResource::InternalRelease()
 	DeformedOffsetBuffer[1].Release();
 }
 
-/////////////////////////////////////////////////////////////////////////////////////////
-// Cluster culling data
-
-FHairStrandsClusterCullingData::FHairStrandsClusterCullingData()
-{
-
-}
-
-void FHairStrandsClusterCullingData::Reset()
-{
-	*this = FHairStrandsClusterCullingData();
-}
-
-/////////////////////////////////////////////////////////////////////////////////////////
-// Cluster culling bulk data
-FHairStrandsClusterCullingBulkData::FHairStrandsClusterCullingBulkData()
-{
-
-}
-
-void FHairStrandsClusterCullingBulkData::Reset()
-{
-	Header.ClusterCount = 0;
-	Header.ClusterLODCount = 0;
-	Header.VertexCount = 0;
-	Header.VertexLODCount = 0;
-
-	Header.LODVisibility.Empty();
-	Header.CPULODScreenSize.Empty();
-	Header.LODInfos.Empty();
-
-	Data.ClusterLODInfos.RemoveBulkData();
-	Data.VertexToClusterIds.RemoveBulkData();
-	Data.ClusterVertexIds.RemoveBulkData();
-	Data.PackedClusterInfos.RemoveBulkData();
-
-	// Reset the bulk byte buffer to ensure the (serialize) data size is reset to 0
-	Data.ClusterLODInfos 	= FHairBulkContainer();
-	Data.VertexToClusterIds = FHairBulkContainer();
-	Data.ClusterVertexIds 	= FHairBulkContainer();
-	Data.PackedClusterInfos = FHairBulkContainer();
-}
-
-void FHairStrandsClusterCullingBulkData::SerializeHeader(FArchive& Ar, UObject* Owner)
-{
-	Ar << Header.ClusterCount;
-	Ar << Header.ClusterLODCount;
-	Ar << Header.VertexCount;
-	Ar << Header.VertexLODCount;
-	Ar << Header.LODVisibility;
-	Ar << Header.CPULODScreenSize;
-	uint32 LODInfosCount = Header.LODInfos.Num();
-	Ar << LODInfosCount;
-	if (Ar.IsLoading())
-	{
-		Header.LODInfos.SetNum(LODInfosCount);
-	}
-	for (uint32 It = 0; It < LODInfosCount; ++It)
-	{
-		Ar << Header.LODInfos[It].CurveCount;
-		Ar << Header.LODInfos[It].PointCount;
-	}
-}
-
-uint32 FHairStrandsClusterCullingBulkData::GetResourceCount() const
-{
-	return 4;
-}
-
-void FHairStrandsClusterCullingBulkData::GetResources(FHairStrandsBulkCommon::FQuery & Out)
-{
-	if (Header.ClusterLODCount)
-	{
-		Out.Add(Data.ClusterLODInfos, TEXT("_ClusterLODInfos"));
-	}
-
-	if (Header.VertexCount)
-	{
-		Out.Add(Data.VertexToClusterIds, TEXT("_VertexToClusterIds"));
-	}
-
-	if (Header.VertexLODCount)
-	{
-		Out.Add(Data.ClusterVertexIds, TEXT("_ClusterVertexIds"));
-	}
-
-	if (Header.ClusterCount)
-	{
-		Out.Add(Data.PackedClusterInfos, TEXT("_PackedClusterInfos"));
-	}
-
-	if (GHairStrandsBulkData_Validation > 0 && Out.Type == FHairStrandsBulkCommon::FQuery::WriteDDC || Out.Type == FHairStrandsBulkCommon::FQuery::ReadWriteIO)
-	{
-		Validate(true);
-	}
-}
-
-void FHairStrandsClusterCullingBulkData::Validate(bool bIsSaving)
-{
-	if (Header.ClusterCount == 0)
-	{
-		return;
-	}
-
-	const FHairClusterInfo::Packed* Datas = (const FHairClusterInfo::Packed*)Data.PackedClusterInfos.Data.Lock(LOCK_READ_ONLY);
-	
-	// Simple heuristic to check if the data are valid
-	const uint32 MaxCount = FMath::Min(Header.ClusterCount, 128u);
-	bool bIsValid = true;
-	for (uint32 It = 0; It < MaxCount; ++It)
-	{
-		bIsValid = bIsValid && Datas[It].LODCount <= 8;
-		if (!bIsValid) break;
-	}
-	if (!bIsValid)
-	{
-		FString DebugName = Data.ClusterLODInfos.GetDebugName();
-		UE_LOG(LogHairStrands, Error, TEXT("[Groom/DDC] Strands - Invalid ClusterCullingBulkData when %s bulk data - %s"), bIsSaving ? TEXT("Saving") : TEXT("Loading"), *DebugName);
-	}
-
-	Data.PackedClusterInfos.Data.Unlock();
-}
-
 /////////////////////////////////////////////////////////////////////////////////////////
 // Cluster culling resources
 FHairStrandsClusterCullingResource::FHairStrandsClusterCullingResource(FHairStrandsClusterCullingBulkData& InBulkData, const FHairResourceName& InResourceName, const FName& InOwnerName):
@@ -1257,7 +1135,7 @@ void FHairStrandsClusterCullingResource::InternalAllocate(FRDGBuilder& GraphBuil
 	check(BulkDataRequest.IsCompleted());
 	BulkDataRequest = FHairResourceRequest();
 
-	if (GHairStrandsBulkData_Validation > 0)
+	if (ValidateHairBulkData())
 	{
 		BulkData.Validate(false);
 	}
diff --git a/Engine/Plugins/Runtime/HairStrands/Source/HairStrandsCore/Private/HairStrandsDatas.cpp b/Engine/Plugins/Runtime/HairStrands/Source/HairStrandsCore/Private/HairStrandsDatas.cpp
--- a/Engine/Plugins/Runtime/HairStrands/Source/HairStrandsCore/Private/HairStrandsDatas.cpp
+++ b/Engine/Plugins/Runtime/HairStrands/Source/HairStrandsCore/Private/HairStrandsDatas.cpp
@@ -47,6 +47,14 @@
 	if (HasHairAttribute(InAttributes, EHairAttribute::AO))			{ PointsAO.SetNum(NumPoints); }
 }
 
+void FHairStrandsDatas::Reset()
+{
+	StrandsCurves.Reset();
+	StrandsPoints.Reset();
+	HairDensity = 1;
+	BoundingBox = FBox(EForceInit::ForceInit);
+}
+
 void FHairStrandsInterpolationDatas::Reset()
 {
 	PointsSimCurvesVertexWeights.Reset();
@@ -55,6 +63,11 @@
 	PointsSimCurvesIndex.Reset();
 }
 
+void FHairStrandsClusterCullingData::Reset()
+{
+	*this = FHairStrandsClusterCullingData();
+}
+
 void FHairStrandsCurves::Reset()
 {
 	CurvesOffset.Reset();
@@ -96,25 +109,8 @@
 	return MaxRadius;
 }
 
-FArchive& operator<<(FArchive& Ar, FVector4_16& Vertex)
-{
-	Ar << Vertex.X;
-	Ar << Vertex.Y;
-	Ar << Vertex.Z;
-	Ar << Vertex.W;
-	return Ar;
-}
+/////////////////////////////////////////////////////////////////////////////////////////
+// Common bulk data
 
-FArchive& operator<<(FArchive& Ar, FPackedHairVertex& Vertex)
-{
-	Ar << Vertex.X;
-	Ar << Vertex.Y;
-	Ar << Vertex.Z;
-	Ar << Vertex.PackedRadiusAndType;
-	Ar << Vertex.UCoord;
-
-	return Ar;
-}
-
 namespace HairStrands
 {
@@ -231,43 +227,8 @@
 	}
 }
 
-void FHairStrandsInterpolationBulkData::Reset()
-{
-	Header.Flags = 0;
-	Header.PointCount = 0;
-	Header.SimPointCount = 0;
-	
-	// Deallocate memory if needed
-	Data.Interpolation.RemoveBulkData();
-	Data.SimRootPointIndex.RemoveBulkData();
+/////////////////////////////////////////////////////////////////////////////////////////
+// Rest bulk data
 
-	// Reset the bulk byte buffer to ensure the (serialize) data size is reset to 0
-	Data.Interpolation		= FHairBulkContainer();
-	Data.SimRootPointIndex	= FHairBulkContainer();
-}
-
-void FHairStrandsInterpolationBulkData::SerializeHeader(FArchive& Ar, UObject* Owner)
-{
-	Ar << Header.Flags;
-	Ar << Header.PointCount;
-	Ar << Header.SimPointCount;
-}
-
-uint32 FHairStrandsInterpolationBulkData::GetResourceCount() const
-{
-	return (Header.Flags & DataFlags_HasData) ? 2 : 0;
-}
-
-void FHairStrandsInterpolationBulkData::GetResources(FHairStrandsBulkCommon::FQuery& Out)
-{
-	static_assert(sizeof(FHairStrandsRootIndexFormat::BulkType) == sizeof(FHairStrandsRootIndexFormat::Type));
-
-	if (Header.Flags & DataFlags_HasData)
-	{
-		Out.Add(Data.Interpolation, TEXT("_Interpolation"));
-		Out.Add(Data.SimRootPointIndex, TEXT("_SimRootPointIndex"));
-	}
-}
-
 void FHairStrandsBulkData::SerializeHeader(FArchive& Ar, UObject* Owner)
 {
@@ -354,10 +315,153 @@
 	Data.Curves			= FHairBulkContainer();
 }
 
-void FHairStrandsDatas::Reset()
+
+/////////////////////////////////////////////////////////////////////////////////////////
+// Interpolation bulk data
+
+void FHairStrandsInterpolationBulkData::Reset()
 {
-	StrandsCurves.Reset();
-	StrandsPoints.Reset();
-	HairDensity = 1;
-	BoundingBox = FBox(EForceInit::ForceInit);
+	Header.Flags = 0;
+	Header.PointCount = 0;
+	Header.SimPointCount = 0;
+	
+	// Deallocate memory if needed
+	Data.Interpolation.RemoveBulkData();
+	Data.SimRootPointIndex.RemoveBulkData();
+
+	// Reset the bulk byte buffer to ensure the (serialize) data size is reset to 0
+	Data.Interpolation		= FHairBulkContainer();
+	Data.SimRootPointIndex	= FHairBulkContainer();
 }
+
+void FHairStrandsInterpolationBulkData::SerializeHeader(FArchive& Ar, UObject* Owner)
+{
+	Ar << Header.Flags;
+	Ar << Header.PointCount;
+	Ar << Header.SimPointCount;
+}
+
+uint32 FHairStrandsInterpolationBulkData::GetResourceCount() const
+{
+	return (Header.Flags & DataFlags_HasData) ? 2 : 0;
+}
+
+void FHairStrandsInterpolationBulkData::GetResources(FHairStrandsBulkCommon::FQuery& Out)
+{
+	static_assert(sizeof(FHairStrandsRootIndexFormat::BulkType) == sizeof(FHairStrandsRootIndexFormat::Type));
+
+	if (Header.Flags & DataFlags_HasData)
+	{
+		Out.Add(Data.Interpolation, TEXT("_Interpolation"));
+		Out.Add(Data.SimRootPointIndex, TEXT("_SimRootPointIndex"));
+	}
+}
+
+
+/////////////////////////////////////////////////////////////////////////////////////////
+// Cluster culling bulk data
+
+void FHairStrandsClusterCullingBulkData::Reset()
+{
+	Header.ClusterCount = 0;
+	Header.ClusterLODCount = 0;
+	Header.VertexCount = 0;
+	Header.VertexLODCount = 0;
+
+	Header.LODVisibility.Empty();
+	Header.CPULODScreenSize.Empty();
+	Header.LODInfos.Empty();
+
+	Data.ClusterLODInfos.RemoveBulkData();
+	Data.VertexToClusterIds.RemoveBulkData();
+	Data.ClusterVertexIds.RemoveBulkData();
+	Data.PackedClusterInfos.RemoveBulkData();
+
+	// Reset the bulk byte buffer to ensure the (serialize) data size is reset to 0
+	Data.ClusterLODInfos 	= FHairBulkContainer();
+	Data.VertexToClusterIds = FHairBulkContainer();
+	Data.ClusterVertexIds 	= FHairBulkContainer();
+	Data.PackedClusterInfos = FHairBulkContainer();
+}
+
+void FHairStrandsClusterCullingBulkData::SerializeHeader(FArchive& Ar, UObject* Owner)
+{
+	Ar << Header.ClusterCount;
+	Ar << Header.ClusterLODCount;
+	Ar << Header.VertexCount;
+	Ar << Header.VertexLODCount;
+	Ar << Header.LODVisibility;
+	Ar << Header.CPULODScreenSize;
+	uint32 LODInfosCount = Header.LODInfos.Num();
+	Ar << LODInfosCount;
+	if (Ar.IsLoading())
+	{
+		Header.LODInfos.SetNum(LODInfosCount);
+	}
+	for (uint32 It = 0; It < LODInfosCount; ++It)
+	{
+		Ar << Header.LODInfos[It].CurveCount;
+		Ar << Header.LODInfos[It].PointCount;
+	}
+}
+
+uint32 FHairStrandsClusterCullingBulkData::GetResourceCount() const
+{
+	return 4;
+}
+
+bool ValidateHairBulkData();
+
+void FHairStrandsClusterCullingBulkData::GetResources(FHairStrandsBulkCommon::FQuery & Out)
+{
+	if (Header.ClusterLODCount)
+	{
+		Out.Add(Data.ClusterLODInfos, TEXT("_ClusterLODInfos"));
+	}
+
+	if (Header.VertexCount)
+	{
+		Out.Add(Data.VertexToClusterIds, TEXT("_VertexToClusterIds"));
+	}
+
+	if (Header.VertexLODCount)
+	{
+		Out.Add(Data.ClusterVertexIds, TEXT("_ClusterVertexIds"));
+	}
+
+	if (Header.ClusterCount)
+	{
+		Out.Add(Data.PackedClusterInfos, TEXT("_PackedClusterInfos"));
+	}
+
+	if (ValidateHairBulkData() && (Out.Type == FHairStrandsBulkCommon::FQuery::WriteDDC || Out.Type == FHairStrandsBulkCommon::FQuery::ReadWriteIO))
+	{
+		Validate(true);
+	}
+}
+
+void FHairStrandsClusterCullingBulkData::Validate(bool bIsSaving)
+{
+	if (Header.ClusterCount == 0)
+	{
+		return;
+	}
+
+	const FHairClusterInfo::Packed* Datas = (const FHairClusterInfo::Packed*)Data.PackedClusterInfos.Data.Lock(LOCK_READ_ONLY);
+	
+	// Simple heuristic to check if the data are valid
+	const uint32 MaxCount = FMath::Min(Header.ClusterCount, 128u);
+	bool bIsValid = true;
+	for (uint32 It = 0; It < MaxCount; ++It)
+	{
+		bIsValid = bIsValid && Datas[It].LODCount <= 8;
+		if (!bIsValid) break;
+	}
+	if (!bIsValid)
+	{
+		FString DebugName = Data.ClusterLODInfos.GetDebugName();
+		UE_LOG(LogHairStrands, Error, TEXT("[Groom/DDC] Strands - Invalid ClusterCullingBulkData when %s bulk data - %s"), bIsSaving ? TEXT("Saving") : TEXT("Loading"), *DebugName);
+	}
+
+	Data.PackedClusterInfos.Data.Unlock();
+}
diff --git a/Engine/Plugins/Runtime/HairStrands/Source/HairStrandsCore/Public/HairStrandsDatas.h b/Engine/Plugins/Runtime/HairStrands/Source/HairStrandsCore/Public/HairStrandsDatas.h
--- a/Engine/Plugins/Runtime/HairStrands/Source/HairStrandsCore/Public/HairStrandsDatas.h
+++ b/Engine/Plugins/Runtime/HairStrands/Source/HairStrandsCore/Public/HairStrandsDatas.h
@@ -759,7 +759,6 @@ struct FHairLODInfo
 
 struct HAIRSTRANDSCORE_API FHairStrandsClusterCullingData
 {
-	FHairStrandsClusterCullingData();
 	void Reset();
 	bool IsValid() const { return ClusterCount > 0 && VertexCount > 0; }
 
@@ -785,7 +784,6 @@ struct HAIRSTRANDSCORE_API FHairStrandsClusterCullingData
 
 struct HAIRSTRANDSCORE_API FHairStrandsClusterCullingBulkData : FHairStrandsBulkCommon
 {
-	FHairStrandsClusterCullingBulkData();
 	void Reset();
 
 	virtual void SerializeHeader(FArchive& Ar, UObject* Owner) override;
