Shared groom resource creation.

diff --git a/Engine/Plugins/Runtime/HairStrands/Source/HairStrandsCore/Private/GroomResources.cpp b/Engine/Plugins/Runtime/HairStrands/Source/HairStrandsCore/Private/GroomResources.cpp
--- a/Engine/Plugins/Runtime/HairStrands/Source/HairStrandsCore/Private/GroomResources.cpp
+++ b/Engine/Plugins/Runtime/HairStrands/Source/HairStrandsCore/Private/GroomResources.cpp
@@ -300,117 +300,80 @@ inline void InternalSetBulkDataFlags(FByteBulkData& In)
 #endif
 }
 
-static FRDGBufferRef InternalCreateVertexBuffer(
-	FRDGBuilder& GraphBuilder,
-	const TCHAR* Name,
-	const FRDGBufferDesc& Desc,
-	const void* InitialData,
-	uint64 InitialDataSize,
-	ERDGInitialDataFlags InitialDataFlags,
-	const FName& OwnerName,
-	bool bIsByteAddressBuffer = false)
-{
-	if (!bIsByteAddressBuffer)
-	{
-		checkf(EnumHasAnyFlags(Desc.Usage, EBufferUsageFlags::VertexBuffer), TEXT("CreateVertexBuffer called with an FRDGBufferDesc underlying type that is not 'VertexBuffer'. Buffer: %s"), Name);
-	}
-	FRDGBufferRef Buffer = GraphBuilder.CreateBuffer(Desc, Name, ERDGBufferFlags::MultiFrame);
-	Buffer->SetOwnerName(OwnerName);
-	if (InitialData && InitialDataSize)
-	{
-		GraphBuilder.QueueBufferUpload(Buffer, InitialData, InitialDataSize, InitialDataFlags);
-	}
-	return Buffer;
-}
-
-template<typename FormatType>
-void InternalCreateVertexBufferRDG_FromBulkData(FRDGBuilder& GraphBuilder, FByteBulkData& InBulkData, uint32 InDataCount, FRDGExternalBuffer& Out, const TCHAR* DebugName, const FName& OwnerName, EHairResourceUsageType UsageType)
+void InternalCreateBufferRDG_FromBulkData(FRDGBuilder& GraphBuilder, FByteBulkData& InBulkData, FRDGExternalBuffer& Out, const EPixelFormat OutFormat, const FRDGBufferDesc Desc, const TCHAR* DebugName, const FName& OwnerName)
 {
-	const uint32 InDataCount_Check = InBulkData.GetBulkDataSize() / sizeof(typename FormatType::BulkType);
-	check(InDataCount_Check >= InDataCount);
+	InternalSetBulkDataFlags(InBulkData);
 
-	const uint32 DataSizeInBytes = FormatType::SizeInByte * InDataCount;
+	const uint32 DataSizeInBytes = Desc.GetSize();
+	check(InBulkData.GetBulkDataSize() >= DataSizeInBytes);
 	if (DataSizeInBytes == 0)
 	{
 		Out.Buffer = nullptr;
 		return;
 	}
 
-	const FRDGBufferDesc Desc = ApplyUsage(FRDGBufferDesc::CreateBufferDesc(FormatType::SizeInByte, InDataCount), UsageType);
-
-	InternalSetBulkDataFlags(InBulkData);
-
-	const typename FormatType::BulkType* BulkData = (const typename FormatType::BulkType*)InBulkData.Lock(LOCK_READ_ONLY);
-	FRDGBufferRef Buffer = nullptr;
-	#if !WITH_EDITORONLY_DATA
-	if (ReleaseAfterUse())
+	const uint8* Data = (const uint8*)InBulkData.Lock(LOCK_READ_ONLY);
+	FRDGBufferRef Buffer = GraphBuilder.CreateBuffer(Desc, DebugName, ERDGBufferFlags::MultiFrame);
+	Buffer->SetOwnerName(OwnerName);
+	if (Data && DataSizeInBytes)
 	{
-		checkf(EnumHasAnyFlags(Desc.Usage, EBufferUsageFlags::VertexBuffer), TEXT("CreateVertexBuffer called with an FRDGBufferDesc underlying type that is not 'VertexBuffer'. Buffer: %s"), DebugName);
-		Buffer = GraphBuilder.CreateBuffer(Desc, DebugName, ERDGBufferFlags::MultiFrame);
-		Buffer->SetOwnerName(OwnerName);
-		if (BulkData && DataSizeInBytes)
+		#if !WITH_EDITORONLY_DATA
+		if (ReleaseAfterUse())
 		{
-			GraphBuilder.QueueBufferUpload(Buffer, BulkData, DataSizeInBytes, [&InBulkData](const void* Ptr) { InBulkData.Unlock(); });
+			GraphBuilder.QueueBufferUpload(Buffer, Data, DataSizeInBytes, [&InBulkData](const void* Ptr) { InBulkData.Unlock(); });
+		}
+		else
+		#endif
+		{
+			GraphBuilder.QueueBufferUpload(Buffer, Data, DataSizeInBytes, ERDGInitialDataFlags::None);  // Copy data internally
+			InBulkData.Unlock();
 		}
 	}
-	else
-	#endif
-	{
-		Buffer = InternalCreateVertexBuffer(
-			GraphBuilder,
-			DebugName,
-			Desc,
-			BulkData,
-			DataSizeInBytes,
-			ERDGInitialDataFlags::None, // Copy data internally
-			OwnerName);
-		InBulkData.Unlock();
-	}
+	ConvertToExternalBufferWithViews(GraphBuilder, Buffer, Out, OutFormat);
+}
 
-	ConvertToExternalBufferWithViews(GraphBuilder, Buffer, Out, FormatType::Format);
+template<typename FormatType>
+void InternalCreateVertexBufferRDG_FromBulkData(FRDGBuilder& GraphBuilder, FByteBulkData& InBulkData, uint32 InDataCount, FRDGExternalBuffer& Out, const TCHAR* DebugName, const FName& OwnerName, EHairResourceUsageType UsageType)
+{
+	const FRDGBufferDesc Desc = ApplyUsage(FRDGBufferDesc::CreateBufferDesc(FormatType::SizeInByte, InDataCount), UsageType);
+	InternalCreateBufferRDG_FromBulkData(GraphBuilder, InBulkData, Out, FormatType::Format, Desc, DebugName, OwnerName);
 }
 
-void InternalCreateByteAddressBufferRDG_FromBulkData(FRDGBuilder& GraphBuilder, FByteBulkData& InBulkData, FRDGExternalBuffer& Out, const TCHAR* DebugName, const FName& OwnerName, EHairResourceUsageType UsageType)
+template<typename FormatType>
+void InternalCreateStructuredBufferRDG_FromBulkData(FRDGBuilder& GraphBuilder, FByteBulkData& InBulkData, uint32 InDataCount, FRDGExternalBuffer& Out, const TCHAR* DebugName, const FName& OwnerName, EHairResourceUsageType UsageType)
 {
-	const uint32 DataSizeInBytes = InBulkData.GetBulkDataSize();
-	if (DataSizeInBytes == 0)
-	{
-		Out.Buffer = nullptr;
-		return;
-	}
+	const FRDGBufferDesc Desc = ApplyUsage(FRDGBufferDesc::CreateStructuredDesc(FormatType::SizeInByte, InDataCount), UsageType);
+	InternalCreateBufferRDG_FromBulkData(GraphBuilder, InBulkData, Out, PF_Unknown, Desc, DebugName, OwnerName);
+}
 
-	const FRDGBufferDesc Desc = ApplyUsage(FRDGBufferDesc::CreateByteAddressDesc(DataSizeInBytes), UsageType);
+void InternalCreateByteAddressBufferRDG_FromBulkData(FRDGBuilder& GraphBuilder, FByteBulkData& InBulkData, FRDGExternalBuffer& Out, const TCHAR* DebugName, const FName& OwnerName, EHairResourceUsageType UsageType)
+{
+	const FRDGBufferDesc Desc = ApplyUsage(FRDGBufferDesc::CreateByteAddressDesc(InBulkData.GetBulkDataSize()), UsageType);
+	InternalCreateBufferRDG_FromBulkData(GraphBuilder, InBulkData, Out, PF_Unknown, Desc, DebugName, OwnerName);
+}
 
-	InternalSetBulkDataFlags(InBulkData);
 
-	const uint8* BulkData = (const uint8*)InBulkData.Lock(LOCK_READ_ONLY);
-	FRDGBufferRef Buffer = nullptr;
-#if !WITH_EDITORONLY_DATA
-	if (ReleaseAfterUse())
-	{		
-		Buffer = GraphBuilder.CreateBuffer(Desc, DebugName, ERDGBufferFlags::MultiFrame);
-		Buffer->SetOwnerName(OwnerName);
-		if (BulkData && DataSizeInBytes)
-		{
-			GraphBuilder.QueueBufferUpload(Buffer, BulkData, DataSizeInBytes, [&InBulkData](const void* Ptr) { InBulkData.Unlock(); });
-		}
+static FRDGBufferRef InternalCreateVertexBuffer(
+	FRDGBuilder& GraphBuilder,
+	const TCHAR* Name,
+	const FRDGBufferDesc& Desc,
+	const void* InitialData,
+	uint64 InitialDataSize,
+	ERDGInitialDataFlags InitialDataFlags,
+	const FName& OwnerName,
+	bool bIsByteAddressBuffer = false)
+{
+	if (!bIsByteAddressBuffer)
+	{
+		checkf(EnumHasAnyFlags(Desc.Usage, EBufferUsageFlags::VertexBuffer), TEXT("CreateVertexBuffer called with an FRDGBufferDesc underlying type that is not 'VertexBuffer'. Buffer: %s"), Name);
 	}
-	else
-#endif
+	FRDGBufferRef Buffer = GraphBuilder.CreateBuffer(Desc, Name, ERDGBufferFlags::MultiFrame);
+	Buffer->SetOwnerName(OwnerName);
+	if (InitialData && InitialDataSize)
 	{
-		Buffer = InternalCreateVertexBuffer(
-			GraphBuilder,
-			DebugName,
-			Desc,
-			BulkData,
-			DataSizeInBytes,
-			ERDGInitialDataFlags::None, 
-			OwnerName,
-			true /*bIsByteAddressBuffer*/); // Copy data internally
-		InBulkData.Unlock();
+		GraphBuilder.QueueBufferUpload(Buffer, InitialData, InitialDataSize, InitialDataFlags);
 	}
-
-	ConvertToExternalBufferWithViews(GraphBuilder, Buffer, Out, PF_Unknown);
+	return Buffer;
 }
 
 template<typename FormatType>
@@ -544,44 +507,6 @@ void InternalCreateStructuredBufferRDG(FRDGBuilder& GraphBuilder, uint32 DataCou
 	ConvertToExternalBufferWithViews(GraphBuilder, Buffer, Out, FormatType::Format);
 }
 
-template<typename FormatType>
-void InternalCreateStructuredBufferRDG_FromBulkData(FRDGBuilder& GraphBuilder, FByteBulkData& InBulkData, uint32 InDataCount, FRDGExternalBuffer& Out, const TCHAR* DebugName, const FName& OwnerName, EHairResourceUsageType UsageType)
-{
-	const uint32 InSizeInByte = sizeof(typename FormatType::Type);
-	const uint32 DataSizeInBytes = InSizeInByte * InDataCount;
-	check(InBulkData.GetBulkDataSize() >= DataSizeInBytes);
-
-	if (DataSizeInBytes == 0)
-	{
-		Out.Buffer = nullptr;
-		return;
-	}
-
-	const FRDGBufferDesc Desc = ApplyUsage(FRDGBufferDesc::CreateStructuredDesc(InSizeInByte, InDataCount), UsageType);
-
-	InternalSetBulkDataFlags(InBulkData);
-
-	const typename FormatType::Type* Data = (const typename FormatType::Type*)InBulkData.Lock(LOCK_READ_ONLY);
-	FRDGBufferRef Buffer = GraphBuilder.CreateBuffer(Desc, DebugName, ERDGBufferFlags::MultiFrame);
-	Buffer->SetOwnerName(OwnerName);
-	if (Data && DataSizeInBytes)
-	{
-		#if !WITH_EDITORONLY_DATA
-		if (ReleaseAfterUse())
-		{
-			GraphBuilder.QueueBufferUpload(Buffer, Data, DataSizeInBytes, [&InBulkData](const void* Ptr) { InBulkData.Unlock(); });
-		}
-		else
-		#endif
-		{
-			GraphBuilder.QueueBufferUpload(Buffer, Data, DataSizeInBytes, ERDGInitialDataFlags::None);  // Copy data internally
-			InBulkData.Unlock();
-		}
-	}
-
-	ConvertToExternalBufferWithViews(GraphBuilder, Buffer, Out);
-}
-
 /////////////////////////////////////////////////////////////////////////////////////////
 
 static UTexture2D* CreateCardTexture(FIntPoint Resolution)
