Change how bound are computed for groom when attached to skel. mesh.

diff --git a/Engine/Plugins/Runtime/HairStrands/Source/HairStrandsCore/Private/GroomComponent.cpp b/Engine/Plugins/Runtime/HairStrands/Source/HairStrandsCore/Private/GroomComponent.cpp
--- a/Engine/Plugins/Runtime/HairStrands/Source/HairStrandsCore/Private/GroomComponent.cpp
+++ b/Engine/Plugins/Runtime/HairStrands/Source/HairStrandsCore/Private/GroomComponent.cpp
@@ -1790,34 +1790,33 @@ FBoxSphereBounds UGroomComponent::CalcBounds(const FTransform& InLocalToWorld) c
 		// * GHairStrands_BoundsMode=2 : use the skel mesh bounds + groom bounds. This is more conservative compares to GHairStrands_BoundsMode=0.
 		if (RegisteredMeshComponent && GHairStrands_BoundsMode == 1)
 		{
-			return RegisteredMeshComponent->Bounds;
+			return RegisteredMeshComponent->CalcBounds(InLocalToWorld).GetBox();
 		}
 		else if (RegisteredMeshComponent && GHairStrands_BoundsMode == 2)
 		{
 			const FVector3d GroomExtends = LocalHairBound.GetExtent();
 			const float BoundExtraRadius = 0.5f * static_cast<float>(FMath::Max(0.0, FMath::Max3(GroomExtends.X, GroomExtends.Y, GroomExtends.Z)));
-			FBox EffectiveBound = RegisteredMeshComponent->Bounds.GetBox();
+			FBox EffectiveBound = RegisteredMeshComponent->CalcBounds(FTransform::Identity).GetBox();
 			EffectiveBound.Min -= FVector3d(BoundExtraRadius);
 			EffectiveBound.Max += FVector3d(BoundExtraRadius);
-			return FBoxSphereBounds(EffectiveBound);
+			return FBoxSphereBounds(EffectiveBound.TransformBy(InLocalToWorld));
 		}
 		else if (RegisteredMeshComponent && AttachmentName.IsEmpty())
 		{
-			const FBox LocalSkeletalBound = RegisteredMeshComponent->CalcBounds(FTransform::Identity).GetBox();
-
 			// Compute an extra 'radius' which will be added to the skel. mesh bound. 
 			// It is based on the different of the groom bound and the skel. bound at rest position
 			float BoundExtraRadius = 0.f;
 			if (USkeletalMeshComponent* SkeletalMeshComponent = Cast<USkeletalMeshComponent>(RegisteredMeshComponent))
 			{
-				FBoxSphereBounds RestBounds;
-				SkeletalMeshComponent->GetPreSkinnedLocalBounds(RestBounds);
-				const FBox RestBox = RestBounds.GetBox();
-				const FVector MinDiff = LocalHairBound.Min - RestBox.Min;
-				const FVector MaxDiff = LocalHairBound.Max - RestBox.Max;
+				FBoxSphereBounds SkelRestBounds;
+				SkeletalMeshComponent->GetPreSkinnedLocalBounds(SkelRestBounds);
+				FBoxSphereBounds ExtentedBound = Union(SkelRestBounds, FBoxSphereBounds(LocalHairBound));
+				const FVector MinDiff = SkelRestBounds.GetBox().Min - ExtentedBound.GetBox().Min;
+				const FVector MaxDiff = ExtentedBound.GetBox().Max  - SkelRestBounds.GetBox().Max;
 				BoundExtraRadius = static_cast<float>(FMath::Max3(0.0, FMath::Max3(MinDiff.X, MinDiff.Y, MinDiff.Z), FMath::Max3(MaxDiff.X, MaxDiff.Y, MaxDiff.Z)));
 			}
 			
+			const FBox LocalSkeletalBound = RegisteredMeshComponent->CalcBounds(FTransform::Identity).GetBox();
 			FBox LocalBound(EForceInit::ForceInitToZero);
 			LocalBound += LocalSkeletalBound.Min - BoundExtraRadius;
 			LocalBound += LocalSkeletalBound.Max + BoundExtraRadius;
