Fix issue in groom deformer where groom instance wouldn't be valid when queried.

diff --git a/Engine/Plugins/Runtime/HairStrands/Source/HairStrandsCore/Private/Deformer/DeformerDataInterfaceGroom.cpp b/Engine/Plugins/Runtime/HairStrands/Source/HairStrandsCore/Private/Deformer/DeformerDataInterfaceGroom.cpp
--- a/Engine/Plugins/Runtime/HairStrands/Source/HairStrandsCore/Private/Deformer/DeformerDataInterfaceGroom.cpp
+++ b/Engine/Plugins/Runtime/HairStrands/Source/HairStrandsCore/Private/Deformer/DeformerDataInterfaceGroom.cpp
@@ -157,10 +157,19 @@ bool FOptimusGroomDataProviderProxy::IsValid(FValidationData const& InValidation
 	{
 		return false;
 	}
-	if (InValidationData.NumInvocations != GroomComponent->GetGroupCount())
+	const uint32 GroupCount = GroomComponent->GetGroupCount();
+	if (InValidationData.NumInvocations != GroupCount)
 	{
 		return false;
 	}
+	// Earlier invalidation in case one of the instance is not ready/valid
+	for (uint32 GroupIndex = 0; GroupIndex < GroupCount; ++GroupIndex)
+	{
+		if (GroomComponent->GetGroupInstance(GroupIndex) == nullptr)
+		{
+			return false;
+		}
+	}
 	
 	return true;
 }
@@ -200,6 +209,17 @@ void FOptimusGroomDataProviderProxy::GatherDispatchData(FDispatchData const& InD
 	const TStridedView<FParameters> ParameterArray = MakeStridedParameterView<FParameters>(InDispatchData);
 	for (int32 InvocationIndex = 0; InvocationIndex < ParameterArray.Num(); ++InvocationIndex)
 	{
+		// Default valid in case the instance is not valid
+		FParameters& Parameters = ParameterArray[InvocationIndex];
+		Parameters.Common.PointCount = 0;
+		Parameters.Common.CurveCount = 0;
+		Parameters.Resources.PositionBuffer			= FallbackSRV;
+		Parameters.Resources.PositionOffsetBuffer 	= FallbackSRV;
+		Parameters.Resources.CurveAttributeBuffer	= FallbackSRV;
+		Parameters.Resources.PointAttributeBuffer	= FallbackSRV;
+		Parameters.Resources.PointToCurveBuffer		= FallbackSRV;
+		Parameters.Resources.CurveBuffer			= FallbackSRV;
+
 		if (FHairGroupInstance* Instance = GroomComponent->GetGroupInstance(InvocationIndex))
 		{
 			const bool bIsSRVValid = Resources[InvocationIndex].PositionBuffer != nullptr;
@@ -208,22 +228,12 @@ void FOptimusGroomDataProviderProxy::GatherDispatchData(FDispatchData const& InD
 
 			const FHairGroupPublicData::FVertexFactoryInput VFInput = ComputeHairStrandsVertexInputData(Instance, EGroomViewMode::None);
 			
-			FParameters& Parameters = ParameterArray[InvocationIndex];
 			Parameters.Common = VFInput.Strands.Common;
 
 			if (bIsSRVValid)
 			{
 				Parameters.Resources = Resources[InvocationIndex];
 			}
-			else
-			{
-				Parameters.Resources.PositionBuffer			= FallbackSRV;
-				Parameters.Resources.PositionOffsetBuffer 	= FallbackSRV;
-				Parameters.Resources.CurveAttributeBuffer	= FallbackSRV;
-				Parameters.Resources.PointAttributeBuffer	= FallbackSRV;
-				Parameters.Resources.PointToCurveBuffer		= FallbackSRV;
-				Parameters.Resources.CurveBuffer			= FallbackSRV;
-			}
 		}
 	}
 }
\ No newline at end of file
