Change groom binding root triangle position to be stored into a single buffer.

diff --git a/Engine/Plugins/Runtime/HairStrands/Shaders/Private/NiagaraDataInterfaceHairStrandsTemplate.ush b/Engine/Plugins/Runtime/HairStrands/Shaders/Private/NiagaraDataInterfaceHairStrandsTemplate.ush
--- a/Engine/Plugins/Runtime/HairStrands/Shaders/Private/NiagaraDataInterfaceHairStrandsTemplate.ush
+++ b/Engine/Plugins/Runtime/HairStrands/Shaders/Private/NiagaraDataInterfaceHairStrandsTemplate.ush
@@ -7,12 +7,8 @@ NiagaraDataInterfaceHairStrands.ush
 RWBuffer<uint4>		{ParameterName}_DeformedPositionBuffer;
 Buffer<uint4>		{ParameterName}_RestPositionBuffer;
 Buffer<uint>		{ParameterName}_CurvesOffsetsBuffer;
-Buffer<float4>		{ParameterName}_RestTrianglePositionABuffer;
-Buffer<float4>		{ParameterName}_RestTrianglePositionBBuffer;
-Buffer<float4>		{ParameterName}_RestTrianglePositionCBuffer;
-Buffer<float4>		{ParameterName}_DeformedTrianglePositionABuffer;
-Buffer<float4>		{ParameterName}_DeformedTrianglePositionBBuffer;
-Buffer<float4>		{ParameterName}_DeformedTrianglePositionCBuffer;
+Buffer<float4>		{ParameterName}_RestTrianglePositionBuffer;
+Buffer<float4>		{ParameterName}_DeformedTrianglePositionBuffer;
 Buffer<uint>		{ParameterName}_RootBarycentricCoordinatesBuffer;
 Buffer<uint>		{ParameterName}_RootToUniqueTriangleIndexBuffer;
 float4x4			{ParameterName}_WorldTransform;
@@ -609,9 +605,9 @@ void ComputeProjectionTriangle_{ParameterName}(in float2 ProjectionUV, in float3
 void BuildDeformedTriangle_{ParameterName}(in float2 ProjectionUV, in int TriangleIndex, out float3 OutTrianglePosition, out float4 OutTriangleRotation)
 {
 	ComputeProjectionTriangle_{ParameterName}(ProjectionUV, 
-		{ParameterName}_DeformedTrianglePositionABuffer[TriangleIndex].xyz,
-		{ParameterName}_DeformedTrianglePositionBBuffer[TriangleIndex].xyz,
-		{ParameterName}_DeformedTrianglePositionCBuffer[TriangleIndex].xyz,
+		{ParameterName}_DeformedTrianglePositionBuffer[TriangleIndex * 3 + 0].xyz,
+		{ParameterName}_DeformedTrianglePositionBuffer[TriangleIndex * 3 + 1].xyz,
+		{ParameterName}_DeformedTrianglePositionBuffer[TriangleIndex * 3 + 2].xyz,
 	    {ParameterName}_DeformedRootOffset, OutTrianglePosition, OutTriangleRotation);
 }
 
@@ -619,9 +615,9 @@ void BuildDeformedTriangle_{ParameterName}(in float2 ProjectionUV, in int Triang
 void BuildRestTriangle_{ParameterName}(in float2 ProjectionUV, in int TriangleIndex, out float3 OutTrianglePosition, out float4 OutTriangleRotation)
 {
 	ComputeProjectionTriangle_{ParameterName}(ProjectionUV, 
-		{ParameterName}_RestTrianglePositionABuffer[TriangleIndex].xyz,
-		{ParameterName}_RestTrianglePositionBBuffer[TriangleIndex].xyz,
-		{ParameterName}_RestTrianglePositionCBuffer[TriangleIndex].xyz,
+		{ParameterName}_RestTrianglePositionBuffer[TriangleIndex * 3 + 0].xyz,
+		{ParameterName}_RestTrianglePositionBuffer[TriangleIndex * 3 + 1].xyz,
+		{ParameterName}_RestTrianglePositionBuffer[TriangleIndex * 3 + 2].xyz,
 	    {ParameterName}_RestRootOffset, OutTrianglePosition, OutTriangleRotation);
 }
 
diff --git a/Engine/Plugins/Runtime/HairStrands/Shaders/Private/NiagaraInterpolateGroomGuides.usf b/Engine/Plugins/Runtime/HairStrands/Shaders/Private/NiagaraInterpolateGroomGuides.usf
--- a/Engine/Plugins/Runtime/HairStrands/Shaders/Private/NiagaraInterpolateGroomGuides.usf
+++ b/Engine/Plugins/Runtime/HairStrands/Shaders/Private/NiagaraInterpolateGroomGuides.usf
@@ -15,12 +15,8 @@ NiagaraInterpolateGroomGuides.usf
 
 #include "NiagaraGroomInterpolationUtils.ush"
 
-Buffer<float4>		RestTrianglePositionABuffer;
-Buffer<float4>		RestTrianglePositionBBuffer;
-Buffer<float4>		RestTrianglePositionCBuffer;
-Buffer<float4>		DeformedTrianglePositionABuffer;
-Buffer<float4>		DeformedTrianglePositionBBuffer;
-Buffer<float4>		DeformedTrianglePositionCBuffer;
+Buffer<float4>		RestTrianglePositionBuffer;
+Buffer<float4>		DeformedTrianglePositionBuffer;
 Buffer<uint>		RootBarycentricCoordinatesBuffer;
 Buffer<uint>		RootToUniqueTriangleIndexBuffer;
 
@@ -50,9 +46,9 @@ void DIHairStrands_BuildDeformedTriangle(in float2 ProjectionUV, in uint Triangl
 {
 	DIHairStrands_ComputeProjectionTriangle(
 		ProjectionUV, 
-		DeformedTrianglePositionABuffer[TriangleIndex].xyz,
-		DeformedTrianglePositionBBuffer[TriangleIndex].xyz, 
-		DeformedTrianglePositionCBuffer[TriangleIndex].xyz, 
+		DeformedTrianglePositionBuffer[TriangleIndex * 3 + 0].xyz,
+		DeformedTrianglePositionBuffer[TriangleIndex * 3 + 1].xyz, 
+		DeformedTrianglePositionBuffer[TriangleIndex * 3 + 2].xyz, 
 		float3(0,0,0), OutTrianglePosition, OutTriangleRotation);
 }
 
@@ -61,9 +57,9 @@ void DIHairStrands_BuildRestTriangle(in float2 ProjectionUV, in uint TriangleInd
 {
 	DIHairStrands_ComputeProjectionTriangle(
 		ProjectionUV, 
-		RestTrianglePositionABuffer[TriangleIndex].xyz,
-		RestTrianglePositionBBuffer[TriangleIndex].xyz,
-		RestTrianglePositionCBuffer[TriangleIndex].xyz,
+		RestTrianglePositionBuffer[TriangleIndex * 3 + 0].xyz,
+		RestTrianglePositionBuffer[TriangleIndex * 3 + 1].xyz,
+		RestTrianglePositionBuffer[TriangleIndex * 3 + 2].xyz,
 		float3(0,0,0), OutTrianglePosition, OutTriangleRotation);
 }
 
diff --git a/Engine/Plugins/Runtime/HairStrands/Source/HairStrandsCore/Private/GroomBindingBuilder.cpp b/Engine/Plugins/Runtime/HairStrands/Source/HairStrandsCore/Private/GroomBindingBuilder.cpp
--- a/Engine/Plugins/Runtime/HairStrands/Source/HairStrandsCore/Private/GroomBindingBuilder.cpp
+++ b/Engine/Plugins/Runtime/HairStrands/Source/HairStrandsCore/Private/GroomBindingBuilder.cpp
@@ -57,7 +57,7 @@ static FAutoConsoleVariableRef CVarHairStrandsBindingBuilderWarningEnable(TEXT("
 FString FGroomBindingBuilder::GetVersion()
 {
 	// Important to update the version when groom building changes
-	return TEXT("2gc");
+	return TEXT("2h");
 }
 
 ///////////////////////////////////////////////////////////////////////////////////////////////////
@@ -1242,12 +1242,8 @@ namespace GroomBinding_RootProjection
 			TArray<FHairStrandsUniqueTriangleIndexFormat::Type> RootTriangleIndexBuffer;
 			RootTriangleIndexBuffer.SetNum(CurveCount);
 
-			TArray<FHairStrandsMeshTrianglePositionFormat::Type> RestRootTrianglePosition0Buffer;
-			TArray<FHairStrandsMeshTrianglePositionFormat::Type> RestRootTrianglePosition1Buffer;
-			TArray<FHairStrandsMeshTrianglePositionFormat::Type> RestRootTrianglePosition2Buffer;
-			RestRootTrianglePosition0Buffer.SetNum(CurveCount);
-			RestRootTrianglePosition1Buffer.SetNum(CurveCount);
-			RestRootTrianglePosition2Buffer.SetNum(CurveCount);
+			TArray<FHairStrandsMeshTrianglePositionFormat::Type> RestRootTrianglePositionBuffer;
+			RestRootTrianglePositionBuffer.SetNum(CurveCount * 3);
 
 		#if BINDING_PARALLEL_BUILDING
 			TAtomic<uint32> bIsValid(1);
@@ -1257,9 +1253,7 @@ namespace GroomBinding_RootProjection
 					&InStrandsData,
 					&Grid,
 					&RootTriangleIndexBuffer,
-					&RestRootTrianglePosition0Buffer,
-					&RestRootTrianglePosition1Buffer,
-					&RestRootTrianglePosition2Buffer,
+					&RestRootTrianglePositionBuffer,
 					&OutRootData,
 					&bIsValid
 				] (uint32 CurveIndex)
@@ -1305,9 +1299,9 @@ namespace GroomBinding_RootProjection
 				OutRootData.MeshProjectionLODs[LODIt].RootBarycentricBuffer[CurveIndex] = EncodedBarycentrics;
 
 				RootTriangleIndexBuffer[CurveIndex] = EncodedTriangleIndex;
-				RestRootTrianglePosition0Buffer[CurveIndex] = FVector4f((FVector3f)ClosestTriangle.P0, FHairStrandsRootUtils::PackUVsToFloat(FVector2f(ClosestTriangle.UV0)));	// LWC_TODO: Precision loss
-				RestRootTrianglePosition1Buffer[CurveIndex] = FVector4f((FVector3f)ClosestTriangle.P1, FHairStrandsRootUtils::PackUVsToFloat(FVector2f(ClosestTriangle.UV1)));	// LWC_TODO: Precision loss
-				RestRootTrianglePosition2Buffer[CurveIndex] = FVector4f((FVector3f)ClosestTriangle.P2, FHairStrandsRootUtils::PackUVsToFloat(FVector2f(ClosestTriangle.UV2)));	// LWC_TODO: Precision loss
+				RestRootTrianglePositionBuffer[CurveIndex * 3 + 0] = FVector4f((FVector3f)ClosestTriangle.P0, FHairStrandsRootUtils::PackUVsToFloat(FVector2f(ClosestTriangle.UV0)));	// LWC_TODO: Precision loss
+				RestRootTrianglePositionBuffer[CurveIndex * 3 + 1] = FVector4f((FVector3f)ClosestTriangle.P1, FHairStrandsRootUtils::PackUVsToFloat(FVector2f(ClosestTriangle.UV1)));	// LWC_TODO: Precision loss
+				RestRootTrianglePositionBuffer[CurveIndex * 3 + 2] = FVector4f((FVector3f)ClosestTriangle.P2, FHairStrandsRootUtils::PackUVsToFloat(FVector2f(ClosestTriangle.UV2)));	// LWC_TODO: Precision loss
 			}
 		#if BINDING_PARALLEL_BUILDING
 			);
@@ -1349,10 +1343,8 @@ namespace GroomBinding_RootProjection
 
 			// Build final unique triangle list and the root-to-unique-triangle mapping
 			const uint32 UniqueTriangleCount = UniqueTriangleToRootList.Num();
-			OutRootData.MeshProjectionLODs[LODIt].UniqueTriangleIndexBuffer.Reserve(UniqueTriangleCount);
-			OutRootData.MeshProjectionLODs[LODIt].RestUniqueTrianglePosition0Buffer.Reserve(UniqueTriangleCount);
-			OutRootData.MeshProjectionLODs[LODIt].RestUniqueTrianglePosition1Buffer.Reserve(UniqueTriangleCount);
-			OutRootData.MeshProjectionLODs[LODIt].RestUniqueTrianglePosition2Buffer.Reserve(UniqueTriangleCount);
+			OutRootData.MeshProjectionLODs[LODIt].UniqueTriangleIndexBuffer.Reserve(UniqueTriangleCount );
+			OutRootData.MeshProjectionLODs[LODIt].RestUniqueTrianglePositionBuffer.Reserve(UniqueTriangleCount * 3);
 			for (uint32 EncodedTriangleId : UniqueTriangleToRootList)
 			{
 				auto It = UniqueTriangleToRootMap.Find(EncodedTriangleId);
@@ -1361,9 +1353,9 @@ namespace GroomBinding_RootProjection
 				OutRootData.MeshProjectionLODs[LODIt].UniqueTriangleIndexBuffer.Add(EncodedTriangleId);
 
 				const uint32 FirstCurveIndex = (*It)[0];
-				OutRootData.MeshProjectionLODs[LODIt].RestUniqueTrianglePosition0Buffer.Add(RestRootTrianglePosition0Buffer[FirstCurveIndex]);
-				OutRootData.MeshProjectionLODs[LODIt].RestUniqueTrianglePosition1Buffer.Add(RestRootTrianglePosition1Buffer[FirstCurveIndex]);
-				OutRootData.MeshProjectionLODs[LODIt].RestUniqueTrianglePosition2Buffer.Add(RestRootTrianglePosition2Buffer[FirstCurveIndex]);
+				OutRootData.MeshProjectionLODs[LODIt].RestUniqueTrianglePositionBuffer.Add(RestRootTrianglePositionBuffer[FirstCurveIndex * 3 + 0]);
+				OutRootData.MeshProjectionLODs[LODIt].RestUniqueTrianglePositionBuffer.Add(RestRootTrianglePositionBuffer[FirstCurveIndex * 3 + 1]);
+				OutRootData.MeshProjectionLODs[LODIt].RestUniqueTrianglePositionBuffer.Add(RestRootTrianglePositionBuffer[FirstCurveIndex * 3 + 2]);
 
 				// Write for each root, the index of the triangle
 				const uint32 UniqueTriangleIndex = OutRootData.MeshProjectionLODs[LODIt].UniqueTriangleIndexBuffer.Num()-1;
@@ -1375,9 +1367,7 @@ namespace GroomBinding_RootProjection
 
 			// Sanity check
 			check(OutRootData.MeshProjectionLODs[LODIt].RootToUniqueTriangleIndexBuffer.Num() == CurveCount);
-			check(OutRootData.MeshProjectionLODs[LODIt].RestUniqueTrianglePosition0Buffer.Num() == UniqueTriangleCount);
-			check(OutRootData.MeshProjectionLODs[LODIt].RestUniqueTrianglePosition1Buffer.Num() == UniqueTriangleCount);
-			check(OutRootData.MeshProjectionLODs[LODIt].RestUniqueTrianglePosition2Buffer.Num() == UniqueTriangleCount);
+			check(OutRootData.MeshProjectionLODs[LODIt].RestUniqueTrianglePositionBuffer.Num() == UniqueTriangleCount * 3);
 			check(OutRootData.MeshProjectionLODs[LODIt].UniqueTriangleIndexBuffer.Num() == UniqueTriangleCount);
 
 			// Update the root mesh projection data with unique valid mesh section IDs, based on the projection data
@@ -1918,15 +1908,12 @@ static void BuildRootBulkData(
 
 		Out.MeshProjectionLODs[MeshLODIt].LODIndex = In.MeshProjectionLODs[MeshLODIt].LODIndex;
 		Out.MeshProjectionLODs[MeshLODIt].SampleCount = bHasValidSamples ? In.MeshProjectionLODs[MeshLODIt].SampleCount : 0u;
-		Out.MeshProjectionLODs[MeshLODIt].UniqueTriangleCount = In.MeshProjectionLODs[MeshLODIt].RestUniqueTrianglePosition0Buffer.Num();
+		Out.MeshProjectionLODs[MeshLODIt].UniqueTriangleCount = In.MeshProjectionLODs[MeshLODIt].UniqueTriangleIndexBuffer.Num();
 
 		CopyToBulkData<FHairStrandsUniqueTriangleIndexFormat>(Out.MeshProjectionLODs[MeshLODIt].UniqueTriangleIndexBuffer, In.MeshProjectionLODs[MeshLODIt].UniqueTriangleIndexBuffer);
 		CopyToBulkData<FHairStrandsRootBarycentricFormat>(Out.MeshProjectionLODs[MeshLODIt].RootBarycentricBuffer, In.MeshProjectionLODs[MeshLODIt].RootBarycentricBuffer);
 		CopyToBulkData<FHairStrandsRootToUniqueTriangleIndexFormat>(Out.MeshProjectionLODs[MeshLODIt].RootToUniqueTriangleIndexBuffer, In.MeshProjectionLODs[MeshLODIt].RootToUniqueTriangleIndexBuffer);
-
-		CopyToBulkData<FHairStrandsMeshTrianglePositionFormat>(Out.MeshProjectionLODs[MeshLODIt].RestUniqueTrianglePosition0Buffer, In.MeshProjectionLODs[MeshLODIt].RestUniqueTrianglePosition0Buffer);
-		CopyToBulkData<FHairStrandsMeshTrianglePositionFormat>(Out.MeshProjectionLODs[MeshLODIt].RestUniqueTrianglePosition1Buffer, In.MeshProjectionLODs[MeshLODIt].RestUniqueTrianglePosition1Buffer);
-		CopyToBulkData<FHairStrandsMeshTrianglePositionFormat>(Out.MeshProjectionLODs[MeshLODIt].RestUniqueTrianglePosition2Buffer, In.MeshProjectionLODs[MeshLODIt].RestUniqueTrianglePosition2Buffer);
+		CopyToBulkData<FHairStrandsMeshTrianglePositionFormat>(Out.MeshProjectionLODs[MeshLODIt].RestUniqueTrianglePositionBuffer, In.MeshProjectionLODs[MeshLODIt].RestUniqueTrianglePositionBuffer);
 
 		Out.MeshProjectionLODs[MeshLODIt].UniqueSectionIndices = In.MeshProjectionLODs[MeshLODIt].UniqueSectionIds;
 
@@ -1965,10 +1952,7 @@ static void BuildRootData(
 		CopyFromBulkData<FHairStrandsUniqueTriangleIndexFormat>(Out.MeshProjectionLODs[MeshLODIt].UniqueTriangleIndexBuffer, In.MeshProjectionLODs[MeshLODIt].UniqueTriangleIndexBuffer);
 		CopyFromBulkData<FHairStrandsRootToUniqueTriangleIndexFormat>(Out.MeshProjectionLODs[MeshLODIt].RootToUniqueTriangleIndexBuffer, In.MeshProjectionLODs[MeshLODIt].RootToUniqueTriangleIndexBuffer);
 		CopyFromBulkData<FHairStrandsRootBarycentricFormat>(Out.MeshProjectionLODs[MeshLODIt].RootBarycentricBuffer, In.MeshProjectionLODs[MeshLODIt].RootBarycentricBuffer);
-
-		CopyFromBulkData<FHairStrandsMeshTrianglePositionFormat>(Out.MeshProjectionLODs[MeshLODIt].RestUniqueTrianglePosition0Buffer, In.MeshProjectionLODs[MeshLODIt].RestUniqueTrianglePosition0Buffer);
-		CopyFromBulkData<FHairStrandsMeshTrianglePositionFormat>(Out.MeshProjectionLODs[MeshLODIt].RestUniqueTrianglePosition1Buffer, In.MeshProjectionLODs[MeshLODIt].RestUniqueTrianglePosition1Buffer);
-		CopyFromBulkData<FHairStrandsMeshTrianglePositionFormat>(Out.MeshProjectionLODs[MeshLODIt].RestUniqueTrianglePosition2Buffer, In.MeshProjectionLODs[MeshLODIt].RestUniqueTrianglePosition2Buffer);
+		CopyFromBulkData<FHairStrandsMeshTrianglePositionFormat>(Out.MeshProjectionLODs[MeshLODIt].RestUniqueTrianglePositionBuffer, In.MeshProjectionLODs[MeshLODIt].RestUniqueTrianglePositionBuffer);
 
 		Out.MeshProjectionLODs[MeshLODIt].UniqueSectionIds = In.MeshProjectionLODs[MeshLODIt].UniqueSectionIndices;
 
diff --git a/Engine/Plugins/Runtime/HairStrands/Source/HairStrandsCore/Private/GroomDebug.cpp b/Engine/Plugins/Runtime/HairStrands/Source/HairStrandsCore/Private/GroomDebug.cpp
--- a/Engine/Plugins/Runtime/HairStrands/Source/HairStrandsCore/Private/GroomDebug.cpp
+++ b/Engine/Plugins/Runtime/HairStrands/Source/HairStrandsCore/Private/GroomDebug.cpp
@@ -242,17 +242,10 @@ class FHairProjectionHairDebugCS : public FGlobalShader
 		SHADER_PARAMETER(uint32, DeformedFrameEnable)
 		SHADER_PARAMETER(FMatrix44f, RootLocalToWorld)
 
-		SHADER_PARAMETER_RDG_BUFFER_SRV(Buffer, RestPosition0Buffer)
-		SHADER_PARAMETER_RDG_BUFFER_SRV(Buffer, RestPosition1Buffer)
-		SHADER_PARAMETER_RDG_BUFFER_SRV(Buffer, RestPosition2Buffer)
-
-		SHADER_PARAMETER_RDG_BUFFER_SRV(Buffer, DeformedPosition0Buffer)
-		SHADER_PARAMETER_RDG_BUFFER_SRV(Buffer, DeformedPosition1Buffer)
-		SHADER_PARAMETER_RDG_BUFFER_SRV(Buffer, DeformedPosition2Buffer)
-
+		SHADER_PARAMETER_RDG_BUFFER_SRV(Buffer, RestPositionBuffer)
+		SHADER_PARAMETER_RDG_BUFFER_SRV(Buffer, DeformedPositionBuffer)
 		SHADER_PARAMETER_RDG_BUFFER_SRV(Buffer, RestSamplePositionsBuffer)
 		SHADER_PARAMETER_RDG_BUFFER_SRV(Buffer, DeformedSamplePositionsBuffer)
-
 		SHADER_PARAMETER_RDG_BUFFER_SRV(Buffer, RootBarycentricBuffer)
 		SHADER_PARAMETER_RDG_BUFFER_SRV(Buffer, RootToUniqueTriangleIndexBuffer)
 
@@ -334,18 +327,13 @@ static void AddDebugProjectionHairPass(
 	const FHairStrandsRestRootResource::FLOD& RestLODDatas = RestRootResources->LODs[MeshLODIndex];
 	const FHairStrandsDeformedRootResource::FLOD& DeformedLODDatas = DeformedRootResources->LODs[MeshLODIndex];
 
-	if (!RestLODDatas.RestUniqueTrianglePosition0Buffer.Buffer ||
-		!RestLODDatas.RestUniqueTrianglePosition1Buffer.Buffer ||
-		!RestLODDatas.RestUniqueTrianglePosition2Buffer.Buffer ||
-		!DeformedLODDatas.DeformedUniqueTrianglePosition0Buffer[0].Buffer ||
-		!DeformedLODDatas.DeformedUniqueTrianglePosition1Buffer[0].Buffer ||
-		!DeformedLODDatas.DeformedUniqueTrianglePosition2Buffer[0].Buffer)
+	if (!RestLODDatas.RestUniqueTrianglePositionBuffer.Buffer || !DeformedLODDatas.DeformedUniqueTrianglePositionBuffer[0].Buffer)
 	{
 		return;
 	}
 
 	// Double buffering is disabled by default unless the read-only cvar r.HairStrands.ContinuousDecimationReordering is set
-	if (IsHairStrandContinuousDecimationReorderingEnabled() && (!DeformedLODDatas.DeformedUniqueTrianglePosition0Buffer[1].Buffer || !DeformedLODDatas.DeformedUniqueTrianglePosition1Buffer[1].Buffer || !DeformedLODDatas.DeformedUniqueTrianglePosition2Buffer[1].Buffer))
+	if (IsHairStrandContinuousDecimationReorderingEnabled() && !DeformedLODDatas.DeformedUniqueTrianglePositionBuffer[1].Buffer)
 	{
 		return;
 	}
@@ -365,13 +353,8 @@ static void AddDebugProjectionHairPass(
 
 	Parameters->RootToUniqueTriangleIndexBuffer = RegisterAsSRV(GraphBuilder, RestLODDatas.RootToUniqueTriangleIndexBuffer);
 
-	Parameters->RestPosition0Buffer = RegisterAsSRV(GraphBuilder, RestLODDatas.RestUniqueTrianglePosition0Buffer);
-	Parameters->RestPosition1Buffer = RegisterAsSRV(GraphBuilder, RestLODDatas.RestUniqueTrianglePosition1Buffer);
-	Parameters->RestPosition2Buffer = RegisterAsSRV(GraphBuilder, RestLODDatas.RestUniqueTrianglePosition2Buffer);
-
-	Parameters->DeformedPosition0Buffer = RegisterAsSRV(GraphBuilder, DeformedLODDatas.GetDeformedUniqueTrianglePosition0Buffer(FHairStrandsDeformedRootResource::FLOD::Current));
-	Parameters->DeformedPosition1Buffer = RegisterAsSRV(GraphBuilder, DeformedLODDatas.GetDeformedUniqueTrianglePosition1Buffer(FHairStrandsDeformedRootResource::FLOD::Current));
-	Parameters->DeformedPosition2Buffer = RegisterAsSRV(GraphBuilder, DeformedLODDatas.GetDeformedUniqueTrianglePosition2Buffer(FHairStrandsDeformedRootResource::FLOD::Current));
+	Parameters->RestPositionBuffer = RegisterAsSRV(GraphBuilder, RestLODDatas.RestUniqueTrianglePositionBuffer);
+	Parameters->DeformedPositionBuffer = RegisterAsSRV(GraphBuilder, DeformedLODDatas.GetDeformedUniqueTrianglePositionBuffer(FHairStrandsDeformedRootResource::FLOD::Current));
 
 	Parameters->RestSamplePositionsBuffer = RegisterAsSRV(GraphBuilder, RestLODDatas.RestSamplePositionsBuffer);
 	Parameters->DeformedSamplePositionsBuffer = RegisterAsSRV(GraphBuilder, DeformedLODDatas.GetDeformedSamplePositionsBuffer(FHairStrandsDeformedRootResource::FLOD::Current));
diff --git a/Engine/Plugins/Runtime/HairStrands/Source/HairStrandsCore/Private/GroomManager.cpp b/Engine/Plugins/Runtime/HairStrands/Source/HairStrandsCore/Private/GroomManager.cpp
--- a/Engine/Plugins/Runtime/HairStrands/Source/HairStrandsCore/Private/GroomManager.cpp
+++ b/Engine/Plugins/Runtime/HairStrands/Source/HairStrandsCore/Private/GroomManager.cpp
@@ -359,9 +359,7 @@ static void RunInternalHairStrandsInterpolation(
 						// Add manual transition for the GPU solver as Niagara does not track properly the RDG buffer, and so doesn't issue the correct transitions
 						if (MeshLODIndex >= 0)
 						{
-							GraphBuilder.UseExternalAccessMode(Register(GraphBuilder, Instance->Guides.DeformedRootResource->LODs[MeshLODIndex].GetDeformedUniqueTrianglePosition0Buffer(FHairStrandsDeformedRootResource::FLOD::Current), ERDGImportedBufferFlags::CreateSRV).Buffer, ERHIAccess::SRVMask);
-							GraphBuilder.UseExternalAccessMode(Register(GraphBuilder, Instance->Guides.DeformedRootResource->LODs[MeshLODIndex].GetDeformedUniqueTrianglePosition1Buffer(FHairStrandsDeformedRootResource::FLOD::Current), ERDGImportedBufferFlags::CreateSRV).Buffer, ERHIAccess::SRVMask);
-							GraphBuilder.UseExternalAccessMode(Register(GraphBuilder, Instance->Guides.DeformedRootResource->LODs[MeshLODIndex].GetDeformedUniqueTrianglePosition2Buffer(FHairStrandsDeformedRootResource::FLOD::Current), ERDGImportedBufferFlags::CreateSRV).Buffer, ERHIAccess::SRVMask);
+							GraphBuilder.UseExternalAccessMode(Register(GraphBuilder, Instance->Guides.DeformedRootResource->LODs[MeshLODIndex].GetDeformedUniqueTrianglePositionBuffer(FHairStrandsDeformedRootResource::FLOD::Current), ERDGImportedBufferFlags::CreateSRV).Buffer, ERHIAccess::SRVMask);
 
 							if (bGlobalDeformationEnable)
 							{
diff --git a/Engine/Plugins/Runtime/HairStrands/Source/HairStrandsCore/Private/GroomRBFDeformer.cpp b/Engine/Plugins/Runtime/HairStrands/Source/HairStrandsCore/Private/GroomRBFDeformer.cpp
--- a/Engine/Plugins/Runtime/HairStrands/Source/HairStrandsCore/Private/GroomRBFDeformer.cpp
+++ b/Engine/Plugins/Runtime/HairStrands/Source/HairStrandsCore/Private/GroomRBFDeformer.cpp
@@ -84,15 +84,8 @@ void DeformStrands(
 	const TArray<FHairStrandsIndexFormat::Type>& PointToCurveBuffer,
 	const TArray<FHairStrandsIndexFormat::Type>& RootToUniqueTriangleBuffer,
 	const TArray<FHairStrandsRootBarycentricFormat::Type>& RootBarycentricBuffer,
-
-	const TArray<FHairStrandsMeshTrianglePositionFormat::Type>& UniqueTrianglePosition0Buffer_Rest,
-	const TArray<FHairStrandsMeshTrianglePositionFormat::Type>& UniqueTrianglePosition1Buffer_Rest,
-	const TArray<FHairStrandsMeshTrianglePositionFormat::Type>& UniqueTrianglePosition2Buffer_Rest,
-
-	const TArray<FHairStrandsMeshTrianglePositionFormat::Type>& UniqueTrianglePosition0Buffer_Deformed,
-	const TArray<FHairStrandsMeshTrianglePositionFormat::Type>& UniqueTrianglePosition1Buffer_Deformed,
-	const TArray<FHairStrandsMeshTrianglePositionFormat::Type>& UniqueTrianglePosition2Buffer_Deformed,
-
+	const TArray<FHairStrandsMeshTrianglePositionFormat::Type>& UniqueTrianglePositionBuffer_Rest,
+	const TArray<FHairStrandsMeshTrianglePositionFormat::Type>& UniqueTrianglePositionBuffer_Deformed,
 	uint32 VertexCount,
 	uint32 SampleCount,
 	const TArray<FVector3f>& RestPosePositionBuffer,
@@ -131,15 +124,15 @@ void DeformStrands(
 		const FVector3f   B  = FVector3f(B0.X, B0.Y, 1.f - B0.X - B0.Y);
 
 		/* Strand hair roots translation and rotation in rest position relative to the bound triangle. Positions are relative to the rest root center */
-		const FVector3f& Rest_V0 = UniqueTrianglePosition0Buffer_Rest[TriangleIndex];
-		const FVector3f& Rest_V1 = UniqueTrianglePosition1Buffer_Rest[TriangleIndex];
-		const FVector3f& Rest_V2 = UniqueTrianglePosition2Buffer_Rest[TriangleIndex];
+		const FVector3f& Rest_V0 = UniqueTrianglePositionBuffer_Rest[TriangleIndex * 3 + 0];
+		const FVector3f& Rest_V1 = UniqueTrianglePositionBuffer_Rest[TriangleIndex * 3 + 1];
+		const FVector3f& Rest_V2 = UniqueTrianglePositionBuffer_Rest[TriangleIndex * 3 + 2];
 
-		const FVector3f& Deform_V0 = UniqueTrianglePosition0Buffer_Deformed[TriangleIndex];
-		const FVector3f& Deform_V1 = UniqueTrianglePosition1Buffer_Deformed[TriangleIndex];
-		const FVector3f& Deform_V2 = UniqueTrianglePosition2Buffer_Deformed[TriangleIndex];
+		const FVector3f& Deform_V0 = UniqueTrianglePositionBuffer_Deformed[TriangleIndex * 3 + 0];
+		const FVector3f& Deform_V1 = UniqueTrianglePositionBuffer_Deformed[TriangleIndex * 3 + 1];
+		const FVector3f& Deform_V2 = UniqueTrianglePositionBuffer_Deformed[TriangleIndex * 3 + 2];
 
-		const FVector3f Rest_RootPosition		=   Rest_V0 * B.X +   Rest_V1 * B.Y +   Rest_V2 * B.Z;
+		const FVector3f Rest_RootPosition	=   Rest_V0 * B.X +   Rest_V1 * B.Y +   Rest_V2 * B.Z;
 		const FVector3f Deform_RootPosition	= Deform_V0 * B.X + Deform_V1 * B.Y + Deform_V2 * B.Z;
 
 		const FVector3f RestOffset = Rest_Position - Rest_RootPosition;
@@ -164,13 +157,9 @@ void ExtractUniqueTrianglePositions(
 	const uint32 MeshLODIndex,
 	const uint32 UniqueTriangleCount,
 	const FSkeletalMeshRenderData* InMeshRenderData, 
-	TArray<FHairStrandsMeshTrianglePositionFormat::Type>& OutDeformUniqueTrianglePosition0Buffer,
-	TArray<FHairStrandsMeshTrianglePositionFormat::Type>& OutDeformUniqueTrianglePosition1Buffer,
-	TArray<FHairStrandsMeshTrianglePositionFormat::Type>& OutDeformUniqueTrianglePosition2Buffer)
+	TArray<FHairStrandsMeshTrianglePositionFormat::Type>& OutDeformUniqueTrianglePositionBuffer)
 {
-	OutDeformUniqueTrianglePosition0Buffer.SetNum(UniqueTriangleCount);
-	OutDeformUniqueTrianglePosition1Buffer.SetNum(UniqueTriangleCount);
-	OutDeformUniqueTrianglePosition2Buffer.SetNum(UniqueTriangleCount);
+	OutDeformUniqueTrianglePositionBuffer.SetNum(UniqueTriangleCount * 3);
 
 	const uint32 SectionCount = InMeshRenderData->LODRenderData[MeshLODIndex].RenderSections.Num();
 	TArray<uint32> IndexBuffer;
@@ -195,9 +184,9 @@ void ExtractUniqueTrianglePositions(
 		const FVector3f P1 = InMeshRenderData->LODRenderData[MeshLODIndex].StaticVertexBuffers.PositionVertexBuffer.VertexPosition(I1);
 		const FVector3f P2 = InMeshRenderData->LODRenderData[MeshLODIndex].StaticVertexBuffers.PositionVertexBuffer.VertexPosition(I2);
 
-		OutDeformUniqueTrianglePosition0Buffer[UniqueTriangleIndex] = P0;
-		OutDeformUniqueTrianglePosition1Buffer[UniqueTriangleIndex] = P1;
-		OutDeformUniqueTrianglePosition2Buffer[UniqueTriangleIndex] = P2;
+		OutDeformUniqueTrianglePositionBuffer[UniqueTriangleIndex * 3 + 0] = P0;
+		OutDeformUniqueTrianglePositionBuffer[UniqueTriangleIndex * 3 + 1] = P1;
+		OutDeformUniqueTrianglePositionBuffer[UniqueTriangleIndex * 3 + 2] = P2;
 	}
 }
 
@@ -234,21 +223,14 @@ TArray<FVector3f> GetDeformedHairStrandsPositions(
 	TArray<FVector3f> OutPositions = Points.PointsPosition;
 
 	// Use the vertex position of the binding, as the source asset might not have the same topology (in case the groom has been transfered from one mesh toanother using UV sharing)
-	TArray<FHairStrandsMeshTrianglePositionFormat::Type> UniqueTrianglePosition0Buffer_Rest = RestLODData.RestUniqueTrianglePosition0Buffer;
-	TArray<FHairStrandsMeshTrianglePositionFormat::Type> UniqueTrianglePosition1Buffer_Rest = RestLODData.RestUniqueTrianglePosition1Buffer;
-	TArray<FHairStrandsMeshTrianglePositionFormat::Type> UniqueTrianglePosition2Buffer_Rest = RestLODData.RestUniqueTrianglePosition2Buffer;
-
-	TArray<FHairStrandsMeshTrianglePositionFormat::Type> UniqueTrianglePosition0Buffer_Deformed;
-	TArray<FHairStrandsMeshTrianglePositionFormat::Type> UniqueTrianglePosition1Buffer_Deformed;
-	TArray<FHairStrandsMeshTrianglePositionFormat::Type> UniqueTrianglePosition2Buffer_Deformed;
+	TArray<FHairStrandsMeshTrianglePositionFormat::Type> UniqueTrianglePositionBuffer_Rest = RestLODData.RestUniqueTrianglePositionBuffer;
+	TArray<FHairStrandsMeshTrianglePositionFormat::Type> UniqueTrianglePositionBuffer_Deformed;
 	ExtractUniqueTrianglePositions(
 		RestLODData,
 		MeshLODIndex,
-		RestLODData.RestUniqueTrianglePosition0Buffer.Num(),
+		RestLODData.RestUniqueTrianglePositionBuffer.Num(),
 		InMeshRenderData,
-		UniqueTrianglePosition0Buffer_Deformed,
-		UniqueTrianglePosition1Buffer_Deformed,
-		UniqueTrianglePosition2Buffer_Deformed);
+		UniqueTrianglePositionBuffer_Deformed);
 
 	// Deform the strands vertices with the deformed mesh samples
 	const TArray<FVector3f>& RestPosePositionBuffer = OutPositions;
@@ -261,15 +243,8 @@ TArray<FVector3f> GetDeformedHairStrandsPositions(
 		PointToCurveBuffer,
 		RootToUniqueTriangleBuffer,
 		RestLODData.RootBarycentricBuffer,
-
-		UniqueTrianglePosition0Buffer_Rest,
-		UniqueTrianglePosition1Buffer_Rest,
-		UniqueTrianglePosition2Buffer_Rest,
-
-		UniqueTrianglePosition0Buffer_Deformed,
-		UniqueTrianglePosition1Buffer_Deformed,
-		UniqueTrianglePosition2Buffer_Deformed,
-
+		UniqueTrianglePositionBuffer_Rest,
+		UniqueTrianglePositionBuffer_Deformed,
 		VertexCount, 
 		MaxSampleCount, 
 		RestPosePositionBuffer, 
diff --git a/Engine/Plugins/Runtime/HairStrands/Source/HairStrandsCore/Private/GroomResources.cpp b/Engine/Plugins/Runtime/HairStrands/Source/HairStrandsCore/Private/GroomResources.cpp
--- a/Engine/Plugins/Runtime/HairStrands/Source/HairStrandsCore/Private/GroomResources.cpp
+++ b/Engine/Plugins/Runtime/HairStrands/Source/HairStrandsCore/Private/GroomResources.cpp
@@ -1128,10 +1128,7 @@ bool FHairStrandsRestRootResource::InternalIsLODDataLoaded(int32 LODIndex)
 				Batch.Read(CPUData.RootBarycentricBuffer);
 				Batch.Read(CPUData.RootToUniqueTriangleIndexBuffer);
 				Batch.Read(CPUData.UniqueTriangleIndexBuffer);
-
-				Batch.Read(CPUData.RestUniqueTrianglePosition0Buffer);
-				Batch.Read(CPUData.RestUniqueTrianglePosition1Buffer);
-				Batch.Read(CPUData.RestUniqueTrianglePosition2Buffer);
+				Batch.Read(CPUData.RestUniqueTrianglePositionBuffer);
 			}
 			
 			const bool bHasValidCPUWeights = CPUData.MeshSampleIndicesBuffer.GetBulkDataSize() > 0;
@@ -1190,10 +1187,7 @@ void FHairStrandsRestRootResource::InternalAllocateLOD(FRDGBuilder& GraphBuilder
 			InternalCreateVertexBufferRDG_FromBulkData<FHairStrandsRootBarycentricFormat>(GraphBuilder, CPUData.RootBarycentricBuffer, BulkData.RootCount, GPUData.RootBarycentricBuffer, ToHairResourceDebugName(HAIRSTRANDS_RESOUCE_NAME(CurveType, Hair.StrandsRestRoot_RootTriangleBarycentricBuffer), ResourceName), OwnerName, EHairResourceUsageType::Static);
 			InternalCreateVertexBufferRDG_FromBulkData<FHairStrandsRootToUniqueTriangleIndexFormat>(GraphBuilder, CPUData.RootToUniqueTriangleIndexBuffer, BulkData.RootCount, GPUData.RootToUniqueTriangleIndexBuffer, ToHairResourceDebugName(HAIRSTRANDS_RESOUCE_NAME(CurveType, Hair.StrandsRestRoot_RootToUniqueTriangleIndexBuffer), ResourceName), OwnerName, EHairResourceUsageType::Static);
 			InternalCreateVertexBufferRDG_FromBulkData<FHairStrandsUniqueTriangleIndexFormat>(GraphBuilder, CPUData.UniqueTriangleIndexBuffer, CPUData.UniqueTriangleCount, GPUData.UniqueTriangleIndexBuffer, ToHairResourceDebugName(HAIRSTRANDS_RESOUCE_NAME(CurveType, Hair.StrandsRestRoot_UniqueTriangleIndexBuffer), ResourceName), OwnerName, EHairResourceUsageType::Static);
-
-			InternalCreateVertexBufferRDG_FromBulkData<FHairStrandsMeshTrianglePositionFormat>(GraphBuilder, CPUData.RestUniqueTrianglePosition0Buffer, CPUData.UniqueTriangleCount, GPUData.RestUniqueTrianglePosition0Buffer, ToHairResourceDebugName(HAIRSTRANDS_RESOUCE_NAME(CurveType, Hair.StrandsRestRoot_RestUniqueTrianglePosition0Buffer), ResourceName), OwnerName, EHairResourceUsageType::Static);
-			InternalCreateVertexBufferRDG_FromBulkData<FHairStrandsMeshTrianglePositionFormat>(GraphBuilder, CPUData.RestUniqueTrianglePosition1Buffer, CPUData.UniqueTriangleCount, GPUData.RestUniqueTrianglePosition1Buffer, ToHairResourceDebugName(HAIRSTRANDS_RESOUCE_NAME(CurveType, Hair.StrandsRestRoot_RestUniqueTrianglePosition1Buffer), ResourceName), OwnerName, EHairResourceUsageType::Static);
-			InternalCreateVertexBufferRDG_FromBulkData<FHairStrandsMeshTrianglePositionFormat>(GraphBuilder, CPUData.RestUniqueTrianglePosition2Buffer, CPUData.UniqueTriangleCount, GPUData.RestUniqueTrianglePosition2Buffer, ToHairResourceDebugName(HAIRSTRANDS_RESOUCE_NAME(CurveType, Hair.StrandsRestRoot_RestUniqueTrianglePosition2Buffer), ResourceName), OwnerName, EHairResourceUsageType::Static);
+			InternalCreateVertexBufferRDG_FromBulkData<FHairStrandsMeshTrianglePositionFormat>(GraphBuilder, CPUData.RestUniqueTrianglePositionBuffer, CPUData.UniqueTriangleCount * 3, GPUData.RestUniqueTrianglePositionBuffer, ToHairResourceDebugName(HAIRSTRANDS_RESOUCE_NAME(CurveType, Hair.StrandsRestRoot_RestUniqueTrianglePosition0Buffer), ResourceName), OwnerName, EHairResourceUsageType::Static);
 		}
 		else
 		{
@@ -1204,9 +1198,7 @@ void FHairStrandsRestRootResource::InternalAllocateLOD(FRDGBuilder& GraphBuilder
 			InternalCreateVertexBufferRDG<FHairStrandsUniqueTriangleIndexFormat>(GraphBuilder, CPUData.UniqueTriangleCount, GPUData.UniqueTriangleIndexBuffer, ToHairResourceDebugName(HAIRSTRANDS_RESOUCE_NAME(CurveType, Hair.StrandsRestRoot_UniqueTriangleIndexBuffer), ResourceName), OwnerName, EHairResourceUsageType::Dynamic);
 			
 			// Create buffers. Initialization will be done by render passes
-			InternalCreateVertexBufferRDG<FHairStrandsMeshTrianglePositionFormat>(GraphBuilder, CPUData.UniqueTriangleCount, GPUData.RestUniqueTrianglePosition0Buffer, ToHairResourceDebugName(HAIRSTRANDS_RESOUCE_NAME(CurveType, Hair.StrandsRestRoot_RestUniqueTrianglePosition0Buffer), ResourceName), OwnerName, EHairResourceUsageType::Dynamic);
-			InternalCreateVertexBufferRDG<FHairStrandsMeshTrianglePositionFormat>(GraphBuilder, CPUData.UniqueTriangleCount, GPUData.RestUniqueTrianglePosition1Buffer, ToHairResourceDebugName(HAIRSTRANDS_RESOUCE_NAME(CurveType, Hair.StrandsRestRoot_RestUniqueTrianglePosition1Buffer), ResourceName), OwnerName, EHairResourceUsageType::Dynamic);
-			InternalCreateVertexBufferRDG<FHairStrandsMeshTrianglePositionFormat>(GraphBuilder, CPUData.UniqueTriangleCount, GPUData.RestUniqueTrianglePosition2Buffer, ToHairResourceDebugName(HAIRSTRANDS_RESOUCE_NAME(CurveType, Hair.StrandsRestRoot_RestUniqueTrianglePosition2Buffer), ResourceName), OwnerName, EHairResourceUsageType::Dynamic);
+			InternalCreateVertexBufferRDG<FHairStrandsMeshTrianglePositionFormat>(GraphBuilder, CPUData.UniqueTriangleCount * 3, GPUData.RestUniqueTrianglePositionBuffer, ToHairResourceDebugName(HAIRSTRANDS_RESOUCE_NAME(CurveType, Hair.StrandsRestRoot_RestUniqueTrianglePosition0Buffer), ResourceName), OwnerName, EHairResourceUsageType::Dynamic);
 		}
 
 		GPUData.SampleCount = CPUData.SampleCount;
@@ -1236,9 +1228,7 @@ void FHairStrandsRestRootResource::InternalRelease()
 		GPUData.RootBarycentricBuffer.Release();
 		GPUData.RootToUniqueTriangleIndexBuffer.Release();
 		GPUData.UniqueTriangleIndexBuffer.Release();
-		GPUData.RestUniqueTrianglePosition0Buffer.Release();
-		GPUData.RestUniqueTrianglePosition1Buffer.Release();
-		GPUData.RestUniqueTrianglePosition2Buffer.Release();
+		GPUData.RestUniqueTrianglePositionBuffer.Release();
 		GPUData.SampleCount = 0;
 		GPUData.MeshInterpolationWeightsBuffer.Release();
 		GPUData.MeshSampleIndicesBuffer.Release();
@@ -1296,16 +1286,12 @@ void FHairStrandsDeformedRootResource::InternalAllocateLOD(FRDGBuilder& GraphBui
 				}
 			}
 
-			InternalCreateVertexBufferRDG<FHairStrandsMeshTrianglePositionFormat>(GraphBuilder, RootCount, LOD.DeformedUniqueTrianglePosition0Buffer[0], ToHairResourceDebugName(HAIRSTRANDS_RESOUCE_NAME(CurveType, Hair.StrandsRootDeformed_DeformedUniqueTrianglePosition0Buffer0), ResourceName), OwnerName, EHairResourceUsageType::Dynamic);
-			InternalCreateVertexBufferRDG<FHairStrandsMeshTrianglePositionFormat>(GraphBuilder, RootCount, LOD.DeformedUniqueTrianglePosition1Buffer[0], ToHairResourceDebugName(HAIRSTRANDS_RESOUCE_NAME(CurveType, Hair.StrandsRootDeformed_DeformedUniqueTrianglePosition1Buffer0), ResourceName), OwnerName, EHairResourceUsageType::Dynamic);
-			InternalCreateVertexBufferRDG<FHairStrandsMeshTrianglePositionFormat>(GraphBuilder, RootCount, LOD.DeformedUniqueTrianglePosition2Buffer[0], ToHairResourceDebugName(HAIRSTRANDS_RESOUCE_NAME(CurveType, Hair.StrandsRootDeformed_DeformedUniqueTrianglePosition2Buffer0), ResourceName), OwnerName, EHairResourceUsageType::Dynamic);
+			InternalCreateVertexBufferRDG<FHairStrandsMeshTrianglePositionFormat>(GraphBuilder, RootCount * 3, LOD.DeformedUniqueTrianglePositionBuffer[0], ToHairResourceDebugName(HAIRSTRANDS_RESOUCE_NAME(CurveType, Hair.StrandsRootDeformed_DeformedUniqueTrianglePosition0Buffer0), ResourceName), OwnerName, EHairResourceUsageType::Dynamic);
 
 			// Double buffering is disabled by default unless the read-only cvar r.HairStrands.ContinuousDecimationReordering is set
 			if (IsHairStrandContinuousDecimationReorderingEnabled())
 			{
-				InternalCreateVertexBufferRDG<FHairStrandsMeshTrianglePositionFormat>(GraphBuilder, RootCount, LOD.DeformedUniqueTrianglePosition0Buffer[1], ToHairResourceDebugName(HAIRSTRANDS_RESOUCE_NAME(CurveType, Hair.StrandsRootDeformed_DeformedUniqueTrianglePosition0Buffer1), ResourceName), OwnerName, EHairResourceUsageType::Dynamic);
-				InternalCreateVertexBufferRDG<FHairStrandsMeshTrianglePositionFormat>(GraphBuilder, RootCount, LOD.DeformedUniqueTrianglePosition1Buffer[1], ToHairResourceDebugName(HAIRSTRANDS_RESOUCE_NAME(CurveType, Hair.StrandsRootDeformed_DeformedUniqueTrianglePosition1Buffer1), ResourceName), OwnerName, EHairResourceUsageType::Dynamic);
-				InternalCreateVertexBufferRDG<FHairStrandsMeshTrianglePositionFormat>(GraphBuilder, RootCount, LOD.DeformedUniqueTrianglePosition2Buffer[1], ToHairResourceDebugName(HAIRSTRANDS_RESOUCE_NAME(CurveType, Hair.StrandsRootDeformed_DeformedUniqueTrianglePosition2Buffer1), ResourceName), OwnerName, EHairResourceUsageType::Dynamic);
+				InternalCreateVertexBufferRDG<FHairStrandsMeshTrianglePositionFormat>(GraphBuilder, RootCount * 3, LOD.DeformedUniqueTrianglePositionBuffer[1], ToHairResourceDebugName(HAIRSTRANDS_RESOUCE_NAME(CurveType, Hair.StrandsRootDeformed_DeformedUniqueTrianglePosition0Buffer1), ResourceName), OwnerName, EHairResourceUsageType::Dynamic);
 			}
 		}
 	}
@@ -1316,18 +1302,14 @@ void FHairStrandsDeformedRootResource::InternalRelease()
 	for (FLOD& GPUData : LODs)
 	{
 		GPUData.Status = FLOD::EStatus::Invalid;
-		GPUData.DeformedUniqueTrianglePosition0Buffer[0].Release();
-		GPUData.DeformedUniqueTrianglePosition1Buffer[0].Release();
-		GPUData.DeformedUniqueTrianglePosition2Buffer[0].Release();
+		GPUData.DeformedUniqueTrianglePositionBuffer[0].Release();
 		GPUData.DeformedSamplePositionsBuffer[0].Release();
 		GPUData.MeshSampleWeightsBuffer[0].Release();
 	
 		// Double buffering is disabled by default unless the read-only cvar r.HairStrands.ContinuousDecimationReordering is set
 		if (IsHairStrandContinuousDecimationReorderingEnabled())
 		{
-			GPUData.DeformedUniqueTrianglePosition0Buffer[1].Release();
-			GPUData.DeformedUniqueTrianglePosition1Buffer[1].Release();
-			GPUData.DeformedUniqueTrianglePosition2Buffer[1].Release();
+			GPUData.DeformedUniqueTrianglePositionBuffer[1].Release();
 			GPUData.DeformedSamplePositionsBuffer[1].Release();
 			GPUData.MeshSampleWeightsBuffer[1].Release();
 		}
@@ -1371,9 +1353,7 @@ static void InternalSerialize(FArchive& Ar, UObject* Owner, FHairStrandsRootBulk
 		LOD.RootBarycentricBuffer.SetBulkDataFlags(BulkFlags);
 		LOD.RootToUniqueTriangleIndexBuffer.SetBulkDataFlags(BulkFlags);
 		LOD.UniqueTriangleIndexBuffer.SetBulkDataFlags(BulkFlags);
-		LOD.RestUniqueTrianglePosition0Buffer.SetBulkDataFlags(BulkFlags);
-		LOD.RestUniqueTrianglePosition1Buffer.SetBulkDataFlags(BulkFlags);
-		LOD.RestUniqueTrianglePosition2Buffer.SetBulkDataFlags(BulkFlags);
+		LOD.RestUniqueTrianglePositionBuffer.SetBulkDataFlags(BulkFlags);
 
 		LOD.MeshInterpolationWeightsBuffer.SetBulkDataFlags(BulkFlags);
 		LOD.MeshSampleIndicesBuffer.SetBulkDataFlags(BulkFlags);
@@ -1385,9 +1365,7 @@ static void InternalSerialize(FArchive& Ar, UObject* Owner, FHairStrandsRootBulk
 	LOD.RootBarycentricBuffer.Serialize(Ar, Owner, ChunkIndex, bAttemptFileMapping);
 	LOD.RootToUniqueTriangleIndexBuffer.Serialize(Ar, Owner, ChunkIndex, bAttemptFileMapping);
 	LOD.UniqueTriangleIndexBuffer.Serialize(Ar, Owner, ChunkIndex, bAttemptFileMapping);
-	LOD.RestUniqueTrianglePosition0Buffer.Serialize(Ar, Owner, ChunkIndex, bAttemptFileMapping);
-	LOD.RestUniqueTrianglePosition1Buffer.Serialize(Ar, Owner, ChunkIndex, bAttemptFileMapping);
-	LOD.RestUniqueTrianglePosition2Buffer.Serialize(Ar, Owner, ChunkIndex, bAttemptFileMapping);
+	LOD.RestUniqueTrianglePositionBuffer.Serialize(Ar, Owner, ChunkIndex, bAttemptFileMapping);
 
 	Ar << LOD.SampleCount;
 	if (LOD.SampleCount)
@@ -1434,9 +1412,7 @@ void FHairStrandsRootBulkData::Reset()
 		LOD.RootBarycentricBuffer.RemoveBulkData();
 		LOD.RootToUniqueTriangleIndexBuffer.RemoveBulkData();
 		LOD.UniqueTriangleIndexBuffer.RemoveBulkData();
-		LOD.RestUniqueTrianglePosition0Buffer.RemoveBulkData();
-		LOD.RestUniqueTrianglePosition1Buffer.RemoveBulkData();
-		LOD.RestUniqueTrianglePosition2Buffer.RemoveBulkData();
+		LOD.RestUniqueTrianglePositionBuffer.RemoveBulkData();
 		LOD.MeshInterpolationWeightsBuffer.RemoveBulkData();
 		LOD.MeshSampleIndicesBuffer.RemoveBulkData();
 		LOD.RestSamplePositionsBuffer.RemoveBulkData();
@@ -1445,9 +1421,7 @@ void FHairStrandsRootBulkData::Reset()
 		LOD.RootBarycentricBuffer 			 = FByteBulkData();
 		LOD.RootToUniqueTriangleIndexBuffer  = FByteBulkData();
 		LOD.UniqueTriangleIndexBuffer		 = FByteBulkData();
-		LOD.RestUniqueTrianglePosition0Buffer= FByteBulkData();
-		LOD.RestUniqueTrianglePosition1Buffer= FByteBulkData();
-		LOD.RestUniqueTrianglePosition2Buffer= FByteBulkData();
+		LOD.RestUniqueTrianglePositionBuffer = FByteBulkData();
 		LOD.MeshInterpolationWeightsBuffer 	 = FByteBulkData();
 		LOD.MeshSampleIndicesBuffer 		 = FByteBulkData();
 		LOD.RestSamplePositionsBuffer 		 = FByteBulkData();
@@ -1475,9 +1449,7 @@ bool FHairStrandsRootData::HasProjectionData() const
 			bIsValid = bIsValid && LOD.RootBarycentricBuffer.Num() > 0;
 			bIsValid = bIsValid && LOD.RootToUniqueTriangleIndexBuffer.Num() > 0;
 			bIsValid = bIsValid && LOD.UniqueTriangleIndexBuffer.Num() > 0;
-			bIsValid = bIsValid && LOD.RestUniqueTrianglePosition0Buffer.Num() > 0;
-			bIsValid = bIsValid && LOD.RestUniqueTrianglePosition1Buffer.Num() > 0;
-			bIsValid = bIsValid && LOD.RestUniqueTrianglePosition2Buffer.Num() > 0;
+			bIsValid = bIsValid && LOD.RestUniqueTrianglePositionBuffer.Num() > 0;
 
 			if (!bIsValid) break;
 		}
diff --git a/Engine/Plugins/Runtime/HairStrands/Source/HairStrandsCore/Private/HairStrandsMeshProjection.cpp b/Engine/Plugins/Runtime/HairStrands/Source/HairStrandsCore/Private/HairStrandsMeshProjection.cpp
--- a/Engine/Plugins/Runtime/HairStrands/Source/HairStrandsCore/Private/HairStrandsMeshProjection.cpp
+++ b/Engine/Plugins/Runtime/HairStrands/Source/HairStrandsCore/Private/HairStrandsMeshProjection.cpp
@@ -239,13 +239,8 @@ private:
 		SHADER_PARAMETER_SRV(Buffer, MeshUVsBuffer7)
 
 		SHADER_PARAMETER_RDG_BUFFER_SRV(Buffer, UniqueTriangleIndices)
-		SHADER_PARAMETER_RDG_BUFFER_UAV(RWStructuredBuffer, OutUniqueTrianglePrevPosition0)
-		SHADER_PARAMETER_RDG_BUFFER_UAV(RWStructuredBuffer, OutUniqueTrianglePrevPosition1)
-		SHADER_PARAMETER_RDG_BUFFER_UAV(RWStructuredBuffer, OutUniqueTrianglePrevPosition2)
-
-		SHADER_PARAMETER_RDG_BUFFER_UAV(RWStructuredBuffer, OutUniqueTriangleCurrPosition0)
-		SHADER_PARAMETER_RDG_BUFFER_UAV(RWStructuredBuffer, OutUniqueTriangleCurrPosition1)
-		SHADER_PARAMETER_RDG_BUFFER_UAV(RWStructuredBuffer, OutUniqueTriangleCurrPosition2)
+		SHADER_PARAMETER_RDG_BUFFER_UAV(RWStructuredBuffer, OutUniqueTrianglePrevPosition)
+		SHADER_PARAMETER_RDG_BUFFER_UAV(RWStructuredBuffer, OutUniqueTriangleCurrPosition)
 	END_SHADER_PARAMETER_STRUCT()
 
 public:
@@ -317,33 +312,27 @@ void AddHairStrandUpdateMeshTrianglesPass(
 	const uint32 PassCount = FMath::DivideAndRoundUp(ValidSectionCount, FHairUpdateMeshTriangleCS::SectionArrayCount);
 	const bool bComputePreviousDeformedPosition = IsHairStrandContinuousDecimationReorderingEnabled() && Type == HairStrandsTriangleType::DeformedPose;
 
-	const uint32 MaxUniqueTriangleCount = RestLODData.RestUniqueTrianglePosition0Buffer.Buffer->Desc.NumElements;
+	const uint32 MaxUniqueTriangleCount = RestLODData.RestUniqueTrianglePositionBuffer.Buffer->Desc.NumElements * 3;
 
 	FHairUpdateMeshTriangleCS::FParameters CommonParameters;
 
-	FRDGImportedBuffer OutputCurrBuffers[3];
-	FRDGImportedBuffer OutputPrevBuffers[3];
+	FRDGImportedBuffer OutputCurrBuffer;
+	FRDGImportedBuffer OutputPrevBuffer;
 	const bool bEnableUAVOverlap = true;
 	const ERDGUnorderedAccessViewFlags UAVFlags = bEnableUAVOverlap ? ERDGUnorderedAccessViewFlags::SkipBarrier : ERDGUnorderedAccessViewFlags::None;
 	CommonParameters.UniqueTriangleIndices = RegisterAsSRV(GraphBuilder, RestLODData.UniqueTriangleIndexBuffer);
 	if (Type == HairStrandsTriangleType::RestPose)
 	{
-		OutputCurrBuffers[0] = Register(GraphBuilder, RestLODData.RestUniqueTrianglePosition0Buffer, ERDGImportedBufferFlags::CreateUAV, UAVFlags);
-		OutputCurrBuffers[1] = Register(GraphBuilder, RestLODData.RestUniqueTrianglePosition1Buffer, ERDGImportedBufferFlags::CreateUAV, UAVFlags);
-		OutputCurrBuffers[2] = Register(GraphBuilder, RestLODData.RestUniqueTrianglePosition2Buffer, ERDGImportedBufferFlags::CreateUAV, UAVFlags);
+		OutputCurrBuffer = Register(GraphBuilder, RestLODData.RestUniqueTrianglePositionBuffer, ERDGImportedBufferFlags::CreateUAV, UAVFlags);
 	}
 	else if (Type == HairStrandsTriangleType::DeformedPose)
 	{
 		FHairStrandsDeformedRootResource::FLOD& DeformedLODData = DeformedResources->LODs[LODIndex];
-		OutputCurrBuffers[0] = Register(GraphBuilder, DeformedLODData.GetDeformedUniqueTrianglePosition0Buffer(FHairStrandsDeformedRootResource::FLOD::Current), ERDGImportedBufferFlags::CreateUAV, UAVFlags);
-		OutputCurrBuffers[1] = Register(GraphBuilder, DeformedLODData.GetDeformedUniqueTrianglePosition1Buffer(FHairStrandsDeformedRootResource::FLOD::Current), ERDGImportedBufferFlags::CreateUAV, UAVFlags);
-		OutputCurrBuffers[2] = Register(GraphBuilder, DeformedLODData.GetDeformedUniqueTrianglePosition2Buffer(FHairStrandsDeformedRootResource::FLOD::Current), ERDGImportedBufferFlags::CreateUAV, UAVFlags);
+		OutputCurrBuffer = Register(GraphBuilder, DeformedLODData.GetDeformedUniqueTrianglePositionBuffer(FHairStrandsDeformedRootResource::FLOD::Current), ERDGImportedBufferFlags::CreateUAV, UAVFlags);
 
 		if (bComputePreviousDeformedPosition)
 		{
-			OutputPrevBuffers[0] = Register(GraphBuilder, DeformedLODData.GetDeformedUniqueTrianglePosition0Buffer(FHairStrandsDeformedRootResource::FLOD::Previous), ERDGImportedBufferFlags::CreateUAV, UAVFlags);
-			OutputPrevBuffers[1] = Register(GraphBuilder, DeformedLODData.GetDeformedUniqueTrianglePosition1Buffer(FHairStrandsDeformedRootResource::FLOD::Previous), ERDGImportedBufferFlags::CreateUAV, UAVFlags);
-			OutputPrevBuffers[2] = Register(GraphBuilder, DeformedLODData.GetDeformedUniqueTrianglePosition2Buffer(FHairStrandsDeformedRootResource::FLOD::Previous), ERDGImportedBufferFlags::CreateUAV, UAVFlags);
+			OutputPrevBuffer = Register(GraphBuilder, DeformedLODData.GetDeformedUniqueTrianglePositionBuffer(FHairStrandsDeformedRootResource::FLOD::Previous), ERDGImportedBufferFlags::CreateUAV, UAVFlags);
 		}
 
 		DeformedLODData.Status = FHairStrandsDeformedRootResource::FLOD::EStatus::Completed;
@@ -354,15 +343,11 @@ void AddHairStrandUpdateMeshTrianglesPass(
 		return;
 	}
 
-	CommonParameters.OutUniqueTriangleCurrPosition0 = OutputCurrBuffers[0].UAV;
-	CommonParameters.OutUniqueTriangleCurrPosition1 = OutputCurrBuffers[1].UAV;
-	CommonParameters.OutUniqueTriangleCurrPosition2 = OutputCurrBuffers[2].UAV;
+	CommonParameters.OutUniqueTriangleCurrPosition = OutputCurrBuffer.UAV;
 
 	if (bComputePreviousDeformedPosition)
 	{
-		CommonParameters.OutUniqueTrianglePrevPosition0 = OutputPrevBuffers[0].UAV;
-		CommonParameters.OutUniqueTrianglePrevPosition1 = OutputPrevBuffers[1].UAV;
-		CommonParameters.OutUniqueTrianglePrevPosition2 = OutputPrevBuffers[2].UAV;
+		CommonParameters.OutUniqueTrianglePrevPosition = OutputPrevBuffer.UAV;
 	}
 
 	for (uint32 PassIt = 0; PassIt < PassCount; ++PassIt)
@@ -539,15 +524,11 @@ void AddHairStrandUpdateMeshTrianglesPass(
 			DispatchGroupCount);
 	}
 
-	GraphBuilder.SetBufferAccessFinal(OutputCurrBuffers[0].Buffer, ERHIAccess::SRVMask);
-	GraphBuilder.SetBufferAccessFinal(OutputCurrBuffers[1].Buffer, ERHIAccess::SRVMask);
-	GraphBuilder.SetBufferAccessFinal(OutputCurrBuffers[2].Buffer, ERHIAccess::SRVMask);
+	GraphBuilder.SetBufferAccessFinal(OutputCurrBuffer.Buffer, ERHIAccess::SRVMask);
 
 	if (bComputePreviousDeformedPosition)
 	{
-		GraphBuilder.SetBufferAccessFinal(OutputPrevBuffers[0].Buffer, ERHIAccess::SRVMask);
-		GraphBuilder.SetBufferAccessFinal(OutputPrevBuffers[1].Buffer, ERHIAccess::SRVMask);
-		GraphBuilder.SetBufferAccessFinal(OutputPrevBuffers[2].Buffer, ERHIAccess::SRVMask);
+		GraphBuilder.SetBufferAccessFinal(OutputPrevBuffer.Buffer, ERHIAccess::SRVMask);
 	}
 
 }
@@ -1032,9 +1013,7 @@ BEGIN_SHADER_PARAMETER_STRUCT(FHairFollicleMaskParameters, )
 	SHADER_PARAMETER(uint32, Channel)
 	SHADER_PARAMETER(uint32, KernelSizeInPixels)
 
-	SHADER_PARAMETER_RDG_BUFFER_SRV(Buffer, UniqueTrianglePosition0Buffer)
-	SHADER_PARAMETER_RDG_BUFFER_SRV(Buffer, UniqueTrianglePosition1Buffer)
-	SHADER_PARAMETER_RDG_BUFFER_SRV(Buffer, UniqueTrianglePosition2Buffer)
+	SHADER_PARAMETER_RDG_BUFFER_SRV(Buffer, UniqueTrianglePositionBuffer)
 	SHADER_PARAMETER_RDG_BUFFER_SRV(Buffer, RootToUniqueTriangleIndexBuffer)
 	SHADER_PARAMETER_RDG_BUFFER_SRV(Buffer, RootBarycentricBuffer)
 	SHADER_PARAMETER_RDG_BUFFER_SRV(Buffer, RootUVsBuffer)
@@ -1101,17 +1080,12 @@ static void AddFollicleMaskPass(
 		return;
 
 	FHairStrandsRestRootResource::FLOD& LODData = RestResources->LODs[LODIndex];
-	if (!LODData.RootBarycentricBuffer.Buffer ||
-		!LODData.RestUniqueTrianglePosition0Buffer.Buffer ||
-		!LODData.RestUniqueTrianglePosition1Buffer.Buffer ||
-		!LODData.RestUniqueTrianglePosition2Buffer.Buffer)
+	if (!LODData.RootBarycentricBuffer.Buffer || !LODData.RestUniqueTrianglePositionBuffer.Buffer)
 		return;
 
 	const FIntPoint OutputResolution = OutTexture->Desc.Extent;
 	FHairFollicleMaskParameters* Parameters = GraphBuilder.AllocParameters<FHairFollicleMaskParameters>();
-	Parameters->UniqueTrianglePosition0Buffer = RegisterAsSRV(GraphBuilder, LODData.RestUniqueTrianglePosition0Buffer);
-	Parameters->UniqueTrianglePosition1Buffer = RegisterAsSRV(GraphBuilder, LODData.RestUniqueTrianglePosition1Buffer);
-	Parameters->UniqueTrianglePosition2Buffer = RegisterAsSRV(GraphBuilder, LODData.RestUniqueTrianglePosition2Buffer);
+	Parameters->UniqueTrianglePositionBuffer = RegisterAsSRV(GraphBuilder, LODData.RestUniqueTrianglePositionBuffer);
 	Parameters->RootToUniqueTriangleIndexBuffer = RegisterAsSRV(GraphBuilder, LODData.RootToUniqueTriangleIndexBuffer);
 	Parameters->RootBarycentricBuffer   = RegisterAsSRV(GraphBuilder, LODData.RootBarycentricBuffer);
 	Parameters->RootUVsBuffer = nullptr;
@@ -1173,9 +1147,7 @@ static void AddFollicleMaskPass(
 {
 	const FIntPoint OutputResolution = OutTexture->Desc.Extent;
 	FHairFollicleMaskParameters* Parameters = GraphBuilder.AllocParameters<FHairFollicleMaskParameters>();
-	Parameters->UniqueTrianglePosition0Buffer = nullptr;
-	Parameters->UniqueTrianglePosition1Buffer = nullptr;
-	Parameters->UniqueTrianglePosition2Buffer = nullptr;
+	Parameters->UniqueTrianglePositionBuffer = nullptr;
 	Parameters->RootBarycentricBuffer = nullptr;
 	Parameters->RootUVsBuffer = GraphBuilder.CreateSRV(RootUVBuffer, PF_G32R32F);
 	Parameters->OutputResolution = OutputResolution;
@@ -1354,8 +1326,8 @@ private:
 	BEGIN_SHADER_PARAMETER_STRUCT(FParameters, )
 		SHADER_PARAMETER(FVector3f, CPUCurrPositionOffset)
 		SHADER_PARAMETER(FVector3f, CPUPrevPositionOffset)
-		SHADER_PARAMETER_RDG_BUFFER_SRV(StructuredBuffer, RootTriangleCurrPosition0Buffer)
-		SHADER_PARAMETER_RDG_BUFFER_SRV(StructuredBuffer, RootTrianglePrevPosition0Buffer)
+		SHADER_PARAMETER_RDG_BUFFER_SRV(StructuredBuffer, RootTriangleCurrPositionBuffer)
+		SHADER_PARAMETER_RDG_BUFFER_SRV(StructuredBuffer, RootTrianglePrevPositionBuffer)
 		SHADER_PARAMETER_RDG_BUFFER_UAV(RWBuffer, OutCurrOffsetBuffer)
 		SHADER_PARAMETER_RDG_BUFFER_UAV(RWBuffer, OutPrevOffsetBuffer)
 	END_SHADER_PARAMETER_STRUCT()
@@ -1395,8 +1367,8 @@ void AddHairStrandUpdatePositionOffsetPass(
 	FRDGImportedBuffer RootTrianglePrevPositionBuffer;
 	if (DeformedRootResources)
 	{
-		RootTriangleCurrPositionBuffer = Register(GraphBuilder, DeformedRootResources->LODs[LODIndex].GetDeformedUniqueTrianglePosition0Buffer(FHairStrandsDeformedRootResource::FLOD::Current), ERDGImportedBufferFlags::CreateSRV);
-		RootTrianglePrevPositionBuffer = Register(GraphBuilder, DeformedRootResources->LODs[LODIndex].GetDeformedUniqueTrianglePosition0Buffer(FHairStrandsDeformedRootResource::FLOD::Previous), ERDGImportedBufferFlags::CreateSRV);
+		RootTriangleCurrPositionBuffer = Register(GraphBuilder, DeformedRootResources->LODs[LODIndex].GetDeformedUniqueTrianglePositionBuffer(FHairStrandsDeformedRootResource::FLOD::Current), ERDGImportedBufferFlags::CreateSRV);
+		RootTrianglePrevPositionBuffer = Register(GraphBuilder, DeformedRootResources->LODs[LODIndex].GetDeformedUniqueTrianglePositionBuffer(FHairStrandsDeformedRootResource::FLOD::Previous), ERDGImportedBufferFlags::CreateSRV);
 	}
 
 	const bool bUseGPUOffset = DeformedRootResources != nullptr && DeformedRootResources->IsValid() && DeformedRootResources->IsInitialized() && GHairStrandsUseGPUPositionOffset > 0;
@@ -1406,8 +1378,8 @@ void AddHairStrandUpdatePositionOffsetPass(
 	FHairUpdatePositionOffsetCS::FParameters* PassParameters = GraphBuilder.AllocParameters<FHairUpdatePositionOffsetCS::FParameters>();
 	PassParameters->CPUCurrPositionOffset			= (FVector3f)DeformedResources->PositionOffset[CurrOffsetIndex];
 	PassParameters->CPUPrevPositionOffset			= (FVector3f)DeformedResources->PositionOffset[PrevOffsetIndex];
-	PassParameters->RootTriangleCurrPosition0Buffer	= RootTriangleCurrPositionBuffer.SRV;
-	PassParameters->RootTrianglePrevPosition0Buffer = RootTrianglePrevPositionBuffer.SRV;
+	PassParameters->RootTriangleCurrPositionBuffer	= RootTriangleCurrPositionBuffer.SRV;
+	PassParameters->RootTrianglePrevPositionBuffer	= RootTrianglePrevPositionBuffer.SRV;
 	PassParameters->OutCurrOffsetBuffer				= OutCurrPositionOffsetBuffer.UAV;
 	PassParameters->OutPrevOffsetBuffer				= OutPrevPositionOffsetBuffer.UAV;
 
diff --git a/Engine/Plugins/Runtime/HairStrands/Source/HairStrandsCore/Private/HairStrandsRendering.cpp b/Engine/Plugins/Runtime/HairStrands/Source/HairStrandsCore/Private/HairStrandsRendering.cpp
--- a/Engine/Plugins/Runtime/HairStrands/Source/HairStrandsCore/Private/HairStrandsRendering.cpp
+++ b/Engine/Plugins/Runtime/HairStrands/Source/HairStrandsCore/Private/HairStrandsRendering.cpp
@@ -294,14 +294,9 @@ class FDeformGuideCS : public FGlobalShader
 	BEGIN_SHADER_PARAMETER_STRUCT(FParameters, )
 		SHADER_PARAMETER(uint32, VertexCount)
 		SHADER_PARAMETER(FVector3f, SimRestOffset)
-
-		SHADER_PARAMETER_RDG_BUFFER_SRV(Buffer<float4>, SimRestPosition0Buffer)
-		SHADER_PARAMETER_RDG_BUFFER_SRV(Buffer<float4>, SimRestPosition1Buffer)
-		SHADER_PARAMETER_RDG_BUFFER_SRV(Buffer<float4>, SimRestPosition2Buffer)
-
-		SHADER_PARAMETER_RDG_BUFFER_SRV(Buffer<float4>, SimDeformedPosition0Buffer)
-		SHADER_PARAMETER_RDG_BUFFER_SRV(Buffer<float4>, SimDeformedPosition1Buffer)
-		SHADER_PARAMETER_RDG_BUFFER_SRV(Buffer<float4>, SimDeformedPosition2Buffer)
+		
+		SHADER_PARAMETER_RDG_BUFFER_SRV(Buffer<float4>, SimRestPositionBuffer)
+		SHADER_PARAMETER_RDG_BUFFER_SRV(Buffer<float4>, SimDeformedPositionBuffer)
 
 		SHADER_PARAMETER_RDG_BUFFER_SRV(Buffer<uint>, SimRootBarycentricBuffer)
 		SHADER_PARAMETER_RDG_BUFFER_SRV(Buffer<uint>, SimPointToCurveBuffer)
@@ -410,14 +405,8 @@ static void AddDeformSimHairStrandsPass(
 				{
 					InternalDeformationType = InternalDeformationType_Skinned;
 					Parameters->SimRootToUniqueTriangleIndexBuffer = RegisterAsSRV(GraphBuilder, RestLODDatas.RootToUniqueTriangleIndexBuffer);
-					Parameters->SimRestPosition0Buffer = RegisterAsSRV(GraphBuilder, RestLODDatas.RestUniqueTrianglePosition0Buffer);
-					Parameters->SimRestPosition1Buffer = RegisterAsSRV(GraphBuilder, RestLODDatas.RestUniqueTrianglePosition1Buffer);
-					Parameters->SimRestPosition2Buffer = RegisterAsSRV(GraphBuilder, RestLODDatas.RestUniqueTrianglePosition2Buffer);
-
-					Parameters->SimDeformedPosition0Buffer = RegisterAsSRV(GraphBuilder, DeformedLODDatas.GetDeformedUniqueTrianglePosition0Buffer(FHairStrandsDeformedRootResource::FLOD::Current));
-					Parameters->SimDeformedPosition1Buffer = RegisterAsSRV(GraphBuilder, DeformedLODDatas.GetDeformedUniqueTrianglePosition1Buffer(FHairStrandsDeformedRootResource::FLOD::Current));
-					Parameters->SimDeformedPosition2Buffer = RegisterAsSRV(GraphBuilder, DeformedLODDatas.GetDeformedUniqueTrianglePosition2Buffer(FHairStrandsDeformedRootResource::FLOD::Current));
-
+					Parameters->SimRestPositionBuffer = RegisterAsSRV(GraphBuilder, RestLODDatas.RestUniqueTrianglePositionBuffer);
+					Parameters->SimDeformedPositionBuffer = RegisterAsSRV(GraphBuilder, DeformedLODDatas.GetDeformedUniqueTrianglePositionBuffer(FHairStrandsDeformedRootResource::FLOD::Current));
 					Parameters->SimRootBarycentricBuffer = RegisterAsSRV(GraphBuilder, RestLODDatas.RootBarycentricBuffer);
 				}
 				else
@@ -509,25 +498,15 @@ class FHairInterpolationCS : public FGlobalShader
 
 		SHADER_PARAMETER_RDG_BUFFER_SRV(ByteAddressBuffer, InterpolationBuffer)
 
-		SHADER_PARAMETER_RDG_BUFFER_SRV(Buffer<float4>, RenRestPosition0Buffer)
-		SHADER_PARAMETER_RDG_BUFFER_SRV(Buffer<float4>, RenRestPosition1Buffer)
-		SHADER_PARAMETER_RDG_BUFFER_SRV(Buffer<float4>, RenRestPosition2Buffer)
-
-		SHADER_PARAMETER_RDG_BUFFER_SRV(Buffer<float4>, RenDeformedPosition0Buffer)
-		SHADER_PARAMETER_RDG_BUFFER_SRV(Buffer<float4>, RenDeformedPosition1Buffer)
-		SHADER_PARAMETER_RDG_BUFFER_SRV(Buffer<float4>, RenDeformedPosition2Buffer)
+		SHADER_PARAMETER_RDG_BUFFER_SRV(Buffer<float4>, RenRootRestPositionBuffer)
+		SHADER_PARAMETER_RDG_BUFFER_SRV(Buffer<float4>, RenRootDeformedPositionBuffer)
 
 		SHADER_PARAMETER_RDG_BUFFER_SRV(Buffer<uint>, RenRootBarycentricBuffer)
 		SHADER_PARAMETER_RDG_BUFFER_SRV(Buffer<uint>, RenPointToCurveBuffer)
 		SHADER_PARAMETER_RDG_BUFFER_SRV(Buffer<uint>, RenRootToUniqueTriangleIndexBuffer)
 
-		SHADER_PARAMETER_RDG_BUFFER_SRV(Buffer<float4>, SimRestPosition0Buffer)
-		SHADER_PARAMETER_RDG_BUFFER_SRV(Buffer<float4>, SimRestPosition1Buffer)
-		SHADER_PARAMETER_RDG_BUFFER_SRV(Buffer<float4>, SimRestPosition2Buffer)
-
-		SHADER_PARAMETER_RDG_BUFFER_SRV(Buffer<float4>, SimDeformedPosition0Buffer)
-		SHADER_PARAMETER_RDG_BUFFER_SRV(Buffer<float4>, SimDeformedPosition1Buffer)
-		SHADER_PARAMETER_RDG_BUFFER_SRV(Buffer<float4>, SimDeformedPosition2Buffer)
+		SHADER_PARAMETER_RDG_BUFFER_SRV(Buffer<float4>, SimRootRestPositionBuffer)
+		SHADER_PARAMETER_RDG_BUFFER_SRV(Buffer<float4>, SimRootDeformedPositionBuffer)
 
 		SHADER_PARAMETER_RDG_BUFFER_SRV(Buffer<uint>, SimRootBarycentricBuffer)
 		SHADER_PARAMETER_RDG_BUFFER_SRV(Buffer<uint>, SimPointToCurveBuffer)
@@ -643,17 +622,12 @@ static void AddHairStrandsInterpolationPass(
 		const FHairStrandsDeformedRootResource::FLOD& Sim_DeformedLODDatas = SimDeformedRootResources->LODs[MeshLODIndex];
 		bSupportGlobalInterpolation = Instance->Guides.bHasGlobalInterpolation && (Sim_RestLODDatas.SampleCount > 0);
 		{
-			Parameters->SimRestPosition0Buffer = RegisterAsSRV(GraphBuilder, Sim_RestLODDatas.RestUniqueTrianglePosition0Buffer);
-			Parameters->SimRestPosition1Buffer = RegisterAsSRV(GraphBuilder, Sim_RestLODDatas.RestUniqueTrianglePosition1Buffer);
-			Parameters->SimRestPosition2Buffer = RegisterAsSRV(GraphBuilder, Sim_RestLODDatas.RestUniqueTrianglePosition2Buffer);
-
+			Parameters->SimRootRestPositionBuffer = RegisterAsSRV(GraphBuilder, Sim_RestLODDatas.RestUniqueTrianglePositionBuffer);
 			Parameters->SimRootBarycentricBuffer = RegisterAsSRV(GraphBuilder, Sim_RestLODDatas.RootBarycentricBuffer);
 			Parameters->SimRootToUniqueTriangleIndexBuffer = RegisterAsSRV(GraphBuilder, Sim_RestLODDatas.RootToUniqueTriangleIndexBuffer);
 		}
 		{
-			Parameters->SimDeformedPosition0Buffer = RegisterAsSRV(GraphBuilder, Sim_DeformedLODDatas.GetDeformedUniqueTrianglePosition0Buffer(DeformedFrame));
-			Parameters->SimDeformedPosition1Buffer = RegisterAsSRV(GraphBuilder, Sim_DeformedLODDatas.GetDeformedUniqueTrianglePosition1Buffer(DeformedFrame));
-			Parameters->SimDeformedPosition2Buffer = RegisterAsSRV(GraphBuilder, Sim_DeformedLODDatas.GetDeformedUniqueTrianglePosition2Buffer(DeformedFrame));
+			Parameters->SimRootDeformedPositionBuffer = RegisterAsSRV(GraphBuilder, Sim_DeformedLODDatas.GetDeformedUniqueTrianglePositionBuffer(DeformedFrame));
 		}
 	}
 
@@ -663,16 +637,12 @@ static void AddHairStrandsInterpolationPass(
 		const FHairStrandsRestRootResource::FLOD& Ren_RestLODDatas = RenRestRootResources->LODs[MeshLODIndex];
 		const FHairStrandsDeformedRootResource::FLOD& Ren_DeformedLODDatas = RenDeformedRootResources->LODs[MeshLODIndex];
 		{
-			Parameters->RenRestPosition0Buffer = RegisterAsSRV(GraphBuilder, Ren_RestLODDatas.RestUniqueTrianglePosition0Buffer);
-			Parameters->RenRestPosition1Buffer = RegisterAsSRV(GraphBuilder, Ren_RestLODDatas.RestUniqueTrianglePosition1Buffer);
-			Parameters->RenRestPosition2Buffer = RegisterAsSRV(GraphBuilder, Ren_RestLODDatas.RestUniqueTrianglePosition2Buffer);
+			Parameters->RenRootRestPositionBuffer = RegisterAsSRV(GraphBuilder, Ren_RestLODDatas.RestUniqueTrianglePositionBuffer);
 			Parameters->RenRootBarycentricBuffer = RegisterAsSRV(GraphBuilder, Ren_RestLODDatas.RootBarycentricBuffer);
 			Parameters->RenRootToUniqueTriangleIndexBuffer = RegisterAsSRV(GraphBuilder, Ren_RestLODDatas.RootToUniqueTriangleIndexBuffer);
 		}
 		{
-			Parameters->RenDeformedPosition0Buffer = RegisterAsSRV(GraphBuilder, Ren_DeformedLODDatas.GetDeformedUniqueTrianglePosition0Buffer(DeformedFrame));
-			Parameters->RenDeformedPosition1Buffer = RegisterAsSRV(GraphBuilder, Ren_DeformedLODDatas.GetDeformedUniqueTrianglePosition1Buffer(DeformedFrame));
-			Parameters->RenDeformedPosition2Buffer = RegisterAsSRV(GraphBuilder, Ren_DeformedLODDatas.GetDeformedUniqueTrianglePosition2Buffer(DeformedFrame));
+			Parameters->RenRootDeformedPositionBuffer = RegisterAsSRV(GraphBuilder, Ren_DeformedLODDatas.GetDeformedUniqueTrianglePositionBuffer(DeformedFrame));
 		}
 	}
 
@@ -1003,13 +973,8 @@ class FHairCardsDeformationCS : public FGlobalShader
 		SHADER_PARAMETER_SRV(Buffer, CardsRestTangentBuffer)
 		SHADER_PARAMETER_RDG_BUFFER_SRV(ByteAddressBuffer, CardsInterpolationBuffer)
 
-		SHADER_PARAMETER_RDG_BUFFER_SRV(Buffer<float4>, TriangleRestPosition0Buffer)
-		SHADER_PARAMETER_RDG_BUFFER_SRV(Buffer<float4>, TriangleRestPosition1Buffer)
-		SHADER_PARAMETER_RDG_BUFFER_SRV(Buffer<float4>, TriangleRestPosition2Buffer)
-
-		SHADER_PARAMETER_RDG_BUFFER_SRV(Buffer<float4>, TriangleDeformedPosition0Buffer)
-		SHADER_PARAMETER_RDG_BUFFER_SRV(Buffer<float4>, TriangleDeformedPosition1Buffer)
-		SHADER_PARAMETER_RDG_BUFFER_SRV(Buffer<float4>, TriangleDeformedPosition2Buffer)
+		SHADER_PARAMETER_RDG_BUFFER_SRV(Buffer<float4>, TriangleRestPositionBuffer)
+		SHADER_PARAMETER_RDG_BUFFER_SRV(Buffer<float4>, TriangleDeformedPositionBuffer)
 
 		SHADER_PARAMETER_RDG_BUFFER_SRV(Buffer<uint>, GuideRootToUniqueTriangleIndexBuffer)
 		SHADER_PARAMETER_RDG_BUFFER_SRV(Buffer<uint>, GuideRootBarycentricBuffer)
@@ -1101,14 +1066,8 @@ static void AddHairCardsDeformationPass(
 
 		Parameters->GuideRootBarycentricBuffer = RegisterAsSRV(GraphBuilder, RestLODDatas.RootBarycentricBuffer);
 		Parameters->GuideRootToUniqueTriangleIndexBuffer = RegisterAsSRV(GraphBuilder, RestLODDatas.RootToUniqueTriangleIndexBuffer);
-
-		Parameters->TriangleRestPosition0Buffer = RegisterAsSRV(GraphBuilder, RestLODDatas.RestUniqueTrianglePosition0Buffer);
-		Parameters->TriangleRestPosition1Buffer = RegisterAsSRV(GraphBuilder, RestLODDatas.RestUniqueTrianglePosition1Buffer);
-		Parameters->TriangleRestPosition2Buffer = RegisterAsSRV(GraphBuilder, RestLODDatas.RestUniqueTrianglePosition2Buffer);
-
-		Parameters->TriangleDeformedPosition0Buffer = RegisterAsSRV(GraphBuilder, DeformedLODDatas.GetDeformedUniqueTrianglePosition0Buffer(FHairStrandsDeformedRootResource::FLOD::Current));
-		Parameters->TriangleDeformedPosition1Buffer = RegisterAsSRV(GraphBuilder, DeformedLODDatas.GetDeformedUniqueTrianglePosition1Buffer(FHairStrandsDeformedRootResource::FLOD::Current));
-		Parameters->TriangleDeformedPosition2Buffer = RegisterAsSRV(GraphBuilder, DeformedLODDatas.GetDeformedUniqueTrianglePosition2Buffer(FHairStrandsDeformedRootResource::FLOD::Current));
+		Parameters->TriangleRestPositionBuffer = RegisterAsSRV(GraphBuilder, RestLODDatas.RestUniqueTrianglePositionBuffer);
+		Parameters->TriangleDeformedPositionBuffer = RegisterAsSRV(GraphBuilder, DeformedLODDatas.GetDeformedUniqueTrianglePositionBuffer(FHairStrandsDeformedRootResource::FLOD::Current));
 	}
 
 	if (ShaderPrintData)
diff --git a/Engine/Plugins/Runtime/HairStrands/Source/HairStrandsCore/Private/Niagara/NiagaraDataInterfaceHairStrands.cpp b/Engine/Plugins/Runtime/HairStrands/Source/HairStrandsCore/Private/Niagara/NiagaraDataInterfaceHairStrands.cpp
--- a/Engine/Plugins/Runtime/HairStrands/Source/HairStrandsCore/Private/Niagara/NiagaraDataInterfaceHairStrands.cpp
+++ b/Engine/Plugins/Runtime/HairStrands/Source/HairStrandsCore/Private/Niagara/NiagaraDataInterfaceHairStrands.cpp
@@ -34,12 +34,8 @@ BEGIN_SHADER_PARAMETER_STRUCT(FShaderParameters,)
 	SHADER_PARAMETER_RDG_BUFFER_UAV(RWBuffer<uint4>,	DeformedPositionBuffer)
 	SHADER_PARAMETER_RDG_BUFFER_SRV(Buffer<uint4>,		RestPositionBuffer)
 	SHADER_PARAMETER_RDG_BUFFER_SRV(Buffer<uint>,		CurvesOffsetsBuffer)
-	SHADER_PARAMETER_RDG_BUFFER_SRV(Buffer<float4>,		RestTrianglePositionABuffer)
-	SHADER_PARAMETER_RDG_BUFFER_SRV(Buffer<float4>,		RestTrianglePositionBBuffer)
-	SHADER_PARAMETER_RDG_BUFFER_SRV(Buffer<float4>,		RestTrianglePositionCBuffer)
-	SHADER_PARAMETER_RDG_BUFFER_SRV(Buffer<float4>,		DeformedTrianglePositionABuffer)
-	SHADER_PARAMETER_RDG_BUFFER_SRV(Buffer<float4>,		DeformedTrianglePositionBBuffer)
-	SHADER_PARAMETER_RDG_BUFFER_SRV(Buffer<float4>,		DeformedTrianglePositionCBuffer)
+	SHADER_PARAMETER_RDG_BUFFER_SRV(Buffer<float4>,		RestTrianglePositionBuffer)
+	SHADER_PARAMETER_RDG_BUFFER_SRV(Buffer<float4>,		DeformedTrianglePositionBuffer)
 	SHADER_PARAMETER_RDG_BUFFER_SRV(Buffer<uint>,		RootBarycentricCoordinatesBuffer)
 	SHADER_PARAMETER_RDG_BUFFER_SRV(Buffer<uint>,		RootToUniqueTriangleIndexBuffer)
 	SHADER_PARAMETER(FMatrix44f,						WorldTransform)
@@ -823,12 +819,8 @@ class FInterpolateGroomGuidesCS : public FGlobalShader
 	using FPermutationDomain = TShaderPermutationDomain<FInterpolationType>;
 
 	BEGIN_SHADER_PARAMETER_STRUCT(FParameters, )
-		SHADER_PARAMETER_RDG_BUFFER_SRV(Buffer<float4>,	RestTrianglePositionABuffer)
-		SHADER_PARAMETER_RDG_BUFFER_SRV(Buffer<float4>,	RestTrianglePositionBBuffer)
-		SHADER_PARAMETER_RDG_BUFFER_SRV(Buffer<float4>,	RestTrianglePositionCBuffer)
-		SHADER_PARAMETER_RDG_BUFFER_SRV(Buffer<float4>,	DeformedTrianglePositionABuffer)
-		SHADER_PARAMETER_RDG_BUFFER_SRV(Buffer<float4>,	DeformedTrianglePositionBBuffer)
-		SHADER_PARAMETER_RDG_BUFFER_SRV(Buffer<float4>,	DeformedTrianglePositionCBuffer)
+		SHADER_PARAMETER_RDG_BUFFER_SRV(Buffer<float4>,	RestTrianglePositionBuffer)
+		SHADER_PARAMETER_RDG_BUFFER_SRV(Buffer<float4>,	DeformedTrianglePositionBuffer)
 		SHADER_PARAMETER_RDG_BUFFER_SRV(Buffer<uint>,	RootBarycentricCoordinatesBuffer)
 		SHADER_PARAMETER_RDG_BUFFER_SRV(Buffer<uint>,	RootToUniqueTriangleIndexBuffer)
 
@@ -940,17 +932,10 @@ static void InterpolateGroomGuides(FRDGBuilder& GraphBuilder, FNiagaraDataBuffer
 			const FHairStrandsRestRootResource::FLOD* RestMeshProjection = &(HairStrandsBuffer->SourceRestRootResources->LODs[MeshLODIndex]) ;
 			const FHairStrandsDeformedRootResource::FLOD* DeformedMeshProjection =  &(HairStrandsBuffer->SourceDeformedRootResources->LODs[MeshLODIndex]);
 
-			PassParameters->RestTrianglePositionABuffer = RegisterAsSRV(GraphBuilder, RestMeshProjection->RestUniqueTrianglePosition0Buffer);
-			PassParameters->RestTrianglePositionBBuffer = RegisterAsSRV(GraphBuilder, RestMeshProjection->RestUniqueTrianglePosition1Buffer);
-			PassParameters->RestTrianglePositionCBuffer = RegisterAsSRV(GraphBuilder, RestMeshProjection->RestUniqueTrianglePosition2Buffer);
-
-			PassParameters->DeformedTrianglePositionABuffer = RegisterAsSRV(GraphBuilder, DeformedMeshProjection->GetDeformedUniqueTrianglePosition0Buffer(FHairStrandsDeformedRootResource::FLOD::Current));
-			PassParameters->DeformedTrianglePositionBBuffer = RegisterAsSRV(GraphBuilder, DeformedMeshProjection->GetDeformedUniqueTrianglePosition1Buffer(FHairStrandsDeformedRootResource::FLOD::Current));
-			PassParameters->DeformedTrianglePositionCBuffer = RegisterAsSRV(GraphBuilder, DeformedMeshProjection->GetDeformedUniqueTrianglePosition2Buffer(FHairStrandsDeformedRootResource::FLOD::Current));
-			
+			PassParameters->RestTrianglePositionBuffer = RegisterAsSRV(GraphBuilder, RestMeshProjection->RestUniqueTrianglePositionBuffer);
+			PassParameters->DeformedTrianglePositionBuffer = RegisterAsSRV(GraphBuilder, DeformedMeshProjection->GetDeformedUniqueTrianglePositionBuffer(FHairStrandsDeformedRootResource::FLOD::Current));
 			PassParameters->RootBarycentricCoordinatesBuffer = RegisterAsSRV(GraphBuilder, RestMeshProjection->RootBarycentricBuffer);
 			PassParameters->RootToUniqueTriangleIndexBuffer = RegisterAsSRV(GraphBuilder, RestMeshProjection->RootToUniqueTriangleIndexBuffer);
-
 		}
 
 		const uint32 GroupSize = NIAGARA_HAIR_STRANDS_THREAD_COUNT_INTERPOLATE;
@@ -3416,13 +3401,8 @@ void UNiagaraDataInterfaceHairStrands::SetShaderParameters(const FNiagaraDataInt
 		const FHairStrandsRestRootResource::FLOD* RestMeshProjection = bHasSkinnedInterpolation ? &(HairStrandsBuffer->SourceRestRootResources->LODs[MeshLODIndex]) : nullptr;
 		const FHairStrandsDeformedRootResource::FLOD* DeformedMeshProjection = bHasSkinnedInterpolation ? &(HairStrandsBuffer->SourceDeformedRootResources->LODs[MeshLODIndex]) : nullptr;
 
-		FRDGBufferSRVRef RestTrianglePositionASRV = (bHasSkinnedInterpolation && RestMeshProjection) ? RegisterAsSRV(GraphBuilder, RestMeshProjection->RestUniqueTrianglePosition0Buffer) : Context.GetComputeDispatchInterface().GetEmptyBufferSRV(GraphBuilder, FHairStrandsMeshTrianglePositionFormat::Format);
-		FRDGBufferSRVRef RestTrianglePositionBSRV = (bHasSkinnedInterpolation && RestMeshProjection) ? RegisterAsSRV(GraphBuilder, RestMeshProjection->RestUniqueTrianglePosition1Buffer) : Context.GetComputeDispatchInterface().GetEmptyBufferSRV(GraphBuilder, FHairStrandsMeshTrianglePositionFormat::Format);
-		FRDGBufferSRVRef RestTrianglePositionCSRV = (bHasSkinnedInterpolation && RestMeshProjection) ? RegisterAsSRV(GraphBuilder, RestMeshProjection->RestUniqueTrianglePosition2Buffer) : Context.GetComputeDispatchInterface().GetEmptyBufferSRV(GraphBuilder, FHairStrandsMeshTrianglePositionFormat::Format);
-
-		FRDGBufferSRVRef DeformedTrianglePositionASRV = (bHasSkinnedInterpolation && DeformedMeshProjection) ?  RegisterAsSRV(GraphBuilder, DeformedMeshProjection->GetDeformedUniqueTrianglePosition0Buffer(FHairStrandsDeformedRootResource::FLOD::Current)) : Context.GetComputeDispatchInterface().GetEmptyBufferSRV(GraphBuilder, FHairStrandsMeshTrianglePositionFormat::Format);
-		FRDGBufferSRVRef DeformedTrianglePositionBSRV = (bHasSkinnedInterpolation && DeformedMeshProjection) ?  RegisterAsSRV(GraphBuilder, DeformedMeshProjection->GetDeformedUniqueTrianglePosition1Buffer(FHairStrandsDeformedRootResource::FLOD::Current)) : Context.GetComputeDispatchInterface().GetEmptyBufferSRV(GraphBuilder, FHairStrandsMeshTrianglePositionFormat::Format);
-		FRDGBufferSRVRef DeformedTrianglePositionCSRV = (bHasSkinnedInterpolation && DeformedMeshProjection) ?  RegisterAsSRV(GraphBuilder, DeformedMeshProjection->GetDeformedUniqueTrianglePosition2Buffer(FHairStrandsDeformedRootResource::FLOD::Current)) : Context.GetComputeDispatchInterface().GetEmptyBufferSRV(GraphBuilder, FHairStrandsMeshTrianglePositionFormat::Format);
+		FRDGBufferSRVRef RestTrianglePositionSRV = (bHasSkinnedInterpolation && RestMeshProjection) ? RegisterAsSRV(GraphBuilder, RestMeshProjection->RestUniqueTrianglePositionBuffer) : Context.GetComputeDispatchInterface().GetEmptyBufferSRV(GraphBuilder, FHairStrandsMeshTrianglePositionFormat::Format);
+		FRDGBufferSRVRef DeformedTrianglePositionSRV = (bHasSkinnedInterpolation && DeformedMeshProjection) ?  RegisterAsSRV(GraphBuilder, DeformedMeshProjection->GetDeformedUniqueTrianglePositionBuffer(FHairStrandsDeformedRootResource::FLOD::Current)) : Context.GetComputeDispatchInterface().GetEmptyBufferSRV(GraphBuilder, FHairStrandsMeshTrianglePositionFormat::Format);
 		
 		FRDGBufferSRVRef RootBarycentricCoordinatesSRV = (bHasSkinnedInterpolation && RestMeshProjection) ? RegisterAsSRV(GraphBuilder, RestMeshProjection->RootBarycentricBuffer) : Context.GetComputeDispatchInterface().GetEmptyBufferSRV(GraphBuilder, FHairStrandsRootBarycentricFormat::Format);
 		FRDGBufferSRVRef RootToUniqueTriangleIndexSRV = (bHasSkinnedInterpolation && RestMeshProjection) ? RegisterAsSRV(GraphBuilder, RestMeshProjection->RootToUniqueTriangleIndexBuffer) : Context.GetComputeDispatchInterface().GetEmptyBufferSRV(GraphBuilder, FHairStrandsRootToUniqueTriangleIndexFormat::Format);
@@ -3496,12 +3476,8 @@ void UNiagaraDataInterfaceHairStrands::SetShaderParameters(const FNiagaraDataInt
 		ShaderParameters->ParamsScaleBuffer = HairStrandsBuffer->ParamsScaleBuffer.GetOrCreateSRV(GraphBuilder);
 		ShaderParameters->RestPositionBuffer = RegisterAsSRV(GraphBuilder,HairStrandsBuffer->SourceRestResources->PositionBuffer); //-OPT: If constant across all dispatches can create in PreStage once
 		ShaderParameters->DeformedPositionOffset = DeformedPositionOffsetSRV;
-		ShaderParameters->RestTrianglePositionABuffer = RestTrianglePositionASRV;
-		ShaderParameters->RestTrianglePositionBBuffer = RestTrianglePositionBSRV;
-		ShaderParameters->RestTrianglePositionCBuffer = RestTrianglePositionCSRV;
-		ShaderParameters->DeformedTrianglePositionABuffer = DeformedTrianglePositionASRV;
-		ShaderParameters->DeformedTrianglePositionBBuffer = DeformedTrianglePositionBSRV;
-		ShaderParameters->DeformedTrianglePositionCBuffer = DeformedTrianglePositionCSRV;
+		ShaderParameters->RestTrianglePositionBuffer = RestTrianglePositionSRV;
+		ShaderParameters->DeformedTrianglePositionBuffer = DeformedTrianglePositionSRV;
 		ShaderParameters->RestSamplePositionsBuffer = RestSamplePositionsBufferSRV;
 		ShaderParameters->MeshSampleWeightsBuffer = MeshSampleWeightsBufferSRV;
 		ShaderParameters->RootBarycentricCoordinatesBuffer = RootBarycentricCoordinatesSRV;
@@ -3544,12 +3520,8 @@ void UNiagaraDataInterfaceHairStrands::SetShaderParameters(const FNiagaraDataInt
 		ShaderParameters->CurvesOffsetsBuffer = Context.GetComputeDispatchInterface().GetEmptyBufferSRV(GraphBuilder, FHairStrandsCurveFormat::Format);
 		ShaderParameters->RestPositionBuffer = Context.GetComputeDispatchInterface().GetEmptyBufferSRV(GraphBuilder, FHairStrandsPositionFormat::Format);
 		ShaderParameters->DeformedPositionOffset = Context.GetComputeDispatchInterface().GetEmptyBufferSRV(GraphBuilder, FHairStrandsPositionOffsetFormat::Format);
-		ShaderParameters->RestTrianglePositionABuffer = Context.GetComputeDispatchInterface().GetEmptyBufferSRV(GraphBuilder, FHairStrandsMeshTrianglePositionFormat::Format);
-		ShaderParameters->RestTrianglePositionBBuffer = Context.GetComputeDispatchInterface().GetEmptyBufferSRV(GraphBuilder, FHairStrandsMeshTrianglePositionFormat::Format);
-		ShaderParameters->RestTrianglePositionCBuffer = Context.GetComputeDispatchInterface().GetEmptyBufferSRV(GraphBuilder, FHairStrandsMeshTrianglePositionFormat::Format);
-		ShaderParameters->DeformedTrianglePositionABuffer = Context.GetComputeDispatchInterface().GetEmptyBufferSRV(GraphBuilder, FHairStrandsMeshTrianglePositionFormat::Format);
-		ShaderParameters->DeformedTrianglePositionBBuffer = Context.GetComputeDispatchInterface().GetEmptyBufferSRV(GraphBuilder, FHairStrandsMeshTrianglePositionFormat::Format);
-		ShaderParameters->DeformedTrianglePositionCBuffer = Context.GetComputeDispatchInterface().GetEmptyBufferSRV(GraphBuilder, FHairStrandsMeshTrianglePositionFormat::Format);
+		ShaderParameters->RestTrianglePositionBuffer = Context.GetComputeDispatchInterface().GetEmptyBufferSRV(GraphBuilder, FHairStrandsMeshTrianglePositionFormat::Format);
+		ShaderParameters->DeformedTrianglePositionBuffer = Context.GetComputeDispatchInterface().GetEmptyBufferSRV(GraphBuilder, FHairStrandsMeshTrianglePositionFormat::Format);
 		ShaderParameters->RestSamplePositionsBuffer = Context.GetComputeDispatchInterface().GetEmptyBufferSRV(GraphBuilder, FHairStrandsMeshTrianglePositionFormat::Format);
 		ShaderParameters->MeshSampleWeightsBuffer = Context.GetComputeDispatchInterface().GetEmptyBufferSRV(GraphBuilder, FHairStrandsMeshTrianglePositionFormat::Format);
 		ShaderParameters->RootBarycentricCoordinatesBuffer = Context.GetComputeDispatchInterface().GetEmptyBufferSRV(GraphBuilder, FHairStrandsRootBarycentricFormat::Format);
diff --git a/Engine/Plugins/Runtime/HairStrands/Source/HairStrandsCore/Public/GroomResources.h b/Engine/Plugins/Runtime/HairStrands/Source/HairStrandsCore/Public/GroomResources.h
--- a/Engine/Plugins/Runtime/HairStrands/Source/HairStrandsCore/Public/GroomResources.h
+++ b/Engine/Plugins/Runtime/HairStrands/Source/HairStrandsCore/Public/GroomResources.h
@@ -167,9 +167,7 @@ struct FHairStrandsRestRootResource : public FHairCommonResource
 			Total += GetBufferTotalNumBytes(LOD.UniqueTriangleIndexBuffer);
 			Total += GetBufferTotalNumBytes(LOD.RootToUniqueTriangleIndexBuffer);
 			Total += GetBufferTotalNumBytes(LOD.RootBarycentricBuffer);
-			Total += GetBufferTotalNumBytes(LOD.RestUniqueTrianglePosition0Buffer);
-			Total += GetBufferTotalNumBytes(LOD.RestUniqueTrianglePosition1Buffer);
-			Total += GetBufferTotalNumBytes(LOD.RestUniqueTrianglePosition2Buffer);
+			Total += GetBufferTotalNumBytes(LOD.RestUniqueTrianglePositionBuffer);
 			Total += GetBufferTotalNumBytes(LOD.MeshInterpolationWeightsBuffer);
 			Total += GetBufferTotalNumBytes(LOD.MeshSampleIndicesBuffer);
 			Total += GetBufferTotalNumBytes(LOD.RestSamplePositionsBuffer);
@@ -194,9 +192,7 @@ struct FHairStrandsRestRootResource : public FHairCommonResource
 		FRDGExternalBuffer RootBarycentricBuffer;
 
 		/* Strand hair roots translation and rotation in rest position relative to the bound triangle. Positions are relative to the rest root center */
-		FRDGExternalBuffer RestUniqueTrianglePosition0Buffer;
-		FRDGExternalBuffer RestUniqueTrianglePosition1Buffer;
-		FRDGExternalBuffer RestUniqueTrianglePosition2Buffer;
+		FRDGExternalBuffer RestUniqueTrianglePositionBuffer;
 
 		/* Strand hair mesh interpolation matrix and sample indices */
 		uint32 SampleCount = 0;
@@ -242,18 +238,14 @@ struct FHairStrandsDeformedRootResource : public FHairCommonResource
 		uint32 Total = 0;
 		for (const FLOD& LOD : LODs)
 		{
-			Total += GetBufferTotalNumBytes(LOD.DeformedUniqueTrianglePosition0Buffer[0]);
-			Total += GetBufferTotalNumBytes(LOD.DeformedUniqueTrianglePosition1Buffer[0]);
-			Total += GetBufferTotalNumBytes(LOD.DeformedUniqueTrianglePosition2Buffer[0]);
+			Total += GetBufferTotalNumBytes(LOD.DeformedUniqueTrianglePositionBuffer[0]);
 			Total += GetBufferTotalNumBytes(LOD.DeformedSamplePositionsBuffer[0]);
 			Total += GetBufferTotalNumBytes(LOD.MeshSampleWeightsBuffer[0]);
 
 			// Double buffering is disabled by default unless the read-only cvar r.HairStrands.ContinuousDecimationReordering is set
 			if (IsHairStrandContinuousDecimationReorderingEnabled())
 			{
-				Total += GetBufferTotalNumBytes(LOD.DeformedUniqueTrianglePosition0Buffer[1]);
-				Total += GetBufferTotalNumBytes(LOD.DeformedUniqueTrianglePosition1Buffer[1]);
-				Total += GetBufferTotalNumBytes(LOD.DeformedUniqueTrianglePosition2Buffer[1]);
+				Total += GetBufferTotalNumBytes(LOD.DeformedUniqueTrianglePositionBuffer[1]);
 				Total += GetBufferTotalNumBytes(LOD.DeformedSamplePositionsBuffer[1]);
 				Total += GetBufferTotalNumBytes(LOD.MeshSampleWeightsBuffer[1]);
 			}
@@ -281,9 +273,7 @@ struct FHairStrandsDeformedRootResource : public FHairCommonResource
 		int32 LODIndex = -1;
 
 		/* Strand hair roots translation and rotation in triangle-deformed position relative to the bound triangle. Positions are relative the deformed root center*/
-		FRDGExternalBuffer DeformedUniqueTrianglePosition0Buffer[2];
-		FRDGExternalBuffer DeformedUniqueTrianglePosition1Buffer[2];
-		FRDGExternalBuffer DeformedUniqueTrianglePosition2Buffer[2];
+		FRDGExternalBuffer DeformedUniqueTrianglePositionBuffer[2];
 
 		/* Strand hair mesh interpolation matrix and sample indices */
 		uint32 SampleCount = 0;
@@ -301,9 +291,7 @@ struct FHairStrandsDeformedRootResource : public FHairCommonResource
 
 		// Double buffering is disabled by default unless the read-only cvar r.HairStrands.ContinuousDecimationReordering is set
 		inline uint32 GetIndex(EFrameType T) const { return T == EFrameType::Current ? CurrentIndex : 1u - CurrentIndex; }
-		inline const FRDGExternalBuffer& GetDeformedUniqueTrianglePosition0Buffer(EFrameType T) const { return IsHairStrandContinuousDecimationReorderingEnabled() ? DeformedUniqueTrianglePosition0Buffer[GetIndex(T)] : DeformedUniqueTrianglePosition0Buffer[0]; }
-		inline const FRDGExternalBuffer& GetDeformedUniqueTrianglePosition1Buffer(EFrameType T) const { return IsHairStrandContinuousDecimationReorderingEnabled() ? DeformedUniqueTrianglePosition1Buffer[GetIndex(T)] : DeformedUniqueTrianglePosition1Buffer[0]; }
-		inline const FRDGExternalBuffer& GetDeformedUniqueTrianglePosition2Buffer(EFrameType T) const { return IsHairStrandContinuousDecimationReorderingEnabled() ? DeformedUniqueTrianglePosition2Buffer[GetIndex(T)] : DeformedUniqueTrianglePosition2Buffer[0]; }
+		inline const FRDGExternalBuffer& GetDeformedUniqueTrianglePositionBuffer(EFrameType T) const { return IsHairStrandContinuousDecimationReorderingEnabled() ? DeformedUniqueTrianglePositionBuffer[GetIndex(T)] : DeformedUniqueTrianglePositionBuffer[0]; }
 		inline const FRDGExternalBuffer& GetDeformedSamplePositionsBuffer(EFrameType T) const { return IsHairStrandContinuousDecimationReorderingEnabled() ? DeformedSamplePositionsBuffer[GetIndex(T)] : DeformedSamplePositionsBuffer[0]; }
 		inline const FRDGExternalBuffer& GetMeshSampleWeightsBuffer(EFrameType T) const { return IsHairStrandContinuousDecimationReorderingEnabled() ? MeshSampleWeightsBuffer[GetIndex(T)] : MeshSampleWeightsBuffer[0]; }
 		inline void SwapBuffer() { CurrentIndex = 1u - CurrentIndex; }
diff --git a/Engine/Plugins/Runtime/HairStrands/Source/HairStrandsCore/Public/HairStrandsDatas.h b/Engine/Plugins/Runtime/HairStrands/Source/HairStrandsCore/Public/HairStrandsDatas.h
--- a/Engine/Plugins/Runtime/HairStrands/Source/HairStrandsCore/Public/HairStrandsDatas.h
+++ b/Engine/Plugins/Runtime/HairStrands/Source/HairStrandsCore/Public/HairStrandsDatas.h
@@ -569,9 +569,7 @@ struct FHairStrandsRootBulkData
 			Total += LOD.UniqueTriangleIndexBuffer.IsBulkDataLoaded() ?			LOD.UniqueTriangleIndexBuffer.GetBulkDataSize() : 0u;
 			Total += LOD.RootToUniqueTriangleIndexBuffer.IsBulkDataLoaded() ?	LOD.RootToUniqueTriangleIndexBuffer.GetBulkDataSize() : 0u;
 			Total += LOD.RootBarycentricBuffer.IsBulkDataLoaded() ?				LOD.RootBarycentricBuffer.GetBulkDataSize() : 0u;
-			Total += LOD.RestUniqueTrianglePosition0Buffer.IsBulkDataLoaded() ?	LOD.RestUniqueTrianglePosition0Buffer.GetBulkDataSize() : 0u;
-			Total += LOD.RestUniqueTrianglePosition1Buffer.IsBulkDataLoaded() ?	LOD.RestUniqueTrianglePosition1Buffer.GetBulkDataSize() : 0u;
-			Total += LOD.RestUniqueTrianglePosition2Buffer.IsBulkDataLoaded() ?	LOD.RestUniqueTrianglePosition2Buffer.GetBulkDataSize() : 0u;
+			Total += LOD.RestUniqueTrianglePositionBuffer.IsBulkDataLoaded() ?	LOD.RestUniqueTrianglePositionBuffer.GetBulkDataSize() : 0u;
 			Total += LOD.MeshInterpolationWeightsBuffer.IsBulkDataLoaded() ?	LOD.MeshInterpolationWeightsBuffer.GetBulkDataSize() : 0u;
 			Total += LOD.MeshSampleIndicesBuffer.IsBulkDataLoaded() ?			LOD.MeshSampleIndicesBuffer.GetBulkDataSize() : 0u;
 			Total += LOD.RestSamplePositionsBuffer.IsBulkDataLoaded() ?			LOD.RestSamplePositionsBuffer.GetBulkDataSize() : 0u;
@@ -594,9 +592,7 @@ struct FHairStrandsRootBulkData
 		FByteBulkData UniqueTriangleIndexBuffer;
 
 		/* Rest triangle positions (per-unique-triangle) */
-		FByteBulkData RestUniqueTrianglePosition0Buffer;
-		FByteBulkData RestUniqueTrianglePosition1Buffer;
-		FByteBulkData RestUniqueTrianglePosition2Buffer;
+		FByteBulkData RestUniqueTrianglePositionBuffer;
 
 		/* Number of samples used for the mesh interpolation */
 		uint32 SampleCount = 0;
@@ -644,9 +640,7 @@ struct FHairStrandsRootData
 
 		/* Strand hair roots translation and rotation in rest position relative to the bound triangle. Positions are relative to the rest root center */
 		TArray<FHairStrandsUniqueTriangleIndexFormat::Type> UniqueTriangleIndexBuffer;
-		TArray<FHairStrandsMeshTrianglePositionFormat::Type> RestUniqueTrianglePosition0Buffer;
-		TArray<FHairStrandsMeshTrianglePositionFormat::Type> RestUniqueTrianglePosition1Buffer;
-		TArray<FHairStrandsMeshTrianglePositionFormat::Type> RestUniqueTrianglePosition2Buffer;
+		TArray<FHairStrandsMeshTrianglePositionFormat::Type> RestUniqueTrianglePositionBuffer;
 
 		/* Number of samples used for the mesh interpolation */
 		uint32 SampleCount = 0;
diff --git a/Engine/Shaders/Private/HairStrands/HairCardsDeformation.usf b/Engine/Shaders/Private/HairStrands/HairCardsDeformation.usf
--- a/Engine/Shaders/Private/HairStrands/HairCardsDeformation.usf
+++ b/Engine/Shaders/Private/HairStrands/HairCardsDeformation.usf
@@ -29,13 +29,9 @@ Buffer<float4> CardsRestTangentBuffer;
 RWBuffer<float4> CardsDeformedTangentBuffer;
 
 // Guides == Strands going through the middle of the card geometry
-Buffer<float4>  TriangleRestPosition0Buffer;
-Buffer<float4>  TriangleRestPosition1Buffer;
-Buffer<float4>  TriangleRestPosition2Buffer;
+Buffer<float4>  TriangleRestPositionBuffer;
 
-Buffer<float4>  TriangleDeformedPosition0Buffer;
-Buffer<float4>  TriangleDeformedPosition1Buffer;
-Buffer<float4>  TriangleDeformedPosition2Buffer;
+Buffer<float4>  TriangleDeformedPositionBuffer;
 
 Buffer<uint>	GuideRootBarycentricBuffer;
 Buffer<uint>    GuideVertexToRootIndexBuffer;
@@ -149,8 +145,8 @@ void MainCS(uint DispatchThreadId : SV_DispatchThreadID)
 
 	const float3 RootBarycentric = DecodeBarycentrics(GuideRootBarycentricBuffer[RootIndex]);
 
-	const FHairMeshTriangle RestTriangle = GetTriangleTransformation(TriangleIndex, TriangleRestPosition0Buffer, TriangleRestPosition1Buffer, TriangleRestPosition2Buffer);
-	const FHairMeshTriangle DeformedTriangle = GetTriangleTransformation(TriangleIndex, TriangleDeformedPosition0Buffer, TriangleDeformedPosition1Buffer, TriangleDeformedPosition2Buffer);
+	const FHairMeshTriangle RestTriangle = GetTriangleTransformation(TriangleIndex, TriangleRestPositionBuffer);
+	const FHairMeshTriangle DeformedTriangle = GetTriangleTransformation(TriangleIndex, TriangleDeformedPositionBuffer);
 
 	const float3 Offset0 = ComputeDynamicGeometryOffset(GuideIndex0, 1, RestTriangle, DeformedTriangle, RootBarycentric, GuideDeformedPositionOffset);
 	const float3 Offset1 = ComputeDynamicGeometryOffset(GuideIndex1, 1, RestTriangle, DeformedTriangle, RootBarycentric, GuideDeformedPositionOffset);
diff --git a/Engine/Shaders/Private/HairStrands/HairStrandsDebug.usf b/Engine/Shaders/Private/HairStrands/HairStrandsDebug.usf
--- a/Engine/Shaders/Private/HairStrands/HairStrandsDebug.usf
+++ b/Engine/Shaders/Private/HairStrands/HairStrandsDebug.usf
@@ -314,13 +314,8 @@ Buffer<uint>   RootBarycentricBuffer;
 
 Buffer<uint>  RootToUniqueTriangleIndexBuffer;
 
-Buffer<float4>  RestPosition0Buffer;
-Buffer<float4>  RestPosition1Buffer;
-Buffer<float4>  RestPosition2Buffer;
-
-Buffer<float4>  DeformedPosition0Buffer;
-Buffer<float4>  DeformedPosition1Buffer;
-Buffer<float4>  DeformedPosition2Buffer;
+Buffer<float4>  RestPositionBuffer;
+Buffer<float4>  DeformedPositionBuffer;
 
 Buffer<float4> RestSamplePositionsBuffer;
 Buffer<float4> DeformedSamplePositionsBuffer;
@@ -335,8 +330,8 @@ void MainCS(uint3 DispatchThreadId : SV_DispatchThreadID)
 	if (RootIndex >= MaxRootCount) return;
 	{
 		const uint TriangleIndex = RootToUniqueTriangleIndexBuffer[RootIndex];
-		const FHairMeshTriangleWithUVs RestTriangle		= GetTriangleWithUVsTransformation(TriangleIndex, RestPosition0Buffer, RestPosition1Buffer, RestPosition2Buffer);
-		const FHairMeshTriangleWithUVs DeformedTriangle = GetTriangleWithUVsTransformation(TriangleIndex, DeformedPosition0Buffer, DeformedPosition1Buffer, DeformedPosition2Buffer);
+		const FHairMeshTriangleWithUVs RestTriangle		= GetTriangleWithUVsTransformation(TriangleIndex, RestPositionBuffer);
+		const FHairMeshTriangleWithUVs DeformedTriangle = GetTriangleWithUVsTransformation(TriangleIndex, DeformedPositionBuffer);
 
 		FHairMeshTriangleWithUVs Tri;
 		if (DeformedFrameEnable > 0)
@@ -376,8 +371,8 @@ void MainCS(uint3 DispatchThreadId : SV_DispatchThreadID)
 	const uint TriangleIndex = DispatchThreadId.x;
 	if (TriangleIndex < MaxRootCount)
 	{
-		const FHairMeshTriangleWithUVs RestTriangle		= GetTriangleWithUVsTransformation(TriangleIndex, RestPosition0Buffer, RestPosition1Buffer, RestPosition2Buffer);
-		const FHairMeshTriangleWithUVs DeformedTriangle = GetTriangleWithUVsTransformation(TriangleIndex, DeformedPosition0Buffer, DeformedPosition1Buffer, DeformedPosition2Buffer);
+		const FHairMeshTriangleWithUVs RestTriangle		= GetTriangleWithUVsTransformation(TriangleIndex, RestPositionBuffer);
+		const FHairMeshTriangleWithUVs DeformedTriangle = GetTriangleWithUVsTransformation(TriangleIndex, DeformedPositionBuffer);
 	
 		FHairMeshTriangleWithUVs Tri;
 		if (DeformedFrameEnable > 0)
diff --git a/Engine/Shaders/Private/HairStrands/HairStrandsFollicleMask.usf b/Engine/Shaders/Private/HairStrands/HairStrandsFollicleMask.usf
--- a/Engine/Shaders/Private/HairStrands/HairStrandsFollicleMask.usf
+++ b/Engine/Shaders/Private/HairStrands/HairStrandsFollicleMask.usf
@@ -15,9 +15,7 @@ uint KernelSizeInPixels;
 #define UV_ROOT 1
 
 #if PERMUTATION_UV_TYPE == UV_MESH
-Buffer<float4> UniqueTrianglePosition0Buffer;
-Buffer<float4> UniqueTrianglePosition1Buffer;
-Buffer<float4> UniqueTrianglePosition2Buffer;
+Buffer<float4> UniqueTrianglePositionBuffer;
 Buffer<uint>   RootToUniqueTriangleIndexBuffer;
 Buffer<uint>   RootBarycentricBuffer; 
 #endif 
@@ -50,7 +48,7 @@ void MainVS(
 	}
 
 	if (UniqueTriangleIndex > MaxUniqueTriangleIndex) return;
-	const FHairMeshTriangleWithUVs Tri = GetTriangleWithUVsTransformation(UniqueTriangleIndex, UniqueTrianglePosition0Buffer, UniqueTrianglePosition1Buffer, UniqueTrianglePosition2Buffer);
+	const FHairMeshTriangleWithUVs Tri = GetTriangleWithUVsTransformation(UniqueTriangleIndex, UniqueTrianglePositionBuffer);
 
 	const uint EncodedBarycentric = RootBarycentricBuffer[RootIndex];
 	const float3 Barycentric = DecodeBarycentrics(EncodedBarycentric);
diff --git a/Engine/Shaders/Private/HairStrands/HairStrandsGuideDeform.usf b/Engine/Shaders/Private/HairStrands/HairStrandsGuideDeform.usf
--- a/Engine/Shaders/Private/HairStrands/HairStrandsGuideDeform.usf
+++ b/Engine/Shaders/Private/HairStrands/HairStrandsGuideDeform.usf
@@ -22,13 +22,8 @@ RWBuffer<uint4> OutSimDeformedPositionBuffer;
 
 #if PERMUTATION_DEFORMATION == DEFORMATION_OFFSET_GUIDE_DYNAMIC
 
-Buffer<float4>  SimRestPosition0Buffer;
-Buffer<float4>  SimRestPosition1Buffer;
-Buffer<float4>  SimRestPosition2Buffer;
-
-Buffer<float4>  SimDeformedPosition0Buffer;
-Buffer<float4>  SimDeformedPosition1Buffer;
-Buffer<float4>  SimDeformedPosition2Buffer;
+Buffer<float4>  SimRestPositionBuffer;
+Buffer<float4>  SimDeformedPositionBuffer;
 
 Buffer<uint>	SimRootBarycentricBuffer;
 Buffer<uint>    SimPointToCurveBuffer;
@@ -81,8 +76,8 @@ float3 DisplacePosition(uint VertexIndex, float3 Pos)
 	const float3 RootBarycentric = DecodeBarycentrics(SimRootBarycentricBuffer[RootIndex]);
 
 	const uint TriangleIndex = SimRootToUniqueTriangleIndexBuffer[RootIndex];
-	const FHairMeshTriangle RestTriangle = GetTriangleTransformation(TriangleIndex, SimRestPosition0Buffer, SimRestPosition1Buffer, SimRestPosition2Buffer);
-	const FHairMeshTriangle DeformedTriangle = GetTriangleTransformation(TriangleIndex, SimDeformedPosition0Buffer, SimDeformedPosition1Buffer, SimDeformedPosition2Buffer);
+	const FHairMeshTriangle RestTriangle = GetTriangleTransformation(TriangleIndex, SimRestPositionBuffer);
+	const FHairMeshTriangle DeformedTriangle = GetTriangleTransformation(TriangleIndex, SimDeformedPositionBuffer);
 
 	return DeformGuideByTriangle(VertexIndex, RestTriangle, DeformedTriangle, RootBarycentric) - OutSimHairPositionOffset;
 #endif
diff --git a/Engine/Shaders/Private/HairStrands/HairStrandsInterpolation.usf b/Engine/Shaders/Private/HairStrands/HairStrandsInterpolation.usf
--- a/Engine/Shaders/Private/HairStrands/HairStrandsInterpolation.usf
+++ b/Engine/Shaders/Private/HairStrands/HairStrandsInterpolation.usf
@@ -107,25 +107,13 @@ RWBuffer<uint4> OutRenDeformedPositionBuffer;
 #if PERMUTATION_DYNAMIC_GEOMETRY >= 1 
 
 // Compact all these buffers into 2 buffers: translation + quaternion
-Buffer<float4>  RenRestPosition0Buffer;
-Buffer<float4>  RenRestPosition1Buffer;
-Buffer<float4>  RenRestPosition2Buffer;
-
-Buffer<float4>  RenDeformedPosition0Buffer;
-Buffer<float4>  RenDeformedPosition1Buffer;
-Buffer<float4>  RenDeformedPosition2Buffer;
-
+Buffer<float4>  RenRootRestPositionBuffer;
+Buffer<float4>  RenRootDeformedPositionBuffer;
 Buffer<uint>	RenRootBarycentricBuffer;
 Buffer<uint>	RenRootToUniqueTriangleIndexBuffer;
 
-Buffer<float4>  SimRestPosition0Buffer;
-Buffer<float4>  SimRestPosition1Buffer;
-Buffer<float4>  SimRestPosition2Buffer;
-
-Buffer<float4>  SimDeformedPosition0Buffer;
-Buffer<float4>  SimDeformedPosition1Buffer;
-Buffer<float4>  SimDeformedPosition2Buffer;
-
+Buffer<float4>  SimRootRestPositionBuffer;
+Buffer<float4>  SimRootDeformedPositionBuffer;
 Buffer<uint>	SimRootBarycentricBuffer;
 Buffer<uint>	SimRootToUniqueTriangleIndexBuffer;
 #endif
@@ -237,16 +225,16 @@ FRootTriangleData GetRootTriangleData(uint VertexIndex, bool bSim)
 		Out.RootIndex		 = SimPointToCurveBuffer[VertexIndex];
 		Out.TriangleIndex	 = SimRootToUniqueTriangleIndexBuffer[Out.RootIndex];
 		Out.RootBarycentric  = DecodeBarycentrics(SimRootBarycentricBuffer[Out.RootIndex]);
-		Out.RestTriangle	 = GetTriangleTransformation(Out.TriangleIndex, SimRestPosition0Buffer, SimRestPosition1Buffer, SimRestPosition2Buffer);
-		Out.DeformedTriangle = GetTriangleTransformation(Out.TriangleIndex, SimDeformedPosition0Buffer, SimDeformedPosition1Buffer, SimDeformedPosition2Buffer);
+		Out.RestTriangle	 = GetTriangleTransformation(Out.TriangleIndex, SimRootRestPositionBuffer);
+		Out.DeformedTriangle = GetTriangleTransformation(Out.TriangleIndex, SimRootDeformedPositionBuffer);
 	}
 	else
 	{
 		Out.RootIndex		 = RenPointToCurveBuffer[VertexIndex];
 		Out.TriangleIndex	 = RenRootToUniqueTriangleIndexBuffer[Out.RootIndex];
 		Out.RootBarycentric  = DecodeBarycentrics(RenRootBarycentricBuffer[Out.RootIndex]);
-		Out.RestTriangle	 = GetTriangleTransformation(Out.TriangleIndex, RenRestPosition0Buffer, RenRestPosition1Buffer, RenRestPosition2Buffer);
-		Out.DeformedTriangle = GetTriangleTransformation(Out.TriangleIndex, RenDeformedPosition0Buffer, RenDeformedPosition1Buffer, RenDeformedPosition2Buffer);
+		Out.RestTriangle	 = GetTriangleTransformation(Out.TriangleIndex, RenRootRestPositionBuffer);
+		Out.DeformedTriangle = GetTriangleTransformation(Out.TriangleIndex, RenRootDeformedPositionBuffer);
 	}
 	#endif
 	return Out;
@@ -437,7 +425,7 @@ void MainCS(uint2 DispatchThreadId : SV_DispatchThreadID)
 #if PERMUTATION_CULLING
 	if (HairStrandsVFTODO_bCullingEnable)
 	{
-		const uint FetchIndex = DispatchThreadId;
+		const uint FetchIndex = DispatchThreadId.x;
 		const uint VertexCountAfterCulling = HairStrandsVFTODO_CullingIndirectBuffer[3];
 		if (FetchIndex >= VertexCountAfterCulling)
 		{
diff --git a/Engine/Shaders/Private/HairStrands/HairStrandsMesh.usf b/Engine/Shaders/Private/HairStrands/HairStrandsMesh.usf
--- a/Engine/Shaders/Private/HairStrands/HairStrandsMesh.usf
+++ b/Engine/Shaders/Private/HairStrands/HairStrandsMesh.usf
@@ -321,14 +321,10 @@ Buffer<float2> MeshUVsBuffer8;
 uint MaxUniqueTriangleCount;
 Buffer<uint> UniqueTriangleIndices;
 
-RWBuffer<float4> OutUniqueTriangleCurrPosition0;
-RWBuffer<float4> OutUniqueTriangleCurrPosition1;
-RWBuffer<float4> OutUniqueTriangleCurrPosition2;
+RWBuffer<float4> OutUniqueTriangleCurrPosition;
 
 #if PERMUTATION_PREVIOUS
-RWBuffer<float4> OutUniqueTrianglePrevPosition0;
-RWBuffer<float4> OutUniqueTrianglePrevPosition1;
-RWBuffer<float4> OutUniqueTrianglePrevPosition2;
+RWBuffer<float4> OutUniqueTrianglePrevPosition;
 #endif
 
 #if PERMUTATION_WITHUV == 1
@@ -403,33 +399,33 @@ void MainCS(uint3 DispatchThreadId : SV_DispatchThreadID)
 #if PERMUTATION_WITHUV == 1
 	{
 		const FHairMeshTriangleWithUVs Tri = GetTriangleData(TriangleIndex, SectionSlotIndex, false);
-		OutUniqueTriangleCurrPosition0[UniqueTriangleIndex] = float4(Tri.P0, PackUVsToFloat(Tri.UV0));
-		OutUniqueTriangleCurrPosition1[UniqueTriangleIndex] = float4(Tri.P1, PackUVsToFloat(Tri.UV1));
-		OutUniqueTriangleCurrPosition2[UniqueTriangleIndex] = float4(Tri.P2, PackUVsToFloat(Tri.UV2));
+		OutUniqueTriangleCurrPosition[UniqueTriangleIndex * 3 + 0] = float4(Tri.P0, PackUVsToFloat(Tri.UV0));
+		OutUniqueTriangleCurrPosition[UniqueTriangleIndex * 3 + 1] = float4(Tri.P1, PackUVsToFloat(Tri.UV1));
+		OutUniqueTriangleCurrPosition[UniqueTriangleIndex * 3 + 2] = float4(Tri.P2, PackUVsToFloat(Tri.UV2));
 	}
 
 	#if PERMUTATION_PREVIOUS
 	{
 		const FHairMeshTriangleWithUVs Tri = GetTriangleData(TriangleIndex, SectionSlotIndex, true);
-		OutUniqueTrianglePrevPosition0[UniqueTriangleIndex] = float4(Tri.P0, PackUVsToFloat(Tri.UV0));
-		OutUniqueTrianglePrevPosition1[UniqueTriangleIndex] = float4(Tri.P1, PackUVsToFloat(Tri.UV1));
-		OutUniqueTrianglePrevPosition2[UniqueTriangleIndex] = float4(Tri.P2, PackUVsToFloat(Tri.UV2));
+		OutUniqueTrianglePrevPosition[UniqueTriangleIndex * 3 + 0] = float4(Tri.P0, PackUVsToFloat(Tri.UV0));
+		OutUniqueTrianglePrevPosition[UniqueTriangleIndex * 3 + 1] = float4(Tri.P1, PackUVsToFloat(Tri.UV1));
+		OutUniqueTrianglePrevPosition[UniqueTriangleIndex * 3 + 2] = float4(Tri.P2, PackUVsToFloat(Tri.UV2));
 	}
 	#endif
 #else
 	{
 		const FHairMeshTriangle Tri = GetTriangleData(TriangleIndex, SectionSlotIndex, false);
-		OutUniqueTriangleCurrPosition0[UniqueTriangleIndex] = (float4(Tri.P0, 1));
-		OutUniqueTriangleCurrPosition1[UniqueTriangleIndex] = (float4(Tri.P1, 1));
-		OutUniqueTriangleCurrPosition2[UniqueTriangleIndex] = (float4(Tri.P2, 1));
+		OutUniqueTriangleCurrPosition[UniqueTriangleIndex * 3 + 0] = (float4(Tri.P0, 1));
+		OutUniqueTriangleCurrPosition[UniqueTriangleIndex * 3 + 1] = (float4(Tri.P1, 1));
+		OutUniqueTriangleCurrPosition[UniqueTriangleIndex * 3 + 2] = (float4(Tri.P2, 1));
 	}
 
 	#if PERMUTATION_PREVIOUS
 	{
 		const FHairMeshTriangle Tri = GetTriangleData(TriangleIndex, SectionSlotIndex, true);
-		OutUniqueTrianglePrevPosition0[UniqueTriangleIndex] = (float4(Tri.P0, 1));
-		OutUniqueTrianglePrevPosition1[UniqueTriangleIndex] = (float4(Tri.P1, 1));
-		OutUniqueTrianglePrevPosition2[UniqueTriangleIndex] = (float4(Tri.P2, 1));
+		OutUniqueTrianglePrevPosition[UniqueTriangleIndex * 3 + 0] = (float4(Tri.P0, 1));
+		OutUniqueTrianglePrevPosition[UniqueTriangleIndex * 3 + 1] = (float4(Tri.P1, 1));
+		OutUniqueTrianglePrevPosition[UniqueTriangleIndex * 3 + 2] = (float4(Tri.P2, 1));
 	}
 	#endif
 #endif
@@ -443,12 +439,12 @@ void MainCS(uint3 DispatchThreadId : SV_DispatchThreadID)
 #if SHADER_OFFSET_UPDATE
 
 #if PERMUTATION_USE_GPU_OFFSET
-	Buffer<float4> RootTriangleCurrPosition0Buffer;
-	float3 GetOffsetCurrPosition()	{ return RootTriangleCurrPosition0Buffer[0].xyz; }
+	Buffer<float4> RootTriangleCurrPositionBuffer;
+	float3 GetOffsetCurrPosition()	{ return RootTriangleCurrPositionBuffer[0].xyz; }
 
 	#if PERMUTATION_PREVIOUS
-	Buffer<float4> RootTrianglePrevPosition0Buffer;
-	float3 GetOffsetPrevPosition()	{ return RootTrianglePrevPosition0Buffer[0].xyz; }
+	Buffer<float4> RootTrianglePrevPositionBuffer;
+	float3 GetOffsetPrevPosition()	{ return RootTrianglePrevPositionBuffer[0].xyz; }
 	#endif
 #else
 	float3 CPUCurrPositionOffset;
diff --git a/Engine/Shaders/Private/HairStrands/HairStrandsMeshProjectionCommon.ush b/Engine/Shaders/Private/HairStrands/HairStrandsMeshProjectionCommon.ush
--- a/Engine/Shaders/Private/HairStrands/HairStrandsMeshProjectionCommon.ush
+++ b/Engine/Shaders/Private/HairStrands/HairStrandsMeshProjectionCommon.ush
@@ -107,16 +107,14 @@ FHairMeshTriangleWithUVs GetTriangleData(uint TriangleIndex, Buffer<float> InPos
 
 FHairMeshTriangle GetTriangleTransformation(
 	uint TriangleIndex,
-	Buffer<float4>  InPosition0Buffer,
-	Buffer<float4>  InPosition1Buffer,
-	Buffer<float4>  InPosition2Buffer)
+	Buffer<float4> InPositionBuffer)
 {
 	FHairMeshTriangle Out;
 
 	// Build orthonormal frame
-	Out.P0 = InPosition0Buffer[TriangleIndex].xyz;
-	Out.P1 = InPosition1Buffer[TriangleIndex].xyz;
-	Out.P2 = InPosition2Buffer[TriangleIndex].xyz;
+	Out.P0 = InPositionBuffer[TriangleIndex * 3 + 0].xyz;
+	Out.P1 = InPositionBuffer[TriangleIndex * 3 + 1].xyz;
+	Out.P2 = InPositionBuffer[TriangleIndex * 3 + 2].xyz;
 	Out.E0 = Out.P1 - Out.P0;
 	Out.E1 = Out.P2 - Out.P0;
 
@@ -129,15 +127,13 @@ FHairMeshTriangle GetTriangleTransformation(
 
 FHairMeshTriangleWithUVs GetTriangleWithUVsTransformation(
 	uint TriangleIndex,
-	Buffer<float4>  InPosition0Buffer,
-	Buffer<float4>  InPosition1Buffer,
-	Buffer<float4>  InPosition2Buffer)
+	Buffer<float4> InPositionBuffer)
 {
 	FHairMeshTriangleWithUVs Out;
 
-	const float4 V0 = InPosition0Buffer[TriangleIndex];
-	const float4 V1 = InPosition1Buffer[TriangleIndex];
-	const float4 V2 = InPosition2Buffer[TriangleIndex];
+	const float4 V0 = InPositionBuffer[TriangleIndex * 3 + 0];
+	const float4 V1 = InPositionBuffer[TriangleIndex * 3 + 1];
+	const float4 V2 = InPositionBuffer[TriangleIndex * 3 + 2];
 
 	// Build orthonormal frame
 	Out.P0 = V0.xyz;
