Renamed HairStrandsRendering into HairStrandsInterpolation to avoid confusion with the engine part.

diff --git a/Engine/Plugins/Runtime/HairStrands/Source/HairStrandsCore/Private/GroomComponent.cpp b/Engine/Plugins/Runtime/HairStrands/Source/HairStrandsCore/Private/GroomComponent.cpp
--- a/Engine/Plugins/Runtime/HairStrands/Source/HairStrandsCore/Private/GroomComponent.cpp
+++ b/Engine/Plugins/Runtime/HairStrands/Source/HairStrandsCore/Private/GroomComponent.cpp
@@ -10,7 +10,6 @@
 #include "NiagaraSystem.h"
 #include "PrimitiveSceneProxy.h"
 #include "PhysicsEngine/PhysicsAsset.h"
-#include "HairStrandsRendering.h"
 #include "HairCardsVertexFactory.h"
 #include "HairStrandsVertexFactory.h"
 #include "RayTracingInstance.h"
diff --git a/Engine/Plugins/Runtime/HairStrands/Source/HairStrandsCore/Private/GroomManager.cpp b/Engine/Plugins/Runtime/HairStrands/Source/HairStrandsCore/Private/GroomManager.cpp
--- a/Engine/Plugins/Runtime/HairStrands/Source/HairStrandsCore/Private/GroomManager.cpp
+++ b/Engine/Plugins/Runtime/HairStrands/Source/HairStrandsCore/Private/GroomManager.cpp
@@ -17,6 +17,7 @@
 #include "GeometryCacheSceneProxy.h"
 #include "HAL/IConsoleManager.h"
 #include "SceneView.h"
+#include "HairStrandsInterpolation.h"
 #include "HairCardsVertexFactory.h"
 #include "HairStrandsVertexFactory.h"
 #include "RenderGraphEvent.h"
diff --git a/Engine/Plugins/Runtime/HairStrands/Source/HairStrandsCore/Private/GroomManager.h b/Engine/Plugins/Runtime/HairStrands/Source/HairStrandsCore/Private/GroomManager.h
--- a/Engine/Plugins/Runtime/HairStrands/Source/HairStrandsCore/Private/GroomManager.h
+++ b/Engine/Plugins/Runtime/HairStrands/Source/HairStrandsCore/Private/GroomManager.h
@@ -10,7 +10,6 @@
 #include "RenderResource.h"
 #include "RenderGraphResources.h"
 #include "HairStrandsInterface.h"
-#include "HairStrandsRendering.h"
 #include "HairStrandsMeshProjection.h"
 
 
diff --git a/Engine/Plugins/Runtime/HairStrands/Source/HairStrandsCore/Private/HairStrandsInterpolation.cpp b/Engine/Plugins/Runtime/HairStrands/Source/HairStrandsCore/Private/HairStrandsInterpolation.cpp
new file mode 100644
--- /dev/null
+++ b/Engine/Plugins/Runtime/HairStrands/Source/HairStrandsCore/Private/HairStrandsInterpolation.cpp
@@ -0,0 +1,2479 @@
+// Copyright Epic Games, Inc. All Rights Reserved. 
+
+#include "HairStrandsInterpolation.h"
+#include "HairStrandsDatas.h"
+#include "HairCardsBuilder.h"
+#include "RenderGraphBuilder.h"
+#include "RenderGraphUtils.h"
+#include "ShaderParameterStruct.h"
+#include "GlobalShader.h"
+#include "GroomAsset.h"
+#include "GroomBindingAsset.h"
+#include "GroomCache.h"
+#include "HairStrandsInterface.h"
+#include "SceneView.h"
+#include "Components/SkeletalMeshComponent.h"
+#include "Containers/ResourceArray.h"
+#include "Rendering/SkeletalMeshRenderData.h"
+#include "HAL/ConsoleManager.h"
+#include "ShaderPrintParameters.h"
+#include "Async/ParallelFor.h"
+#include "RenderTargetPool.h"
+#include "GroomTextureBuilder.h"
+#include "GroomBindingBuilder.h"
+#include "SkeletalRenderPublic.h"
+#include "SkeletalMeshDeformerHelpers.h"
+#include "GroomAsset.h" 
+#include "GroomManager.h"
+#include "GroomInstance.h"
+#include "SystemTextures.h"
+#include "ShaderPrintParameters.h"
+
+static float GHairRaytracingRadiusScale = 0;
+static FAutoConsoleVariableRef CVarHairRaytracingRadiusScale(TEXT("r.HairStrands.RaytracingRadiusScale"), GHairRaytracingRadiusScale, TEXT("Override the per instance scale factor for raytracing hair strands geometry (0: disabled, >0:enabled)"));
+
+static int GHairRaytracingProceduralSplits = 4;
+static FAutoConsoleVariableRef CVarHairRaytracingProceduralScale(TEXT("r.HairStrands.RaytracingProceduralSplits"), GHairRaytracingProceduralSplits, TEXT("Change how many AABBs are used per hair segment to balance between BVH build cost and ray tracing performance. (default: 4)"));
+
+static float GStrandHairWidth = 0.0f;
+static FAutoConsoleVariableRef CVarStrandHairWidth(TEXT("r.HairStrands.StrandWidth"), GStrandHairWidth, TEXT("Width of hair strand"));
+
+static int32 GStrandHairInterpolationDebug = 0;
+static FAutoConsoleVariableRef CVarStrandHairInterpolationDebug(TEXT("r.HairStrands.Interpolation.Debug"), GStrandHairInterpolationDebug, TEXT("Enable debug rendering for hair interpolation"));
+
+static int32 GHairCardsInterpolationType = 1;
+static FAutoConsoleVariableRef CVarHairCardsInterpolationType(TEXT("r.HairStrands.Cards.InterpolationType"), GHairCardsInterpolationType, TEXT("Hair cards interpolation type: 0: None, 1:physics simulation, 2: RBF deformation"));
+
+static int32 GHairStrandsTransferPositionOnLODChange = 0;
+static FAutoConsoleVariableRef CVarHairStrandsTransferPositionOnLODChange(TEXT("r.HairStrands.Strands.TransferPrevPos"), GHairStrandsTransferPositionOnLODChange, TEXT("Transfer strands prev. position to current position on LOD switching to avoid large discrepancy causing large motion vector"));
+
+static int32 GHairStrands_Raytracing_ForceRebuildBVH = 1;
+static FAutoConsoleVariableRef CVarHairRTGeomForceRebuild(TEXT("r.HairStrands.Strands.Raytracing.ForceRebuildBVH"), GHairStrands_Raytracing_ForceRebuildBVH, TEXT("Force BVH rebuild instead of doing a BVH refit when hair positions changed"));
+
+////////////////////////////////////////////////////////////////////////////////////////////////////////////////////////
+
+enum class EHairCardsSimulationType
+{
+	None,
+	Guide,
+	RBF
+};
+
+EHairCardsSimulationType GetHairCardsSimulationType()
+{
+	return GHairCardsInterpolationType >= 2 ? 
+		EHairCardsSimulationType::RBF : 
+		(GHairCardsInterpolationType >= 1 ? EHairCardsSimulationType::Guide : EHairCardsSimulationType::None);
+}
+
+bool NeedsUpdateCardsMeshTriangles()
+{
+	return GetHairCardsSimulationType() == EHairCardsSimulationType::Guide;
+}
+
+int GetHairRaytracingProceduralSplits()
+{
+	return FMath::Clamp(GHairRaytracingProceduralSplits, 1, STRANDS_PROCEDURAL_INTERSECTOR_MAX_SPLITS);
+}
+
+bool IsHairStrandsContinousLODEnabled();
+
+inline uint32 ComputeGroupSize()
+{
+	const uint32 GroupSize = IsRHIDeviceAMD() ? 64 : (IsRHIDeviceNVIDIA() ? 32 : 64);
+	check(GroupSize == 64 || GroupSize == 32);
+	return GroupSize;
+}
+
+enum class EDeformationType : uint8
+{
+	Simulation,		// Use the output of the hair simulation
+	RestGuide,		// Use the rest guide as input of the interpolation (no deformation), only weighted interpolation
+	OffsetGuide		// Offset the guides
+};
+
+
+struct FGroomCacheResources
+{
+	FRDGBufferSRVRef PositionBuffer = nullptr;
+	FRDGBufferSRVRef RadiusBuffer = nullptr;
+	bool bHasRadiusData = false;
+};
+
+static FGroomCacheResources CreateGroomCacheBuffer(FRDGBuilder& GraphBuilder, FGroomCacheVertexData& InVertexData)
+{
+	FGroomCacheResources Out;
+	
+	if (InVertexData.PointsPosition.Num() > 0)
+	{
+		FRDGBufferRef PositionBuffer = nullptr;
+		if (InVertexData.PositionBuffer == nullptr)
+		{
+			const uint32 DataCount = InVertexData.PointsPosition.Num();
+			const uint32 DataSizeInBytes = sizeof(FVector3f) * DataCount;
+			// Deformation are upload into a Buffer<float> as the original position are float3 which is is both
+			// 1) incompatible with structure buffer alignment (128bits), and 2) incompatible with vertex buffer 
+			// as R32G32B32_FLOAT format is not well supported for SRV across HW.
+			// So instead the positions are uploaded into vertex buffer Buffer<float>
+			PositionBuffer = CreateVertexBuffer(
+				GraphBuilder,
+				TEXT("GroomCache_PositionBuffer"),
+				FRDGBufferDesc::CreateBufferDesc(sizeof(float), InVertexData.PointsPosition.Num() * 3),
+				InVertexData.PointsPosition.GetData(),
+				DataSizeInBytes,
+				ERDGInitialDataFlags::None);
+			GraphBuilder.QueueBufferExtraction(PositionBuffer, &InVertexData.PositionBuffer);
+		}
+		else
+		{
+			PositionBuffer = GraphBuilder.RegisterExternalBuffer(InVertexData.PositionBuffer);
+		}
+		Out.PositionBuffer = GraphBuilder.CreateSRV(PositionBuffer, PF_R32_FLOAT);
+	}
+
+	const uint32 RadiusDataCount = InVertexData.PointsRadius.Num();
+	if (RadiusDataCount > 0)
+	{
+		FRDGBufferRef RadiusBuffer = nullptr;
+		if (InVertexData.RadiusBuffer == nullptr)
+		{
+			const uint32 RadiusDataSizeInBytes = sizeof(float) * RadiusDataCount;
+			RadiusBuffer = CreateVertexBuffer(
+				GraphBuilder,
+				TEXT("GroomCache_RadiusBuffer"),
+				FRDGBufferDesc::CreateBufferDesc(sizeof(float), InVertexData.PointsRadius.Num()),
+				InVertexData.PointsRadius.GetData(),
+				RadiusDataSizeInBytes,
+				ERDGInitialDataFlags::None);
+			GraphBuilder.QueueBufferExtraction(RadiusBuffer, &InVertexData.PositionBuffer);
+		}
+		else
+		{
+			RadiusBuffer = GraphBuilder.RegisterExternalBuffer(InVertexData.RadiusBuffer);
+		}
+		Out.bHasRadiusData = true;
+		Out.RadiusBuffer = GraphBuilder.CreateSRV(RadiusBuffer, PF_R32_FLOAT);
+	}
+	else
+	{
+		Out.RadiusBuffer = GraphBuilder.CreateSRV(GSystemTextures.GetDefaultBuffer(GraphBuilder, 4), PF_R32_FLOAT);
+	}
+
+	return Out;
+}
+
+////////////////////////////////////////////////////////////////////////////////////////////////////////////////////////
+class FTransferVelocityPassCS : public FGlobalShader
+{
+	DECLARE_GLOBAL_SHADER(FTransferVelocityPassCS);
+	SHADER_USE_PARAMETER_STRUCT(FTransferVelocityPassCS, FGlobalShader);
+
+	BEGIN_SHADER_PARAMETER_STRUCT(FParameters, )
+		SHADER_PARAMETER(uint32, ElementCount)
+		SHADER_PARAMETER_RDG_BUFFER_SRV(Buffer, InBuffer)
+		SHADER_PARAMETER_RDG_BUFFER_UAV(RWBuffer, OutBuffer)
+	END_SHADER_PARAMETER_STRUCT()
+
+public:
+	static uint32 GetGroupSize() { return HAIR_VERTEXCOUNT_GROUP_SIZE; }
+	static bool ShouldCompilePermutation(const FGlobalShaderPermutationParameters& Parameters) { return IsHairStrandsSupported(EHairStrandsShaderType::All, Parameters.Platform); }
+	static void ModifyCompilationEnvironment(const FGlobalShaderPermutationParameters& Parameters, FShaderCompilerEnvironment& OutEnvironment)
+	{
+		FGlobalShader::ModifyCompilationEnvironment(Parameters, OutEnvironment);
+		OutEnvironment.SetDefine(TEXT("SHADER_HAIRTRANSFER_PREV_POSITION"), 1);
+		OutEnvironment.SetDefine(TEXT("GROUP_SIZE"), GetGroupSize());
+	}
+};
+
+IMPLEMENT_GLOBAL_SHADER(FTransferVelocityPassCS, "/Engine/Private/HairStrands/HairStrandsInterpolation.usf", "MainCS", SF_Compute);
+
+static void AddTransferPositionPass(
+	FRDGBuilder& GraphBuilder,
+	FGlobalShaderMap* ShaderMap,
+	const uint32 ElementCount,
+	FRDGBufferSRVRef InBuffer,
+	FRDGBufferUAVRef OutBuffer)
+{
+	if (ElementCount == 0) return;
+
+	FTransferVelocityPassCS::FParameters* Parameters = GraphBuilder.AllocParameters<FTransferVelocityPassCS::FParameters>();
+	Parameters->ElementCount = ElementCount;
+	Parameters->InBuffer = InBuffer;
+	Parameters->OutBuffer = OutBuffer;
+
+	const FIntVector DispatchCount(FMath::DivideAndRoundUp(ElementCount, FTransferVelocityPassCS::GetGroupSize()), 1, 1);
+	TShaderMapRef<FTransferVelocityPassCS> ComputeShader(ShaderMap);
+	FComputeShaderUtils::AddPass(
+		GraphBuilder,
+		RDG_EVENT_NAME("HairStrands::TransferPositionForVelocity"),
+		ComputeShader,
+		Parameters,
+		DispatchCount);
+}
+
+////////////////////////////////////////////////////////////////////////////////////////////////////////////////////////
+
+class FGroomCacheUpdatePassCS : public FGlobalShader
+{
+	DECLARE_GLOBAL_SHADER(FGroomCacheUpdatePassCS);
+	SHADER_USE_PARAMETER_STRUCT(FGroomCacheUpdatePassCS, FGlobalShader);
+
+	BEGIN_SHADER_PARAMETER_STRUCT(FParameters, )
+		SHADER_PARAMETER(uint32, ElementCount)
+		SHADER_PARAMETER(uint32, bHasRadiusData)
+		SHADER_PARAMETER(float, InterpolationFactor)
+		SHADER_PARAMETER_RDG_BUFFER_SRV(Buffer, InPosition0Buffer)
+		SHADER_PARAMETER_RDG_BUFFER_SRV(Buffer, InPosition1Buffer)
+		SHADER_PARAMETER_RDG_BUFFER_SRV(Buffer, InRadius0Buffer)
+		SHADER_PARAMETER_RDG_BUFFER_SRV(Buffer, InRadius1Buffer)
+		SHADER_PARAMETER_RDG_BUFFER_SRV(Buffer, InRestPoseBuffer)
+		SHADER_PARAMETER_RDG_BUFFER_SRV(Buffer, InDeformedOffsetBuffer)
+		SHADER_PARAMETER_RDG_BUFFER_UAV(RWBuffer, OutDeformedBuffer)
+	END_SHADER_PARAMETER_STRUCT()
+
+public:
+	static uint32 GetGroupSize() { return HAIR_VERTEXCOUNT_GROUP_SIZE; }
+	static bool ShouldCompilePermutation(const FGlobalShaderPermutationParameters& Parameters) { return IsHairStrandsSupported(EHairStrandsShaderType::All, Parameters.Platform); }
+	static void ModifyCompilationEnvironment(const FGlobalShaderPermutationParameters& Parameters, FShaderCompilerEnvironment& OutEnvironment)
+	{
+		FGlobalShader::ModifyCompilationEnvironment(Parameters, OutEnvironment);
+		OutEnvironment.SetDefine(TEXT("SHADER_GROOMCACHE_UPDATE"), 1);
+		OutEnvironment.SetDefine(TEXT("GROUP_SIZE"), GetGroupSize());
+	}
+};
+
+IMPLEMENT_GLOBAL_SHADER(FGroomCacheUpdatePassCS, "/Engine/Private/HairStrands/HairStrandsInterpolation.usf", "MainCS", SF_Compute);
+
+static void AddGroomCacheUpdatePass(
+	FRDGBuilder& GraphBuilder,
+	FGlobalShaderMap* ShaderMap,
+	uint32 PointCount,
+	float InterpolationFactor,
+	FGroomCacheResources CacheResources0,
+	FGroomCacheResources CacheResources1,
+	FRDGBufferSRVRef InBuffer,
+	FRDGBufferSRVRef InDeformedOffsetBuffer,
+	FRDGBufferUAVRef OutBuffer)
+{
+	TRACE_CPUPROFILER_EVENT_SCOPE(AddGroomCacheUpdatePass);
+
+	if (PointCount == 0 || !CacheResources0.PositionBuffer || !CacheResources1.PositionBuffer) return;
+
+	FGroomCacheUpdatePassCS::FParameters* Parameters = GraphBuilder.AllocParameters<FGroomCacheUpdatePassCS::FParameters>();
+	Parameters->ElementCount = PointCount;
+	Parameters->InPosition0Buffer = CacheResources0.PositionBuffer;
+	Parameters->InPosition1Buffer = CacheResources1.PositionBuffer;
+	Parameters->InRadius0Buffer = CacheResources0.RadiusBuffer;
+	Parameters->InRadius1Buffer = CacheResources1.RadiusBuffer;
+	Parameters->InRestPoseBuffer = InBuffer;
+	Parameters->InDeformedOffsetBuffer = InDeformedOffsetBuffer;
+	Parameters->OutDeformedBuffer = OutBuffer;
+	Parameters->InterpolationFactor = InterpolationFactor;
+	Parameters->bHasRadiusData = CacheResources0.bHasRadiusData ? 1u : 0u;
+
+	const FIntVector DispatchCount = FIntVector(FMath::DivideAndRoundUp(PointCount, FGroomCacheUpdatePassCS::GetGroupSize()), 1, 1);
+	TShaderMapRef<FGroomCacheUpdatePassCS> ComputeShader(ShaderMap);
+	FComputeShaderUtils::AddPass(
+		GraphBuilder,
+		RDG_EVENT_NAME("HairStrands::GroomCacheUpdate"),
+		ComputeShader,
+		Parameters,
+		DispatchCount);
+}
+
+////////////////////////////////////////////////////////////////////////////////////////////////////////////////////////
+
+class FDeformGuideCS : public FGlobalShader
+{
+	DECLARE_GLOBAL_SHADER(FDeformGuideCS);
+	SHADER_USE_PARAMETER_STRUCT(FDeformGuideCS, FGlobalShader);
+
+	class FDeformationType : SHADER_PERMUTATION_INT("PERMUTATION_DEFORMATION", 5);
+	using FPermutationDomain = TShaderPermutationDomain<FDeformationType>;
+
+	BEGIN_SHADER_PARAMETER_STRUCT(FParameters, )
+		SHADER_PARAMETER(uint32, VertexCount)
+		SHADER_PARAMETER(FVector3f, SimRestOffset)
+		
+		SHADER_PARAMETER_RDG_BUFFER_SRV(Buffer<float4>, SimRestPositionBuffer)
+		SHADER_PARAMETER_RDG_BUFFER_SRV(Buffer<float4>, SimDeformedPositionBuffer)
+
+		SHADER_PARAMETER_RDG_BUFFER_SRV(Buffer<uint>, SimRootBarycentricBuffer)
+		SHADER_PARAMETER_RDG_BUFFER_SRV(Buffer<uint>, SimPointToCurveBuffer)
+		SHADER_PARAMETER_RDG_BUFFER_SRV(Buffer<uint>, SimRootToUniqueTriangleIndexBuffer)
+
+		SHADER_PARAMETER_RDG_BUFFER_SRV(Buffer, SimDeformedOffsetBuffer)
+		SHADER_PARAMETER_RDG_BUFFER_SRV(Buffer, SimRestPosePositionBuffer)
+		SHADER_PARAMETER_RDG_BUFFER_UAV(RWBuffer, OutSimDeformedPositionBuffer)
+
+		SHADER_PARAMETER(uint32, SampleCount)
+		SHADER_PARAMETER_RDG_BUFFER_SRV(Buffer, RestSamplePositionsBuffer)
+		SHADER_PARAMETER_RDG_BUFFER_SRV(Buffer, MeshSampleWeightsBuffer)
+
+		SHADER_PARAMETER_SRV(Buffer<float4>, BoneDeformedPositionBuffer)
+	END_SHADER_PARAMETER_STRUCT()
+
+public:
+	static uint32 GetGroupSize() { return HAIR_VERTEXCOUNT_GROUP_SIZE; }
+	static bool ShouldCompilePermutation(const FGlobalShaderPermutationParameters& Parameters) { return IsHairStrandsSupported(EHairStrandsShaderType::All, Parameters.Platform); }
+	static void ModifyCompilationEnvironment(const FGlobalShaderPermutationParameters& Parameters, FShaderCompilerEnvironment& OutEnvironment)
+	{
+		FGlobalShader::ModifyCompilationEnvironment(Parameters, OutEnvironment);
+		OutEnvironment.SetDefine(TEXT("SHADER_DEFORM_GUIDE"), 1);
+		OutEnvironment.SetDefine(TEXT("GROUP_SIZE"), GetGroupSize());
+	}
+};
+
+IMPLEMENT_GLOBAL_SHADER(FDeformGuideCS, "/Engine/Private/HairStrands/HairStrandsGuideDeform.usf", "MainCS", SF_Compute);
+
+static void AddDeformSimHairStrandsPass(
+	FRDGBuilder& GraphBuilder,
+	FGlobalShaderMap* ShaderMap,
+	EDeformationType DeformationType,
+	const uint32 MeshLODIndex,
+	const uint32 VertexCount,
+	FHairStrandsRestRootResource* SimRestRootResources,
+	FHairStrandsDeformedRootResource* SimDeformedRootResources,
+	FRDGBufferSRVRef SimRestPosePositionBuffer,
+	FRDGBufferSRVRef SimPointToCurveBuffer,
+	FRDGImportedBuffer OutSimDeformedPositionBuffer,
+	const FVector& SimRestOffset,
+	FRDGBufferSRVRef SimDeformedOffsetBuffer,
+	const bool bHasGlobalInterpolation,
+	FRHIShaderResourceView* BoneBufferSRV)
+{
+	enum EInternalDeformationType
+	{
+		InternalDeformationType_ByPass = 0,
+		InternalDeformationType_Offset = 1,
+		InternalDeformationType_Skinned = 2,
+		InternalDeformationType_RBF = 3,
+		InternalDeformationType_Bones = 4,
+		InternalDeformationTypeCount
+	};
+
+	EInternalDeformationType InternalDeformationType = InternalDeformationTypeCount;
+	switch (DeformationType)
+	{
+	case EDeformationType::RestGuide   : InternalDeformationType = InternalDeformationType_ByPass; break;
+	case EDeformationType::OffsetGuide : InternalDeformationType = InternalDeformationType_Offset; break;
+	}
+
+	if (InternalDeformationType == InternalDeformationTypeCount) return;
+
+	FDeformGuideCS::FParameters* Parameters = GraphBuilder.AllocParameters<FDeformGuideCS::FParameters>();
+	Parameters->SimRestPosePositionBuffer = SimRestPosePositionBuffer;
+	Parameters->OutSimDeformedPositionBuffer = OutSimDeformedPositionBuffer.UAV;
+	Parameters->VertexCount = VertexCount;
+	Parameters->SimDeformedOffsetBuffer = SimDeformedOffsetBuffer;
+	Parameters->SimRestOffset = (FVector3f)SimRestOffset;
+
+	if (DeformationType == EDeformationType::OffsetGuide)
+	{
+		const bool bIsPointToCurveBuffersValid = SimPointToCurveBuffer != nullptr;
+		if (bIsPointToCurveBuffersValid)
+		{
+			Parameters->SimPointToCurveBuffer = SimPointToCurveBuffer;
+		}
+
+		const uint32 RootCount = SimRestRootResources ? SimRestRootResources->BulkData.RootCount : 0;
+
+		if(BoneBufferSRV)
+		{
+			InternalDeformationType = InternalDeformationType_Bones;
+			Parameters->BoneDeformedPositionBuffer = BoneBufferSRV;
+		}
+		else
+		{
+			const bool bSupportDynamicMesh = 
+				RootCount > 0 && 
+				MeshLODIndex >= 0 && 
+				MeshLODIndex < uint32(SimRestRootResources->LODs.Num()) &&
+				MeshLODIndex < uint32(SimDeformedRootResources->LODs.Num()) &&
+				SimRestRootResources->LODs[MeshLODIndex].IsValid() &&
+				SimDeformedRootResources->LODs[MeshLODIndex].IsValid() &&
+				bIsPointToCurveBuffersValid;
+			
+			bool bSupportGlobalInterpolation = false;
+			if (bSupportDynamicMesh)
+			{
+				FHairStrandsRestRootResource::FLOD& RestLODDatas = SimRestRootResources->LODs[MeshLODIndex];
+				FHairStrandsDeformedRootResource::FLOD& DeformedLODDatas = SimDeformedRootResources->LODs[MeshLODIndex];
+
+				bSupportGlobalInterpolation = bHasGlobalInterpolation && (RestLODDatas.SampleCount > 0);
+				if (!bSupportGlobalInterpolation) 
+				{
+					InternalDeformationType = InternalDeformationType_Skinned;
+					Parameters->SimRootToUniqueTriangleIndexBuffer = RegisterAsSRV(GraphBuilder, RestLODDatas.RootToUniqueTriangleIndexBuffer);
+					Parameters->SimRestPositionBuffer = RegisterAsSRV(GraphBuilder, RestLODDatas.RestUniqueTrianglePositionBuffer);
+					Parameters->SimDeformedPositionBuffer = RegisterAsSRV(GraphBuilder, DeformedLODDatas.GetDeformedUniqueTrianglePositionBuffer(FHairStrandsDeformedRootResource::FLOD::Current));
+					Parameters->SimRootBarycentricBuffer = RegisterAsSRV(GraphBuilder, RestLODDatas.RootBarycentricBuffer);
+				}
+				else
+				{
+					InternalDeformationType = InternalDeformationType_RBF;
+					Parameters->MeshSampleWeightsBuffer = RegisterAsSRV(GraphBuilder, DeformedLODDatas.GetMeshSampleWeightsBuffer(FHairStrandsDeformedRootResource::FLOD::Current));
+					Parameters->RestSamplePositionsBuffer = RegisterAsSRV(GraphBuilder, RestLODDatas.RestSamplePositionsBuffer);
+					Parameters->SampleCount = RestLODDatas.SampleCount;
+				}
+			}
+		}
+	}
+
+	FDeformGuideCS::FPermutationDomain PermutationVector;
+	PermutationVector.Set<FDeformGuideCS::FDeformationType>(InternalDeformationType);
+
+	const FIntVector DispatchCount(FMath::DivideAndRoundUp(VertexCount,FDeformGuideCS::GetGroupSize()), 1, 1);
+	TShaderMapRef<FDeformGuideCS> ComputeShader(ShaderMap, PermutationVector);
+	FComputeShaderUtils::AddPass(
+		GraphBuilder,
+		RDG_EVENT_NAME("HairStrands::DeformSimHairStrands"),
+		ComputeShader,
+		Parameters,
+		DispatchCount);
+
+	GraphBuilder.SetBufferAccessFinal(OutSimDeformedPositionBuffer.Buffer, ERHIAccess::SRVMask);
+}
+
+///////////////////////////////////////////////////////////////////////////////////////////////////
+
+struct FRDGHairStrandsCullingData
+{
+	bool bCullingResultAvailable = false;
+	FRDGImportedBuffer HairStrandsVF_CullingIndirectBuffer;
+	FRDGImportedBuffer HairStrandsVF_CullingIndexBuffer;
+	FRDGImportedBuffer HairStrandsVF_CullingRadiusScaleBuffer;
+
+	uint32 ClusterCount = 0;
+	FRDGImportedBuffer ClusterAABBBuffer;
+	FRDGImportedBuffer GroupAABBBuffer;
+};
+
+FRDGHairStrandsCullingData ImportCullingData(FRDGBuilder& GraphBuilder, FHairGroupPublicData* In)
+{
+	FRDGHairStrandsCullingData Out;
+	Out.bCullingResultAvailable = In->GetCullingResultAvailable();
+
+	if (Out.bCullingResultAvailable)
+	{
+		Out.HairStrandsVF_CullingIndirectBuffer		= Register(GraphBuilder, In->GetDrawIndirectRasterComputeBuffer(), ERDGImportedBufferFlags::CreateViews);
+		Out.HairStrandsVF_CullingIndexBuffer		= Register(GraphBuilder, In->GetCulledVertexIdBuffer(), ERDGImportedBufferFlags::CreateViews);
+		Out.HairStrandsVF_CullingRadiusScaleBuffer	= Register(GraphBuilder, In->GetCulledVertexRadiusScaleBuffer(), ERDGImportedBufferFlags::CreateViews);
+	}
+
+	Out.ClusterCount		= In->GetClusterCount();
+	Out.ClusterAABBBuffer	= Register(GraphBuilder, In->GetClusterAABBBuffer(), ERDGImportedBufferFlags::CreateViews);
+	Out.GroupAABBBuffer		= Register(GraphBuilder, In->GetGroupAABBBuffer(), ERDGImportedBufferFlags::CreateViews);
+
+	return Out;
+}
+
+class FHairInterpolationCS : public FGlobalShader
+{
+	DECLARE_GLOBAL_SHADER(FHairInterpolationCS);
+	SHADER_USE_PARAMETER_STRUCT(FHairInterpolationCS, FGlobalShader);
+
+	class FDynamicGeometry : SHADER_PERMUTATION_INT("PERMUTATION_DYNAMIC_GEOMETRY", 5);
+	class FSimulation : SHADER_PERMUTATION_BOOL("PERMUTATION_SIMULATION");
+	class FSingleGuide : SHADER_PERMUTATION_BOOL("PERMUTATION_USE_SINGLE_GUIDE");
+	class FCulling : SHADER_PERMUTATION_BOOL("PERMUTATION_CULLING");
+	class FDeformer : SHADER_PERMUTATION_BOOL("PERMUTATION_DEFORMER");
+	using FPermutationDomain = TShaderPermutationDomain<FDynamicGeometry, FSimulation, FSingleGuide, FCulling, FDeformer>;
+
+	BEGIN_SHADER_PARAMETER_STRUCT(FParameters, )
+		SHADER_PARAMETER_STRUCT_INCLUDE(ShaderPrint::FShaderParameters, ShaderDrawParameters)
+		SHADER_PARAMETER(uint32, VertexCount)
+		SHADER_PARAMETER(float, HairLengthScale)
+		SHADER_PARAMETER(FVector3f, InRenHairPositionOffset)
+		SHADER_PARAMETER(FVector3f, InSimHairPositionOffset)
+		SHADER_PARAMETER(uint32,  HairStrandsVFTODO_bCullingEnable)
+		SHADER_PARAMETER(FMatrix44f, LocalToWorldMatrix)
+
+		SHADER_PARAMETER_RDG_BUFFER_SRV(Buffer, RenRestPosePositionBuffer)
+		SHADER_PARAMETER_RDG_BUFFER_SRV(Buffer, RenDeformerPositionBuffer)
+		SHADER_PARAMETER_RDG_BUFFER_UAV(RWBuffer, OutRenDeformedPositionBuffer)
+
+		SHADER_PARAMETER_RDG_BUFFER_SRV(Buffer, SimRestPosePositionBuffer)
+		SHADER_PARAMETER_RDG_BUFFER_SRV(Buffer, SimDeformedPositionBuffer)
+
+		SHADER_PARAMETER_RDG_BUFFER_SRV(ByteAddressBuffer, InterpolationBuffer)
+
+		SHADER_PARAMETER_RDG_BUFFER_SRV(Buffer<float4>, RenRootRestPositionBuffer)
+		SHADER_PARAMETER_RDG_BUFFER_SRV(Buffer<float4>, RenRootDeformedPositionBuffer)
+
+		SHADER_PARAMETER_RDG_BUFFER_SRV(Buffer<uint>, RenRootBarycentricBuffer)
+		SHADER_PARAMETER_RDG_BUFFER_SRV(Buffer<uint>, RenPointToCurveBuffer)
+		SHADER_PARAMETER_RDG_BUFFER_SRV(Buffer<uint>, RenRootToUniqueTriangleIndexBuffer)
+
+		SHADER_PARAMETER_RDG_BUFFER_SRV(Buffer<float4>, SimRootRestPositionBuffer)
+		SHADER_PARAMETER_RDG_BUFFER_SRV(Buffer<float4>, SimRootDeformedPositionBuffer)
+
+		SHADER_PARAMETER_RDG_BUFFER_SRV(Buffer<uint>, SimRootBarycentricBuffer)
+		SHADER_PARAMETER_RDG_BUFFER_SRV(Buffer<uint>, SimPointToCurveBuffer)
+		SHADER_PARAMETER_RDG_BUFFER_SRV(Buffer<uint>, SimRootToUniqueTriangleIndexBuffer)
+
+		SHADER_PARAMETER_RDG_BUFFER_SRV(Buffer<uint>, SimRootPointIndexBuffer)
+
+		SHADER_PARAMETER_RDG_BUFFER_SRV(Buffer<float4>, OutSimHairPositionOffsetBuffer)
+		SHADER_PARAMETER_RDG_BUFFER_SRV(Buffer<float4>, OutRenHairPositionOffsetBuffer)
+
+		SHADER_PARAMETER_RDG_BUFFER_SRV(Buffer<uint>, HairStrandsVFTODO_CullingIndirectBuffer)
+		SHADER_PARAMETER_RDG_BUFFER_SRV(Buffer<uint>, HairStrandsVFTODO_CullingIndexBuffer)
+		SHADER_PARAMETER_RDG_BUFFER_SRV(Buffer<float>, HairStrandsVFTODO_CullingRadiusScaleBuffer)
+		RDG_BUFFER_ACCESS(HairStrandsVFTODO_CullingIndirectBufferArgs, ERHIAccess::IndirectArgs)
+		
+		END_SHADER_PARAMETER_STRUCT()
+
+public:
+	static uint32 GetGroupSize() { return HAIR_VERTEXCOUNT_GROUP_SIZE; }
+	static bool ShouldCompilePermutation(const FGlobalShaderPermutationParameters& Parameters) { return IsHairStrandsSupported(EHairStrandsShaderType::All, Parameters.Platform); }
+	static void ModifyCompilationEnvironment(const FGlobalShaderPermutationParameters& Parameters, FShaderCompilerEnvironment& OutEnvironment)
+	{
+		FGlobalShader::ModifyCompilationEnvironment(Parameters, OutEnvironment);
+		OutEnvironment.SetDefine(TEXT("SHADER_HAIRINTERPOLATION"), 1);
+		OutEnvironment.SetDefine(TEXT("GROUP_SIZE"), GetGroupSize());
+	}
+};
+
+IMPLEMENT_GLOBAL_SHADER(FHairInterpolationCS, "/Engine/Private/HairStrands/HairStrandsInterpolation.usf", "MainCS", SF_Compute);
+
+// AddHairStrandsInterpolationPass takes in input
+// * Guide resources (optional)
+// * Guide root triangles (optional)
+// * Strands resources
+// * Strands root triangles (optional)
+// 
+// This functions is in charge of deforming the rendering strands (OutRenderPositionBuffer) based on the guides
+// The guides are either the result of a simulation deformation, or a global interpolation deformation (aka RBF deformation), or both
+static void AddHairStrandsInterpolationPass(
+	FRDGBuilder& GraphBuilder,
+	FGlobalShaderMap* ShaderMap,
+	const FShaderPrintData* ShaderPrintData,
+	const FHairGroupInstance* Instance,
+	const uint32 VertexCount,
+	const int32 MeshLODIndex,
+	const float HairLengthScale,
+	const EHairInterpolationType HairInterpolationType,
+	const EHairGeometryType InstanceGeometryType,
+	const FRDGHairStrandsCullingData& CullingData,
+	const FVector& InRenHairWorldOffset,
+	const FVector& InSimHairWorldOffset,
+	const FRDGBufferSRVRef& OutRenHairPositionOffsetBuffer,
+	const FRDGBufferSRVRef& OutSimHairPositionOffsetBuffer,
+	const FHairStrandsRestRootResource* RenRestRootResources,
+	const FHairStrandsRestRootResource* SimRestRootResources,
+	const FHairStrandsDeformedRootResource* RenDeformedRootResources,
+	const FHairStrandsDeformedRootResource* SimDeformedRootResources,
+	const FRDGBufferSRVRef& RenRestPosePositionBuffer,
+	const FRDGBufferSRVRef& RenPointToCurveBuffer,
+	const bool bUseSingleGuide,
+	const FRDGBufferSRVRef& InterpolationBuffer,
+	const FRDGBufferSRVRef& SimRestPosePositionBuffer,
+	const FRDGBufferSRVRef& SimDeformedPositionBuffer,
+	const FRDGBufferSRVRef& SimRootPointIndexBuffer,
+	const FRDGBufferSRVRef& SimPointToCurveBuffer,
+	const FRDGBufferSRVRef& RenDeformerPositionBuffer,
+	FRDGImportedBuffer& OutRenPositionBuffer,
+	const FHairStrandsDeformedRootResource::FLOD::EFrameType DeformedFrame)
+{
+	FHairInterpolationCS::FParameters* Parameters = GraphBuilder.AllocParameters<FHairInterpolationCS::FParameters>();
+	Parameters->RenRestPosePositionBuffer = RenRestPosePositionBuffer;
+	Parameters->SimRestPosePositionBuffer = SimRestPosePositionBuffer;
+	Parameters->SimDeformedPositionBuffer = SimDeformedPositionBuffer;
+	Parameters->InterpolationBuffer = InterpolationBuffer;
+	Parameters->OutRenDeformedPositionBuffer = OutRenPositionBuffer.UAV;
+
+	Parameters->VertexCount = VertexCount;
+	Parameters->InRenHairPositionOffset = (FVector3f)InRenHairWorldOffset;
+	Parameters->InSimHairPositionOffset = (FVector3f)InSimHairWorldOffset;
+
+	Parameters->OutSimHairPositionOffsetBuffer = OutSimHairPositionOffsetBuffer;
+	Parameters->OutRenHairPositionOffsetBuffer = OutRenHairPositionOffsetBuffer;
+
+	Parameters->SimRootPointIndexBuffer = SimRootPointIndexBuffer;
+	Parameters->SimPointToCurveBuffer = SimPointToCurveBuffer;
+	Parameters->RenPointToCurveBuffer = RenPointToCurveBuffer;
+	
+	Parameters->LocalToWorldMatrix = FMatrix44f(Instance->LocalToWorld.ToMatrixWithScale());		// LWC_TODO: Precision loss
+	Parameters->HairLengthScale = HairLengthScale;
+
+	// Only needed if DEBUG_ENABLE is manually enabledin HairStrandsInterpolation.usf. This is for manual debug purpose only
+	#if 0
+	if (ShaderPrintData && ShaderPrintData->IsValid())
+	{
+		ShaderPrint::SetParameters(GraphBuilder, *ShaderPrintData, Parameters->ShaderDrawParameters);
+	}
+	#endif
+
+	const bool bSupportDynamicMesh = 
+		RenRestRootResources &&
+		RenRestRootResources->BulkData.RootCount > 0 && 
+		MeshLODIndex >= 0 && 
+		MeshLODIndex < RenRestRootResources->LODs.Num() &&
+		MeshLODIndex < RenDeformedRootResources->LODs.Num() &&
+		RenRestRootResources->LODs[MeshLODIndex].IsValid() &&
+		RenDeformedRootResources->LODs[MeshLODIndex].IsValid();
+
+	// Guides
+	bool bSupportGlobalInterpolation = false;
+	if (bSupportDynamicMesh && (Instance->Guides.bIsSimulationEnable || Instance->Guides.bHasGlobalInterpolation || Instance->Guides.bIsDeformationEnable)) // No need for Instance->Guides.bHasSimulationCache since it is incompatible with binding
+	{
+		const FHairStrandsRestRootResource::FLOD& Sim_RestLODDatas = SimRestRootResources->LODs[MeshLODIndex];
+		const FHairStrandsDeformedRootResource::FLOD& Sim_DeformedLODDatas = SimDeformedRootResources->LODs[MeshLODIndex];
+		bSupportGlobalInterpolation = Instance->Guides.bHasGlobalInterpolation && (Sim_RestLODDatas.SampleCount > 0);
+		{
+			Parameters->SimRootRestPositionBuffer = RegisterAsSRV(GraphBuilder, Sim_RestLODDatas.RestUniqueTrianglePositionBuffer);
+			Parameters->SimRootBarycentricBuffer = RegisterAsSRV(GraphBuilder, Sim_RestLODDatas.RootBarycentricBuffer);
+			Parameters->SimRootToUniqueTriangleIndexBuffer = RegisterAsSRV(GraphBuilder, Sim_RestLODDatas.RootToUniqueTriangleIndexBuffer);
+		}
+		{
+			Parameters->SimRootDeformedPositionBuffer = RegisterAsSRV(GraphBuilder, Sim_DeformedLODDatas.GetDeformedUniqueTrianglePositionBuffer(DeformedFrame));
+		}
+	}
+
+	// Strands
+	if (bSupportDynamicMesh)
+	{
+		const FHairStrandsRestRootResource::FLOD& Ren_RestLODDatas = RenRestRootResources->LODs[MeshLODIndex];
+		const FHairStrandsDeformedRootResource::FLOD& Ren_DeformedLODDatas = RenDeformedRootResources->LODs[MeshLODIndex];
+		{
+			Parameters->RenRootRestPositionBuffer = RegisterAsSRV(GraphBuilder, Ren_RestLODDatas.RestUniqueTrianglePositionBuffer);
+			Parameters->RenRootBarycentricBuffer = RegisterAsSRV(GraphBuilder, Ren_RestLODDatas.RootBarycentricBuffer);
+			Parameters->RenRootToUniqueTriangleIndexBuffer = RegisterAsSRV(GraphBuilder, Ren_RestLODDatas.RootToUniqueTriangleIndexBuffer);
+		}
+		{
+			Parameters->RenRootDeformedPositionBuffer = RegisterAsSRV(GraphBuilder, Ren_DeformedLODDatas.GetDeformedUniqueTrianglePositionBuffer(DeformedFrame));
+		}
+	}
+
+	const bool bSupportDeformer = RenDeformerPositionBuffer != nullptr;
+	if (bSupportDeformer)
+	{
+		Parameters->RenDeformerPositionBuffer = RenDeformerPositionBuffer;
+	}
+
+	const bool bHasLocalDeformation = Instance->Guides.bIsSimulationEnable || bSupportGlobalInterpolation || Instance->Guides.bIsDeformationEnable || Instance->Guides.bIsSimulationCacheEnable;
+	const bool bCullingEnable = IsHairStrandContinuousDecimationReorderingEnabled() ? false : (InstanceGeometryType == EHairGeometryType::Strands && CullingData.bCullingResultAvailable); 	// TODO: improve reordering so that culling can be used effectively
+	Parameters->HairStrandsVFTODO_bCullingEnable = bCullingEnable ? 1 : 0;
+
+	// Select dynamic geometry permutation, based on the Simulation/RBF/InterpolationType
+	int32 DynamicGeometryType = 0;
+	{
+		switch (HairInterpolationType)
+		{
+		case EHairInterpolationType::NoneSkinning			: DynamicGeometryType = 0; break; // INTERPOLATION_RIGID
+		case EHairInterpolationType::RigidSkinning			: DynamicGeometryType = 1; break; // INTERPOLATION_SKINNING_OFFSET
+		case EHairInterpolationType::OffsetSkinning			: DynamicGeometryType = 2; break; // INTERPOLATION_SKINNING_TRANSLATION
+		case EHairInterpolationType::SmoothSkinning			: DynamicGeometryType = 3; break; // INTERPOLATION_SKINNING_ROTATION
+		//case EHairInterpolationType::SmoothOffsetSkinning	: DynamicGeometryType = 4; break; // INTERPOLATION_SKINNING_TRANSLATION_AND_ROTATION - Experimental, not used
+		}		
+
+		if ( bSupportDynamicMesh && !bHasLocalDeformation) DynamicGeometryType = 1; // INTERPOLATION_SKINNING_OFFSET
+		if (!bSupportDynamicMesh && !bHasLocalDeformation) DynamicGeometryType = 0; // INTERPOLATION_RIGID
+		if (!bSupportDynamicMesh &&  bHasLocalDeformation) DynamicGeometryType = 0; // INTERPOLATION_RIGID
+	}
+
+	FHairInterpolationCS::FPermutationDomain PermutationVector;
+	PermutationVector.Set<FHairInterpolationCS::FDynamicGeometry>(DynamicGeometryType);
+	PermutationVector.Set<FHairInterpolationCS::FSimulation>(bHasLocalDeformation);
+	PermutationVector.Set<FHairInterpolationCS::FSingleGuide>(bUseSingleGuide);
+	PermutationVector.Set<FHairInterpolationCS::FCulling>(bCullingEnable);
+	PermutationVector.Set<FHairInterpolationCS::FDeformer>(bSupportDeformer);
+
+	TShaderMapRef<FHairInterpolationCS> ComputeShader(ShaderMap, PermutationVector);
+
+	if (bCullingEnable)
+	{
+		Parameters->HairStrandsVFTODO_CullingIndirectBuffer = CullingData.HairStrandsVF_CullingIndirectBuffer.SRV;
+		Parameters->HairStrandsVFTODO_CullingIndexBuffer = CullingData.HairStrandsVF_CullingIndexBuffer.SRV;
+		Parameters->HairStrandsVFTODO_CullingRadiusScaleBuffer = CullingData.HairStrandsVF_CullingRadiusScaleBuffer.SRV;
+		Parameters->HairStrandsVFTODO_CullingIndirectBufferArgs = CullingData.HairStrandsVF_CullingIndirectBuffer.Buffer;
+
+		FComputeShaderUtils::AddPass(
+			GraphBuilder,
+			RDG_EVENT_NAME("HairStrands::Interpolation(culling=on)"),
+			ComputeShader, 
+			Parameters,
+			CullingData.HairStrandsVF_CullingIndirectBuffer.Buffer,
+			0);
+	}
+	else
+	{
+		const FIntVector DispatchCount(FMath::DivideAndRoundUp(VertexCount, FHairInterpolationCS::GetGroupSize()), 1, 1);
+		FComputeShaderUtils::AddPass(
+			GraphBuilder,
+			RDG_EVENT_NAME("HairStrands::Interpolation(culling=off)"),
+			ComputeShader,
+			Parameters,
+			DispatchCount);
+	}
+
+	GraphBuilder.SetBufferAccessFinal(OutRenPositionBuffer.Buffer, ERHIAccess::SRVMask);
+}
+
+///////////////////////////////////////////////////////////////////////////////////////////////////
+// Patch hair attribute (For debug visualization)
+
+class FHairPatchAttributeCS : public FGlobalShader
+{
+	DECLARE_GLOBAL_SHADER(FHairPatchAttributeCS);
+	SHADER_USE_PARAMETER_STRUCT(FHairPatchAttributeCS, FGlobalShader);
+
+	class FSimulation : SHADER_PERMUTATION_INT("PERMUTATION_SIMULATION", 3); // No simulation, Simlation with 1 guide, Simulation with 3 guides
+	using FPermutationDomain = TShaderPermutationDomain<FSimulation>;
+
+	BEGIN_SHADER_PARAMETER_STRUCT(FParameters, )
+		SHADER_PARAMETER(uint32, CurveCount)
+		SHADER_PARAMETER_ARRAY(FUintVector4, RenCurveAttributeOffsets, [HAIR_CURVE_ATTRIBUTE_OFFSET_COUNT])
+		SHADER_PARAMETER_RDG_BUFFER_SRV(Buffer, RenCurveBuffer)
+		SHADER_PARAMETER_RDG_BUFFER_SRV(ByteAddressBuffer, InterpolationBuffer)
+		SHADER_PARAMETER_RDG_BUFFER_SRV(Buffer, RenVertexToClusterIdBuffer)
+		SHADER_PARAMETER_RDG_BUFFER_UAV(RWByteAddressBuffer, OutRenAttributeBuffer)
+	END_SHADER_PARAMETER_STRUCT()
+
+public:
+	static bool ShouldCompilePermutation(const FGlobalShaderPermutationParameters& Parameters) { return IsHairStrandsSupported(EHairStrandsShaderType::All, Parameters.Platform); }
+	static void ModifyCompilationEnvironment(const FGlobalShaderPermutationParameters& Parameters, FShaderCompilerEnvironment& OutEnvironment)
+	{
+		FGlobalShader::ModifyCompilationEnvironment(Parameters, OutEnvironment);
+		OutEnvironment.SetDefine(TEXT("SHADER_PATCHATTRIBUTE"), 1);
+	}
+};
+
+IMPLEMENT_GLOBAL_SHADER(FHairPatchAttributeCS, "/Engine/Private/HairStrands/HairStrandsInterpolation.usf", "MainCS", SF_Compute);
+
+enum class EHairPatchAttribute : uint8
+{
+	None,
+	GuideInflucence,
+	ClusterInfluence
+};
+
+static void AddPatchAttributePass(
+	FRDGBuilder& GraphBuilder,
+	FGlobalShaderMap* ShaderMap,
+	const uint32 CurveCount,
+	const EHairPatchAttribute Mode,
+	const bool bSimulation,
+	const bool bUseSingleGuide,
+	const uint32* RenCurveAttributeOffests,
+	const FRDGBufferRef& RenAttributeBuffer,
+	const FRDGBufferSRVRef& RenCurveBuffer,
+	const FRDGBufferSRVRef& RenVertexToClusterIdBuffer,
+	const FRDGBufferSRVRef& InterpolationBuffer,
+	FRDGImportedBuffer& OutRenAttributeBuffer)
+{
+	// Sanity check
+	check(Mode == EHairPatchAttribute::ClusterInfluence || Mode == EHairPatchAttribute::GuideInflucence)
+	if (Mode == EHairPatchAttribute::GuideInflucence && !bSimulation) return;
+
+	const FIntVector DispatchCount = FIntVector(FMath::DivideAndRoundUp(CurveCount, 1024u), 1, 1);
+
+	// First copy all the rendering attributes into the new bufffer, before overriding some of them
+	AddCopyBufferPass(GraphBuilder, RenAttributeBuffer, OutRenAttributeBuffer.Buffer);
+
+	FHairPatchAttributeCS::FParameters* Parameters = GraphBuilder.AllocParameters<FHairPatchAttributeCS::FParameters>();
+	Parameters->CurveCount = CurveCount;
+	Parameters->RenCurveBuffer = RenCurveBuffer;
+	Parameters->RenVertexToClusterIdBuffer = RenVertexToClusterIdBuffer;
+	PACK_CURVE_HAIR_ATTRIBUTE_OFFSETS(Parameters->RenCurveAttributeOffsets, RenCurveAttributeOffests)
+
+	if (bSimulation)
+	{
+		Parameters->InterpolationBuffer  = InterpolationBuffer;
+	}
+
+	Parameters->OutRenAttributeBuffer = OutRenAttributeBuffer.UAV;
+
+	FHairPatchAttributeCS::FPermutationDomain PermutationVector;
+	PermutationVector.Set<FHairPatchAttributeCS::FSimulation>(Mode == EHairPatchAttribute::ClusterInfluence ? 0 : (bUseSingleGuide ? 1 : 2));
+
+	TShaderMapRef<FHairPatchAttributeCS> ComputeShader(ShaderMap, PermutationVector);
+	FComputeShaderUtils::AddPass(
+		GraphBuilder,
+		RDG_EVENT_NAME("HairStrands::PatchAttribute"),
+		ComputeShader,
+		Parameters,
+		DispatchCount);
+
+	GraphBuilder.SetBufferAccessFinal(OutRenAttributeBuffer.Buffer, ERHIAccess::SRVMask);
+}
+
+///////////////////////////////////////////////////////////////////////////////////////////////////
+
+class FHairClusterAABBCS : public FGlobalShader
+{
+	DECLARE_GLOBAL_SHADER(FHairClusterAABBCS);
+	SHADER_USE_PARAMETER_STRUCT(FHairClusterAABBCS, FGlobalShader);
+
+	class FGroupSize : SHADER_PERMUTATION_SPARSE_INT("PERMUTATION_GROUP_SIZE", 32, 64);
+	class FCPUAABB : SHADER_PERMUTATION_BOOL("PERMUTATION_USE_CPU_AABB");
+	using FPermutationDomain = TShaderPermutationDomain<FGroupSize, FCPUAABB>;
+
+	BEGIN_SHADER_PARAMETER_STRUCT(FParameters, )
+		SHADER_PARAMETER_STRUCT_INCLUDE(ShaderPrint::FShaderParameters, ShaderDrawParameters)
+		SHADER_PARAMETER(uint32, TotalClusterCount)
+		SHADER_PARAMETER(uint32, PointCount)
+
+		SHADER_PARAMETER(FVector3f, CPUBoundMin)
+		SHADER_PARAMETER(FVector3f, CPUBoundMax)
+		SHADER_PARAMETER(FMatrix44f, LocalToTranslatedWorldMatrix)
+
+		SHADER_PARAMETER_RDG_BUFFER_SRV(Buffer, RenderDeformedPositionBuffer)
+		SHADER_PARAMETER_RDG_BUFFER_SRV(Buffer, RenderDeformedOffsetBuffer)
+		SHADER_PARAMETER_RDG_BUFFER_SRV(Buffer, ClusterVertexIdBuffer)
+		SHADER_PARAMETER_RDG_BUFFER_SRV(Buffer, ClusterIdBuffer)
+		SHADER_PARAMETER_RDG_BUFFER_SRV(Buffer, ClusterIndexOffsetBuffer)
+		SHADER_PARAMETER_RDG_BUFFER_SRV(Buffer, ClusterIndexCountBuffer)
+		SHADER_PARAMETER_RDG_BUFFER_SRV(Buffer, CulledClusterCountBuffer)
+		SHADER_PARAMETER_RDG_BUFFER_SRV(Buffer<uint>, HairStrandsVFTODO_CullingIndirectBuffer)
+		SHADER_PARAMETER_RDG_BUFFER_SRV(Buffer<uint>, IndirectBuffer)
+		SHADER_PARAMETER_RDG_BUFFER_UAV(RWBuffer, OutClusterAABBBuffer)
+		SHADER_PARAMETER_RDG_BUFFER_UAV(RWBuffer, OutGroupAABBBuffer)
+		RDG_BUFFER_ACCESS(IndirectArgsBuffer, ERHIAccess::IndirectArgs)
+	END_SHADER_PARAMETER_STRUCT()
+
+public:
+	static bool ShouldCompilePermutation(const FGlobalShaderPermutationParameters& Parameters) { return IsHairStrandsSupported(EHairStrandsShaderType::Strands, Parameters.Platform); }
+	static void ModifyCompilationEnvironment(const FGlobalShaderPermutationParameters& Parameters, FShaderCompilerEnvironment& OutEnvironment)
+	{
+		FGlobalShader::ModifyCompilationEnvironment(Parameters, OutEnvironment);
+		OutEnvironment.SetDefine(TEXT("SHADER_CLUSTERAABB"), 1);
+	}
+};
+
+IMPLEMENT_GLOBAL_SHADER(FHairClusterAABBCS, "/Engine/Private/HairStrands/HairStrandsInterpolation.usf", "ClusterAABBEvaluationCS", SF_Compute);
+
+enum class EHairAABBUpdateType
+{
+	UpdateClusterAABB,
+	UpdateGroupAABB
+};
+
+static void AddHairClusterAABBPass(
+	FRDGBuilder& GraphBuilder,
+	FGlobalShaderMap* ShaderMap,
+	const FVector& TranslatedWorldOffset,
+	const FShaderPrintData* ShaderPrintData,
+	const EHairAABBUpdateType UpdateType,
+	FHairGroupInstance* Instance,
+	FRDGBufferSRVRef RenderDeformedOffsetBuffer,
+	FHairStrandClusterData::FHairGroup* ClusterData,
+	FRDGHairStrandsCullingData& ClusterAABBData,
+	FRDGBufferSRVRef RenderPositionBufferSRV,
+	FRDGBufferSRVRef& DrawIndirectRasterComputeBuffer)
+{
+	// Clusters AABB are only update if the groom is deformed.
+	const FBoxSphereBounds& Bounds = Instance->Strands.Data->Header.BoundingBox;
+	FTransform InRenLocalToTranslatedWorld = Instance->LocalToWorld;
+	InRenLocalToTranslatedWorld.AddToTranslation(TranslatedWorldOffset);
+	const FBoxSphereBounds TransformedBounds = Bounds.TransformBy(InRenLocalToTranslatedWorld);
+	Instance->HairGroupPublicData->bClusterAABBValid = UpdateType == EHairAABBUpdateType::UpdateClusterAABB;
+
+	// bNeedDeformation: 
+	// * If the instance has deformation (simulation, global interpolation, skinning, then we update all clusters (for voxelization purpose) and the instance AABB
+	// * If the instance is static, we only update the instance AABB
+	const uint32 GroupSize = ComputeGroupSize();
+
+	FHairClusterAABBCS::FParameters* Parameters = GraphBuilder.AllocParameters<FHairClusterAABBCS::FParameters>();
+	Parameters->CPUBoundMin = (FVector3f)TransformedBounds.GetBox().Min;
+	Parameters->CPUBoundMax = (FVector3f)TransformedBounds.GetBox().Max;
+	Parameters->LocalToTranslatedWorldMatrix = FMatrix44f(InRenLocalToTranslatedWorld.ToMatrixWithScale());
+	Parameters->RenderDeformedPositionBuffer = RenderPositionBufferSRV;
+	Parameters->RenderDeformedOffsetBuffer = RenderDeformedOffsetBuffer;
+	Parameters->TotalClusterCount = 1;
+
+	Parameters->PointCount = Instance->HairGroupPublicData->GetActiveStrandsPointCount();
+
+	if (ShaderPrintData)
+	{
+		ShaderPrint::SetParameters(GraphBuilder, *ShaderPrintData, Parameters->ShaderDrawParameters);
+		
+	}
+
+	if (UpdateType == EHairAABBUpdateType::UpdateClusterAABB)
+	{
+		check(ClusterData);
+		check(ClusterData->GroupSize1D == GroupSize);
+		check(ClusterData->CulledCluster2DIndirectArgsBuffer);
+		check(ClusterData->CulledClusterCountBuffer);
+
+		FRDGBufferRef ClusterIdBuffer = GraphBuilder.RegisterExternalBuffer(ClusterData->ClusterIdBuffer);
+		FRDGBufferRef ClusterIndexOffsetBuffer = GraphBuilder.RegisterExternalBuffer(ClusterData->ClusterIndexOffsetBuffer);
+		FRDGBufferRef ClusterIndexCountBuffer  = GraphBuilder.RegisterExternalBuffer(ClusterData->ClusterIndexCountBuffer);
+
+		Parameters->TotalClusterCount = ClusterData->ClusterCount;
+		Parameters->ClusterVertexIdBuffer = RegisterAsSRV(GraphBuilder, *ClusterData->ClusterVertexIdBuffer);
+		Parameters->ClusterIdBuffer = GraphBuilder.CreateSRV(ClusterIdBuffer, PF_R32_UINT);
+		Parameters->ClusterIndexOffsetBuffer = GraphBuilder.CreateSRV(ClusterIndexOffsetBuffer, PF_R32_UINT);
+		Parameters->ClusterIndexCountBuffer = GraphBuilder.CreateSRV(ClusterIndexCountBuffer, PF_R32_UINT);
+		Parameters->CulledClusterCountBuffer = GraphBuilder.CreateSRV(ClusterData->CulledClusterCountBuffer, PF_R32_UINT);
+		Parameters->IndirectBuffer = GraphBuilder.CreateSRV(ClusterData->CulledCluster2DIndirectArgsBuffer, PF_R32_UINT);
+		Parameters->IndirectArgsBuffer = ClusterData->CulledCluster2DIndirectArgsBuffer;
+
+		// Sanity check
+		check(ClusterData->ClusterCount == ClusterIdBuffer->Desc.NumElements);
+		check(ClusterData->ClusterCount == ClusterIndexOffsetBuffer->Desc.NumElements);
+		check(ClusterData->ClusterCount == ClusterIndexCountBuffer->Desc.NumElements);
+		// Currently disabled this sanity check as the RDG allocation can return a larger buffer (for reuse purpose)
+		//check(ClusterData->ClusterCount * 6 == ClusterAABBData.ClusterAABBBuffer.Buffer->Desc.NumElements);
+	}
+	Parameters->HairStrandsVFTODO_CullingIndirectBuffer = DrawIndirectRasterComputeBuffer; // Used for checking max vertex count
+	Parameters->OutClusterAABBBuffer = ClusterAABBData.ClusterAABBBuffer.UAV;
+	Parameters->OutGroupAABBBuffer = ClusterAABBData.GroupAABBBuffer.UAV;
+
+
+	FHairClusterAABBCS::FPermutationDomain PermutationVector;
+	PermutationVector.Set<FHairClusterAABBCS::FGroupSize>(GroupSize);
+	PermutationVector.Set<FHairClusterAABBCS::FCPUAABB>(UpdateType == EHairAABBUpdateType::UpdateClusterAABB ? 0u : 1u);
+	TShaderMapRef<FHairClusterAABBCS> ComputeShader(ShaderMap, PermutationVector);
+
+	if (UpdateType == EHairAABBUpdateType::UpdateClusterAABB)
+	{
+		FComputeShaderUtils::AddPass(
+			GraphBuilder,
+			RDG_EVENT_NAME("HairStrands::ClusterAABB(%s)", TEXT("Cluster,Group"), TEXT("Group Only")),
+			ComputeShader,
+			Parameters,
+			Parameters->IndirectArgsBuffer, 0);
+	}
+	else
+	{
+		FComputeShaderUtils::AddPass(
+			GraphBuilder,
+			RDG_EVENT_NAME("HairStrands::ClusterAABB(%s)", TEXT("Cluster,Group"), TEXT("Group Only")),
+			ComputeShader,
+			Parameters,
+			FIntVector(1,1,1));
+	}
+
+	GraphBuilder.SetBufferAccessFinal(ClusterAABBData.ClusterAABBBuffer.Buffer, ERHIAccess::SRVMask);
+	GraphBuilder.SetBufferAccessFinal(ClusterAABBData.GroupAABBBuffer.Buffer, ERHIAccess::SRVMask);
+}
+
+///////////////////////////////////////////////////////////////////////////////////////////////////
+
+class FHairCardsDeformationCS : public FGlobalShader
+{
+	DECLARE_GLOBAL_SHADER(FHairCardsDeformationCS);
+	SHADER_USE_PARAMETER_STRUCT(FHairCardsDeformationCS, FGlobalShader);
+
+	class FDynamicGeometry : SHADER_PERMUTATION_INT("PERMUTATION_DYNAMIC_GEOMETRY", 2);
+	using FPermutationDomain = TShaderPermutationDomain<FDynamicGeometry>;
+
+	BEGIN_SHADER_PARAMETER_STRUCT(FParameters, )
+		SHADER_PARAMETER(FMatrix44f, LocalToWorld)
+		SHADER_PARAMETER(uint32, CardsVertexCount)
+		SHADER_PARAMETER(uint32, GuideVertexCount)
+		SHADER_PARAMETER(FVector3f, GuideRestPositionOffset)
+		SHADER_PARAMETER_RDG_BUFFER_SRV(Buffer, GuideRestPositionBuffer)
+		SHADER_PARAMETER_RDG_BUFFER_SRV(Buffer, GuideDeformedPositionBuffer)
+		SHADER_PARAMETER_RDG_BUFFER_SRV(Buffer, GuideDeformedPositionOffsetBuffer)
+		SHADER_PARAMETER_SRV(Buffer, CardsRestPositionBuffer)
+		SHADER_PARAMETER_SRV(Buffer, CardsRestTangentBuffer)
+		SHADER_PARAMETER_RDG_BUFFER_SRV(ByteAddressBuffer, CardsInterpolationBuffer)
+
+		SHADER_PARAMETER_RDG_BUFFER_SRV(Buffer<float4>, TriangleRestPositionBuffer)
+		SHADER_PARAMETER_RDG_BUFFER_SRV(Buffer<float4>, TriangleDeformedPositionBuffer)
+
+		SHADER_PARAMETER_RDG_BUFFER_SRV(Buffer<uint>, GuideRootToUniqueTriangleIndexBuffer)
+		SHADER_PARAMETER_RDG_BUFFER_SRV(Buffer<uint>, GuideRootBarycentricBuffer)
+		SHADER_PARAMETER_RDG_BUFFER_SRV(Buffer<uint>, GuideVertexToRootIndexBuffer)
+
+		SHADER_PARAMETER_RDG_BUFFER_UAV(RWBuffer, CardsDeformedPositionBuffer)
+		SHADER_PARAMETER_RDG_BUFFER_UAV(RWBuffer, CardsDeformedTangentBuffer)
+
+		SHADER_PARAMETER_STRUCT_INCLUDE(ShaderPrint::FShaderParameters, ShaderDrawParameters)
+	END_SHADER_PARAMETER_STRUCT()
+
+public:
+	static uint32 GetGroupSize() { return 256u; }
+	static bool ShouldCompilePermutation(const FGlobalShaderPermutationParameters& Parameters) { return IsHairStrandsSupported(EHairStrandsShaderType::Cards, Parameters.Platform); }
+	static void ModifyCompilationEnvironment(const FGlobalShaderPermutationParameters& Parameters, FShaderCompilerEnvironment& OutEnvironment)
+	{
+		FGlobalShader::ModifyCompilationEnvironment(Parameters, OutEnvironment);
+		OutEnvironment.SetDefine(TEXT("SHADER_DEFORM_CARDS"), 1);
+		OutEnvironment.SetDefine(TEXT("GROUP_SIZE"), GetGroupSize());
+	}
+};
+
+IMPLEMENT_GLOBAL_SHADER(FHairCardsDeformationCS, "/Engine/Private/HairStrands/HairCardsDeformation.usf", "MainCS", SF_Compute);
+
+static void AddHairCardsDeformationPass(
+	FRDGBuilder& GraphBuilder,
+	FGlobalShaderMap* ShaderMap,
+	const FShaderPrintData* ShaderPrintData,
+	FHairGroupInstance* Instance,
+	const int32 MeshLODIndex)
+{
+	const int32 HairLODIndex = Instance->HairGroupPublicData->GetIntLODIndex();
+	if (!Instance->Cards.IsValid(HairLODIndex))
+		return;
+
+	FHairGroupInstance::FCards::FLOD& LOD = Instance->Cards.LODs[HairLODIndex];
+	
+	FRDGImportedBuffer CardsDeformedPositionBuffer_Curr = Register(GraphBuilder, LOD.DeformedResource->GetBuffer(FHairCardsDeformedResource::Current), ERDGImportedBufferFlags::CreateViews);
+	FRDGImportedBuffer CardsDeformedPositionBuffer_Prev = Register(GraphBuilder, LOD.DeformedResource->GetBuffer(FHairCardsDeformedResource::Previous), ERDGImportedBufferFlags::CreateViews);
+
+	FRDGImportedBuffer CardsDeformedNormalBuffer = Register(GraphBuilder, LOD.DeformedResource->DeformedNormalBuffer, ERDGImportedBufferFlags::CreateViews);
+	GraphBuilder.UseInternalAccessMode({
+		CardsDeformedPositionBuffer_Curr.Buffer,
+		CardsDeformedPositionBuffer_Prev.Buffer,
+		CardsDeformedNormalBuffer.Buffer
+	});
+
+	// If LOD hasn't switched, copy the current buffer into previous buffer to have correct motion vectors
+	const bool bHasLODSwitch = Instance->HairGroupPublicData->VFInput.bHasLODSwitch;
+	if (!bHasLODSwitch)
+	{
+		AddCopyBufferPass(GraphBuilder, CardsDeformedPositionBuffer_Prev.Buffer, CardsDeformedPositionBuffer_Curr.Buffer);
+	}
+	
+	FHairCardsDeformationCS::FParameters* Parameters = GraphBuilder.AllocParameters<FHairCardsDeformationCS::FParameters>();
+	Parameters->LocalToWorld				= FMatrix44f(Instance->GetCurrentLocalToWorld().ToMatrixWithScale());		// LWC_TODO: Precision loss
+	Parameters->GuideVertexCount			= LOD.Guides.RestResource->GetPointCount();
+	Parameters->GuideRestPositionOffset		= (FVector3f)LOD.Guides.RestResource->GetPositionOffset();
+	Parameters->GuideRestPositionBuffer		= RegisterAsSRV(GraphBuilder, LOD.Guides.RestResource->PositionBuffer);
+	Parameters->GuideDeformedPositionBuffer = RegisterAsSRV(GraphBuilder, LOD.Guides.DeformedResource->GetBuffer(FHairStrandsDeformedResource::Current));
+	Parameters->GuideDeformedPositionOffsetBuffer = RegisterAsSRV(GraphBuilder, LOD.Guides.DeformedResource->GetPositionOffsetBuffer(FHairStrandsDeformedResource::Current));
+
+	Parameters->CardsVertexCount			= LOD.RestResource->GetVertexCount();
+	Parameters->CardsRestPositionBuffer		= LOD.RestResource->RestPositionBuffer.ShaderResourceViewRHI;
+	Parameters->CardsRestTangentBuffer		= LOD.RestResource->NormalsBuffer.ShaderResourceViewRHI;
+	Parameters->CardsDeformedPositionBuffer = CardsDeformedPositionBuffer_Curr.UAV;
+	Parameters->CardsDeformedTangentBuffer	= CardsDeformedNormalBuffer.UAV;
+
+	Parameters->CardsInterpolationBuffer	= RegisterAsSRV(GraphBuilder, LOD.InterpolationResource->InterpolationBuffer);
+
+	const FHairStrandsRestResource* RestResources = LOD.Guides.RestResource;
+	const FHairStrandsRestRootResource* RestRootResources = LOD.Guides.RestRootResource;
+	const FHairStrandsDeformedRootResource* DeformedRootResources = LOD.Guides.DeformedRootResource;
+
+	const bool bSupportDynamicMesh =
+		Instance->BindingType == EHairBindingType::Skinning &&
+		RestRootResources &&
+		RestRootResources->BulkData.RootCount > 0 &&
+		MeshLODIndex >= 0 &&
+		MeshLODIndex < RestRootResources->LODs.Num() &&
+		MeshLODIndex < DeformedRootResources->LODs.Num() &&
+		RestRootResources->LODs[MeshLODIndex].IsValid() &&
+		DeformedRootResources->LODs[MeshLODIndex].IsValid();
+	if (bSupportDynamicMesh)
+	{
+		Parameters->GuideVertexToRootIndexBuffer = RegisterAsSRV(GraphBuilder, RestResources->PointToCurveBuffer);
+		const FHairStrandsRestRootResource::FLOD& RestLODDatas = RestRootResources->LODs[MeshLODIndex];
+		const FHairStrandsDeformedRootResource::FLOD& DeformedLODDatas = DeformedRootResources->LODs[MeshLODIndex];
+
+		Parameters->GuideRootBarycentricBuffer = RegisterAsSRV(GraphBuilder, RestLODDatas.RootBarycentricBuffer);
+		Parameters->GuideRootToUniqueTriangleIndexBuffer = RegisterAsSRV(GraphBuilder, RestLODDatas.RootToUniqueTriangleIndexBuffer);
+		Parameters->TriangleRestPositionBuffer = RegisterAsSRV(GraphBuilder, RestLODDatas.RestUniqueTrianglePositionBuffer);
+		Parameters->TriangleDeformedPositionBuffer = RegisterAsSRV(GraphBuilder, DeformedLODDatas.GetDeformedUniqueTrianglePositionBuffer(FHairStrandsDeformedRootResource::FLOD::Current));
+	}
+
+	if (ShaderPrintData)
+	{
+		ShaderPrint::SetParameters(GraphBuilder, *ShaderPrintData, Parameters->ShaderDrawParameters);
+	}
+
+	FHairCardsDeformationCS::FPermutationDomain PermutationVector;
+	PermutationVector.Set<FHairCardsDeformationCS::FDynamicGeometry>(bSupportDynamicMesh ? 1 : 0);
+	TShaderMapRef<FHairCardsDeformationCS> ComputeShader(ShaderMap, PermutationVector);
+
+	const FIntVector DispatchCount = FIntVector(FMath::DivideAndRoundUp(Parameters->CardsVertexCount, FHairCardsDeformationCS::GetGroupSize()), 1, 1);
+	FComputeShaderUtils::AddPass(
+		GraphBuilder,
+		RDG_EVENT_NAME("HairStrands::CardsDeformation(%s)", bSupportDynamicMesh ? TEXT("Dynamic") : TEXT("Static")),
+		ComputeShader,
+		Parameters,
+		DispatchCount);
+
+	// If LOD has switched, copy the current buffer, so that we don't get incorrect motion vector
+	if (bHasLODSwitch)
+	{
+		AddCopyBufferPass(GraphBuilder, CardsDeformedPositionBuffer_Prev.Buffer, CardsDeformedPositionBuffer_Curr.Buffer);
+	}
+
+	GraphBuilder.UseExternalAccessMode({
+		CardsDeformedPositionBuffer_Curr.Buffer,
+		CardsDeformedPositionBuffer_Prev.Buffer,
+		CardsDeformedNormalBuffer.Buffer
+	}, ERHIAccess::SRVMask);
+}
+
+///////////////////////////////////////////////////////////////////////////////////////////////////
+
+class FHairTangentCS : public FGlobalShader
+{
+	DECLARE_GLOBAL_SHADER(FHairTangentCS);
+	SHADER_USE_PARAMETER_STRUCT(FHairTangentCS, FGlobalShader);
+
+	class FCulling : SHADER_PERMUTATION_BOOL("PERMUTATION_CULLING");
+	using FPermutationDomain = TShaderPermutationDomain<FCulling>;
+
+	BEGIN_SHADER_PARAMETER_STRUCT(FParameters, )
+		SHADER_PARAMETER(uint32, VertexCount)
+		SHADER_PARAMETER(uint32, HairStrandsVF_bCullingEnable)
+		SHADER_PARAMETER_RDG_BUFFER_SRV(Buffer, PositionBuffer)
+		SHADER_PARAMETER_RDG_BUFFER_SRV(Buffer<uint>,	HairStrandsVF_CullingIndirectBuffer)
+		SHADER_PARAMETER_RDG_BUFFER_SRV(Buffer<uint>,	HairStrandsVF_CullingIndexBuffer)
+		SHADER_PARAMETER_RDG_BUFFER_UAV(RWBuffer, OutputTangentBuffer)
+		RDG_BUFFER_ACCESS(IndirectBufferArgs, ERHIAccess::IndirectArgs)
+	END_SHADER_PARAMETER_STRUCT()
+
+public:
+	static uint32 GetGroupSize() { return HAIR_VERTEXCOUNT_GROUP_SIZE; }
+	static bool ShouldCompilePermutation(const FGlobalShaderPermutationParameters& Parameters) { return IsHairStrandsSupported(EHairStrandsShaderType::Strands, Parameters.Platform); }
+	static void ModifyCompilationEnvironment(const FGlobalShaderPermutationParameters& Parameters, FShaderCompilerEnvironment& OutEnvironment)
+	{
+		FGlobalShader::ModifyCompilationEnvironment(Parameters, OutEnvironment);
+		OutEnvironment.SetDefine(TEXT("SHADER_TANGENT"), 1);
+		OutEnvironment.SetDefine(TEXT("GROUP_SIZE"), GetGroupSize());
+	}
+};
+
+IMPLEMENT_GLOBAL_SHADER(FHairTangentCS, "/Engine/Private/HairStrands/HairStrandsTangent.usf", "MainCS", SF_Compute);
+
+void AddHairTangentPass(
+	FRDGBuilder& GraphBuilder,
+	FGlobalShaderMap* ShaderMap,
+	uint32 VertexCount,
+	FHairGroupPublicData* HairGroupPublicData,
+	FRDGBufferSRVRef PositionBuffer,
+	FRDGImportedBuffer OutTangentBuffer)
+{
+	const bool bCullingEnable = IsHairStrandContinuousDecimationReorderingEnabled() ? false : (HairGroupPublicData && HairGroupPublicData->GetCullingResultAvailable());
+
+	FHairTangentCS::FParameters* Parameters = GraphBuilder.AllocParameters<FHairTangentCS::FParameters>();
+	Parameters->PositionBuffer = PositionBuffer;
+	Parameters->OutputTangentBuffer = OutTangentBuffer.UAV;
+	Parameters->VertexCount = VertexCount;
+	Parameters->HairStrandsVF_bCullingEnable = bCullingEnable ? 1 : 0;
+
+	FHairTangentCS::FPermutationDomain PermutationVector;
+	PermutationVector.Set<FHairTangentCS::FCulling>(bCullingEnable);
+
+	TShaderMapRef<FHairTangentCS> ComputeShader(ShaderMap, PermutationVector);
+
+	if (bCullingEnable)
+	{
+		FRDGImportedBuffer DrawIndirectRasterComputeBuffer  = Register(GraphBuilder, HairGroupPublicData->GetDrawIndirectRasterComputeBuffer(), ERDGImportedBufferFlags::CreateSRV);
+		Parameters->HairStrandsVF_CullingIndirectBuffer		= DrawIndirectRasterComputeBuffer.SRV;
+		Parameters->HairStrandsVF_CullingIndexBuffer		= RegisterAsSRV(GraphBuilder, HairGroupPublicData->GetCulledVertexIdBuffer());
+		Parameters->IndirectBufferArgs						= DrawIndirectRasterComputeBuffer.Buffer;
+		
+		FComputeShaderUtils::AddPass(
+			GraphBuilder,
+			RDG_EVENT_NAME("HairStrands::Tangent(culling=on)"),
+			ComputeShader,
+			Parameters,
+			DrawIndirectRasterComputeBuffer.Buffer, 0);
+	}
+	else
+	{
+		const FIntVector DispatchCount(FMath::DivideAndRoundUp(VertexCount, FHairTangentCS::GetGroupSize()), 1, 1);
+		FComputeShaderUtils::AddPass(
+			GraphBuilder,
+			RDG_EVENT_NAME("HairStrands::Tangent(culling=off)"),
+			ComputeShader,
+			Parameters,
+			DispatchCount);
+	}
+}
+
+///////////////////////////////////////////////////////////////////////////////////////////////////
+
+class FHairRaytracingGeometryCS : public FGlobalShader
+{
+	DECLARE_GLOBAL_SHADER(FHairRaytracingGeometryCS);
+	SHADER_USE_PARAMETER_STRUCT(FHairRaytracingGeometryCS, FGlobalShader);
+
+	class FCulling : SHADER_PERMUTATION_BOOL("PERMUTATION_CULLING");
+	class FProceduralPrimitive : SHADER_PERMUTATION_BOOL("PERMUTATION_PROCEDURAL_PRIMITIVE"); 
+	using FPermutationDomain = TShaderPermutationDomain<FCulling, FProceduralPrimitive>;
+
+	BEGIN_SHADER_PARAMETER_STRUCT(FParameters, )
+		SHADER_PARAMETER(uint32, VertexCount)
+		SHADER_PARAMETER(float, HairStrandsVF_HairRadius)
+		SHADER_PARAMETER(float, HairStrandsVF_HairRootScale)
+		SHADER_PARAMETER(float, HairStrandsVF_HairTipScale)
+
+		SHADER_PARAMETER(uint32, HairStrandsVF_bCullingEnable)
+		SHADER_PARAMETER_RDG_BUFFER_SRV(Buffer,	HairStrandsVF_CullingIndirectBuffer)
+		SHADER_PARAMETER_RDG_BUFFER_SRV(Buffer,	HairStrandsVF_CullingIndexBuffer)
+		SHADER_PARAMETER_RDG_BUFFER_SRV(Buffer,	HairStrandsVF_CullingRadiusScaleBuffer)
+		RDG_BUFFER_ACCESS(HairStrandsVF_CullingIndirectBufferArgs, ERHIAccess::IndirectArgs)
+
+		SHADER_PARAMETER_RDG_BUFFER_SRV(Buffer, PositionOffsetBuffer)
+		SHADER_PARAMETER_RDG_BUFFER_SRV(Buffer, PositionBuffer)
+		SHADER_PARAMETER_RDG_BUFFER_SRV(Buffer, TangentBuffer)
+		SHADER_PARAMETER_RDG_BUFFER_UAV(RWBuffer, OutputPositionBuffer)
+		SHADER_PARAMETER_RDG_BUFFER_UAV(RWBuffer, OutputIndexBuffer)
+		SHADER_PARAMETER(uint32, RaytracingProceduralSplits)
+		SHADER_PARAMETER_STRUCT_INCLUDE(ShaderPrint::FShaderParameters, ShaderPrintUniformBuffer)
+	END_SHADER_PARAMETER_STRUCT()
+
+public:
+	static uint32 GetGroupSize() { return HAIR_VERTEXCOUNT_GROUP_SIZE; }
+	static bool ShouldCompilePermutation(const FGlobalShaderPermutationParameters& Parameters) { return IsHairStrandsSupported(EHairStrandsShaderType::Strands, Parameters.Platform); }
+	static void ModifyCompilationEnvironment(const FGlobalShaderPermutationParameters& Parameters, FShaderCompilerEnvironment& OutEnvironment)
+	{
+		FGlobalShader::ModifyCompilationEnvironment(Parameters, OutEnvironment);
+		OutEnvironment.SetDefine(TEXT("SHADER_RT_GEOMETRY"), 1);
+		OutEnvironment.SetDefine(TEXT("GROUP_SIZE"), GetGroupSize());
+	}
+};
+
+BEGIN_SHADER_PARAMETER_STRUCT(FStrandsResourceBLASParameters, )
+	RDG_BUFFER_ACCESS(PositionBuffer, ERHIAccess::SRVCompute)
+	RDG_BUFFER_ACCESS(IndexBuffer, ERHIAccess::SRVCompute)
+END_SHADER_PARAMETER_STRUCT()
+
+BEGIN_SHADER_PARAMETER_STRUCT(FCardsOrMeshesResourceBLASParameters, )
+RDG_BUFFER_ACCESS(PositionBuffer, ERHIAccess::SRVCompute)
+END_SHADER_PARAMETER_STRUCT()
+
+IMPLEMENT_GLOBAL_SHADER(FHairRaytracingGeometryCS, "/Engine/Private/HairStrands/HairStrandsRaytracingGeometry.usf", "MainCS", SF_Compute);
+
+static void AddGenerateRaytracingGeometryPass(
+	FRDGBuilder& GraphBuilder,
+	FGlobalShaderMap* ShaderMap,
+	const FShaderPrintData* ShaderPrintData,
+	uint32 VertexCount,
+	bool bProceduralPrimitive,
+	int ProceduralSplits,
+	float HairRadius,
+	float RootScale,
+	float TipScale,
+	const FRDGBufferSRVRef& HairWorldOffsetBuffer,
+	FRDGHairStrandsCullingData& CullingData,
+	const FRDGBufferSRVRef& PositionBuffer,
+	const FRDGBufferSRVRef& TangentBuffer,
+	const FRDGBufferUAVRef& OutPositionBuffer,
+	const FRDGBufferUAVRef& OutIndexBuffer)
+{
+	FHairRaytracingGeometryCS::FParameters* Parameters = GraphBuilder.AllocParameters<FHairRaytracingGeometryCS::FParameters>();
+	Parameters->VertexCount = VertexCount;
+	Parameters->PositionOffsetBuffer = HairWorldOffsetBuffer;
+	Parameters->HairStrandsVF_HairRadius = HairRadius;
+	Parameters->HairStrandsVF_HairRootScale = RootScale;
+	Parameters->HairStrandsVF_HairTipScale = TipScale;
+	Parameters->PositionBuffer = PositionBuffer;
+	Parameters->TangentBuffer = TangentBuffer;
+	Parameters->OutputPositionBuffer = OutPositionBuffer;
+	Parameters->RaytracingProceduralSplits = GetHairRaytracingProceduralSplits();
+	if (ShaderPrintData)
+	{
+		ShaderPrint::SetParameters(GraphBuilder, *ShaderPrintData, Parameters->ShaderPrintUniformBuffer);
+	}
+	Parameters->OutputIndexBuffer = OutIndexBuffer;
+	Parameters->RaytracingProceduralSplits = ProceduralSplits;
+
+	const bool bCullingEnable = CullingData.bCullingResultAvailable;
+	Parameters->HairStrandsVF_bCullingEnable = bCullingEnable ? 1 : 0;
+	if (bCullingEnable)
+	{
+		Parameters->HairStrandsVF_CullingIndirectBuffer = CullingData.HairStrandsVF_CullingIndirectBuffer.SRV;
+		Parameters->HairStrandsVF_CullingIndexBuffer = CullingData.HairStrandsVF_CullingIndexBuffer.SRV;
+		Parameters->HairStrandsVF_CullingRadiusScaleBuffer = CullingData.HairStrandsVF_CullingRadiusScaleBuffer.SRV;
+		Parameters->HairStrandsVF_CullingIndirectBufferArgs = CullingData.HairStrandsVF_CullingIndirectBuffer.Buffer;
+	}
+
+	FHairRaytracingGeometryCS::FPermutationDomain PermutationVector;
+	PermutationVector.Set<FHairRaytracingGeometryCS::FCulling>(bCullingEnable);
+	PermutationVector.Set<FHairRaytracingGeometryCS::FProceduralPrimitive>(bProceduralPrimitive);
+	TShaderMapRef<FHairRaytracingGeometryCS> ComputeShader(ShaderMap, PermutationVector);
+
+	const FIntVector DispatchCount(FMath::DivideAndRoundUp(VertexCount, FHairRaytracingGeometryCS::GetGroupSize()), 1, 1);
+	FComputeShaderUtils::AddPass(
+		GraphBuilder,
+		RDG_EVENT_NAME("HairStrands::GenerateRaytracingGeometry(bCulling:%d, Procedural:%d)", bCullingEnable, bProceduralPrimitive),
+		ComputeShader,
+		Parameters,
+		DispatchCount);
+}
+
+///////////////////////////////////////////////////////////////////////////////////////////////////
+
+class FClearClusterAABBCS : public FGlobalShader
+{
+	DECLARE_GLOBAL_SHADER(FClearClusterAABBCS);
+	SHADER_USE_PARAMETER_STRUCT(FClearClusterAABBCS, FGlobalShader);
+
+	BEGIN_SHADER_PARAMETER_STRUCT(FParameters, )
+		SHADER_PARAMETER_RDG_BUFFER_UAV(RWBuffer, OutClusterAABBBuffer)
+		SHADER_PARAMETER_RDG_BUFFER_UAV(RWBuffer, OutGroupAABBBuffer)
+		SHADER_PARAMETER(uint32, ClusterCount)
+		SHADER_PARAMETER(uint32, bClearClusterAABBs)
+	END_SHADER_PARAMETER_STRUCT()
+
+public:
+	static bool ShouldCompilePermutation(const FGlobalShaderPermutationParameters& Parameters) { return IsHairStrandsSupported(EHairStrandsShaderType::Strands, Parameters.Platform); }
+	static void ModifyCompilationEnvironment(const FGlobalShaderPermutationParameters& Parameters, FShaderCompilerEnvironment& OutEnvironment)
+	{
+		FGlobalShader::ModifyCompilationEnvironment(Parameters, OutEnvironment);
+		OutEnvironment.SetDefine(TEXT("SHADER_CLEARCLUSTERAABB"), 1);
+	}
+};
+
+IMPLEMENT_GLOBAL_SHADER(FClearClusterAABBCS, "/Engine/Private/HairStrands/HairStrandsClusterCulling.usf", "MainClearClusterAABBCS", SF_Compute);
+
+static void AddClearClusterAABBPass(
+	FRDGBuilder& GraphBuilder,
+	FGlobalShaderMap* ShaderMap,
+	const EHairAABBUpdateType UpdateType,
+	uint32 ClusterCount,
+	FRDGImportedBuffer& OutClusterAABBBuffer,
+	FRDGImportedBuffer& OutGroupAABBBuffer)
+{
+	check(OutClusterAABBBuffer.Buffer);
+
+	FClearClusterAABBCS::FParameters* Parameters = GraphBuilder.AllocParameters<FClearClusterAABBCS::FParameters>();
+	Parameters->ClusterCount = ClusterCount;
+	Parameters->bClearClusterAABBs = UpdateType == EHairAABBUpdateType::UpdateClusterAABB ? 1 : 0;
+	Parameters->OutClusterAABBBuffer = OutClusterAABBBuffer.UAV;
+	Parameters->OutGroupAABBBuffer = OutGroupAABBBuffer.UAV;
+
+	TShaderMapRef<FClearClusterAABBCS> ComputeShader(ShaderMap);
+
+	const FIntVector DispatchCount = 
+		UpdateType == EHairAABBUpdateType::UpdateClusterAABB ? 
+		FIntVector(FMath::DivideAndRoundUp(ClusterCount * 6u, 64u), 1, 1) :
+		FIntVector(1,1,1);
+
+	FComputeShaderUtils::AddPass(
+		GraphBuilder,
+		RDG_EVENT_NAME("HairStrands::ClearClusterAABB(%s)", UpdateType == EHairAABBUpdateType::UpdateClusterAABB ? TEXT("Cluster,Group") : TEXT("Group Only")),
+		ComputeShader,
+		Parameters,
+		DispatchCount);
+
+	GraphBuilder.SetBufferAccessFinal(OutClusterAABBBuffer.Buffer, ERHIAccess::SRVMask),
+	GraphBuilder.SetBufferAccessFinal(OutGroupAABBBuffer.Buffer, ERHIAccess::SRVMask);
+}
+
+///////////////////////////////////////////////////////////////////////////////////////////////////
+
+#if RHI_RAYTRACING
+static void UpdateHairAccelerationStructure(FRHICommandList& RHICmdList, FRayTracingGeometry* RayTracingGeometry, EAccelerationStructureBuildMode InMode)
+{
+	SCOPED_DRAW_EVENT(RHICmdList, CommitHairRayTracingGeometryUpdates);
+
+	FRayTracingGeometryBuildParams Params;
+	Params.BuildMode = InMode;
+	Params.Geometry = RayTracingGeometry->RayTracingGeometryRHI;
+	Params.Segments = RayTracingGeometry->Initializer.Segments;
+
+	RHICmdList.BuildAccelerationStructures(MakeArrayView(&Params, 1));
+}
+
+static void BuildHairAccelerationStructure_Strands(
+	FRHICommandList& RHICmdList, 
+	uint32 RaytracingVertexCount, 
+	uint32 RayTracingIndexCount,
+	FBufferRHIRef& PositionBuffer, 
+	FBufferRHIRef& IndexBuffer,
+	FRayTracingGeometry* OutRayTracingGeometry,
+	const FString& AssetDebugName,
+	const FString& MeshComponentName,
+	uint32 LODIndex,
+	bool bProceduralPrimitive,
+	int ProceduralSplits)
+{
+	FRayTracingGeometryInitializer Initializer;
+#if HAIR_RESOURCE_DEBUG_NAME
+	FString DebugName(TEXT("AS_HairStrands_LOD") + FString::FromInt(LODIndex) + TEXT("_") + AssetDebugName);
+	Initializer.DebugName = FName(DebugName);
+#else
+	static const FName DebugName("AS_HairStrands");	
+	static int32 DebugNumber = 0;
+	Initializer.DebugName = FDebugName(DebugName, DebugNumber++);
+#endif
+#if RHI_ENABLE_RESOURCE_INFO
+	const FName OwnerName(FString::Printf(TEXT("%s [LOD%d]"), *MeshComponentName, LODIndex));
+	Initializer.OwnerName = OwnerName;
+#endif
+	if (bProceduralPrimitive)
+	{
+		Initializer.GeometryType = RTGT_Procedural;
+		Initializer.TotalPrimitiveCount = (RaytracingVertexCount / (2 * STRANDS_PROCEDURAL_INTERSECTOR_MAX_SPLITS)) * ProceduralSplits;
+	}
+	else
+	{
+		Initializer.IndexBuffer = IndexBuffer;
+		Initializer.IndexBufferOffset = 0;
+		Initializer.GeometryType = RTGT_Triangles;
+		Initializer.TotalPrimitiveCount = RayTracingIndexCount / 3;
+	}
+	Initializer.bFastBuild = true;
+	Initializer.bAllowUpdate = false; // we won't be using update, so specify this to allow degenerate primitives to be discarded
+
+	FRayTracingGeometrySegment Segment;
+	Segment.VertexBuffer = PositionBuffer;
+	Segment.VertexBufferElementType = FHairStrandsRaytracingFormat::VertexElementType;
+	if (bProceduralPrimitive)
+	{
+		// stride covers one AABB (which is made up of two vertices)
+		Segment.VertexBufferStride = FHairStrandsRaytracingFormat::SizeInByte * 2;
+		Segment.MaxVertices = RaytracingVertexCount;
+		Segment.NumPrimitives = Initializer.TotalPrimitiveCount;
+	}
+	else
+	{
+		Segment.VertexBufferStride = FHairStrandsRaytracingFormat::SizeInByte;
+		Segment.MaxVertices = RaytracingVertexCount;
+		Segment.NumPrimitives = RayTracingIndexCount / 3;
+	}
+	Initializer.Segments.Add(Segment);
+
+	OutRayTracingGeometry->SetInitializer(Initializer);
+	OutRayTracingGeometry->CreateRayTracingGeometry(ERTAccelerationStructureBuildPriority::Immediate);
+}
+
+static void BuildHairAccelerationStructure_Cards(FRHICommandList& RHICmdList, 	
+	FHairCardsRestResource* RestResource,
+	FHairCardsDeformedResource* DeformedResource,
+	FRayTracingGeometry* OutRayTracingGeometry,
+	const FString& AssetDebugName,
+	const FString& MeshComponentName,
+	uint32 LODIndex)
+{
+	FRayTracingGeometryInitializer Initializer;
+#if HAIR_RESOURCE_DEBUG_NAME
+	FString DebugName(TEXT("AS_HairCards_LOD") + FString::FromInt(LODIndex) + TEXT("_") + AssetDebugName);
+	Initializer.DebugName = FName(DebugName);
+#else
+	static const FName DebugName("AS_HairCards");
+	static int32 DebugNumber = 0;
+	Initializer.DebugName = FDebugName(DebugName, DebugNumber++);
+#endif
+#if RHI_ENABLE_RESOURCE_INFO
+	const FName OwnerName(FString::Printf(TEXT("%s [LOD%d]"), *MeshComponentName, LODIndex));
+	Initializer.OwnerName = OwnerName;
+#endif
+	Initializer.IndexBuffer = RestResource->RestIndexBuffer.IndexBufferRHI;
+	Initializer.IndexBufferOffset = 0;
+	Initializer.GeometryType = RTGT_Triangles;
+	Initializer.TotalPrimitiveCount = RestResource->GetPrimitiveCount();
+	Initializer.bFastBuild = true;
+	Initializer.bAllowUpdate = true;
+
+	FBufferRHIRef PositionBuffer = RestResource->RestPositionBuffer.VertexBufferRHI;
+	if (DeformedResource)
+	{
+		PositionBuffer = DeformedResource->GetBuffer(FHairCardsDeformedResource::Current).Buffer->GetRHI(); // This will likely flicker result in half speed motion as everyother frame will use the wrong buffer
+	}
+
+	FRayTracingGeometrySegment Segment;
+	Segment.VertexBuffer = PositionBuffer;
+	Segment.VertexBufferStride = FHairCardsPositionFormat::SizeInByte;
+	Segment.VertexBufferElementType = FHairCardsPositionFormat::VertexElementType;
+	Segment.NumPrimitives = RestResource->GetPrimitiveCount();
+	Segment.MaxVertices = RestResource->GetVertexCount();
+	Initializer.Segments.Add(Segment);
+
+	OutRayTracingGeometry->SetInitializer(Initializer);
+	OutRayTracingGeometry->CreateRayTracingGeometry(ERTAccelerationStructureBuildPriority::Immediate);
+}
+
+static void BuildHairAccelerationStructure_Meshes(FRHICommandList& RHICmdList,
+	FHairMeshesRestResource* RestResource,
+	FHairMeshesDeformedResource* DeformedResource,
+	FRayTracingGeometry* OutRayTracingGeometry,
+	const FString& AssetDebugName,
+	uint32 LODIndex)
+{
+	FRayTracingGeometryInitializer Initializer;
+#if HAIR_RESOURCE_DEBUG_NAME
+	FString DebugName(TEXT("AS_HairMeshes_LOD") + FString::FromInt(LODIndex) + TEXT("_") + AssetDebugName);
+	Initializer.DebugName = FName(DebugName);
+#else
+	static const FName DebugName("AS_HairMeshes");
+	static int32 DebugNumber = 0;
+	Initializer.DebugName = FDebugName(DebugName, DebugNumber++);
+#endif
+	Initializer.IndexBuffer = RestResource->IndexBuffer.IndexBufferRHI;
+	Initializer.IndexBufferOffset = 0;
+	Initializer.GeometryType = RTGT_Triangles;
+	Initializer.TotalPrimitiveCount = RestResource->GetPrimitiveCount();
+	Initializer.bFastBuild = true;
+	Initializer.bAllowUpdate = true;
+
+	FBufferRHIRef PositionBuffer = RestResource->RestPositionBuffer.VertexBufferRHI;
+	if (DeformedResource)
+	{	
+		PositionBuffer = DeformedResource->GetBuffer(FHairMeshesDeformedResource::Current).Buffer->GetRHI(); // This will likely flicker result in half speed motion as everyother frame will use the wrong buffer
+	}
+
+	FRayTracingGeometrySegment Segment;
+	Segment.VertexBuffer = PositionBuffer;
+	Segment.VertexBufferStride = FHairCardsPositionFormat::SizeInByte;
+	Segment.VertexBufferElementType = FHairCardsPositionFormat::VertexElementType;
+	Segment.NumPrimitives = RestResource->GetPrimitiveCount();
+	Segment.MaxVertices = RestResource->GetVertexCount();
+	Initializer.Segments.Add(Segment);
+
+	OutRayTracingGeometry->SetInitializer(Initializer);
+	OutRayTracingGeometry->CreateRayTracingGeometry(ERTAccelerationStructureBuildPriority::Immediate);
+}
+#endif // RHI_RAYTRACING
+
+
+void RegisterClusterData(FHairGroupInstance* Instance, FHairStrandClusterData* InClusterData)
+{
+	// Initialize group cluster data for culling by the renderer
+	const int32 ClusterDataGroupIndex = InClusterData->HairGroups.Num();
+	FHairStrandClusterData::FHairGroup& HairGroupCluster = InClusterData->HairGroups.Emplace_GetRef();
+	HairGroupCluster.ClusterCount = Instance->HairGroupPublicData->GetClusterCount();
+	HairGroupCluster.VertexCount = Instance->HairGroupPublicData->RestPointCount * HAIR_POINT_TO_VERTEX; // Instance->HairGroupPublicData->GetActiveStrandsPointCount() ?
+	HairGroupCluster.GroupAABBBuffer = &Instance->HairGroupPublicData->GetGroupAABBBuffer();
+	HairGroupCluster.ClusterAABBBuffer = &Instance->HairGroupPublicData->GetClusterAABBBuffer();
+
+	HairGroupCluster.ClusterInfoBuffer = &Instance->Strands.ClusterCullingResource->ClusterInfoBuffer;
+	HairGroupCluster.ClusterLODInfoBuffer = &Instance->Strands.ClusterCullingResource->ClusterLODInfoBuffer;
+	HairGroupCluster.VertexToClusterIdBuffer = &Instance->Strands.ClusterCullingResource->VertexToClusterIdBuffer;
+	HairGroupCluster.ClusterVertexIdBuffer = &Instance->Strands.ClusterCullingResource->ClusterVertexIdBuffer;
+
+	HairGroupCluster.HairGroupPublicPtr = Instance->HairGroupPublicData;
+	HairGroupCluster.LODBias  = Instance->HairGroupPublicData->GetLODBias();
+	HairGroupCluster.LODIndex = Instance->HairGroupPublicData->GetLODIndex();
+	HairGroupCluster.bVisible = Instance->HairGroupPublicData->GetLODVisibility();
+
+	// These buffer are create during the culling pass
+	// HairGroupCluster.ClusterIdBuffer = nullptr;
+	// HairGroupCluster.ClusterIndexOffsetBuffer = nullptr;
+	// HairGroupCluster.ClusterIndexCountBuffer = nullptr;
+
+	HairGroupCluster.HairGroupPublicPtr->ClusterDataIndex = ClusterDataGroupIndex;
+}
+
+static void ConvertHairStrandsVFParameters(
+	FRDGBuilder* GraphBuilder,
+	FRDGImportedBuffer& OutBuffer,
+	FRDGExternalBuffer& OutExternalBuffer,
+	FShaderResourceViewRHIRef& OutBufferRHISRV,
+	const FRDGExternalBuffer& ExternalBuffer)
+{
+	if (GraphBuilder)
+	{
+		OutBuffer = Register(*GraphBuilder, ExternalBuffer, ERDGImportedBufferFlags::CreateSRV);
+	}
+	OutBufferRHISRV = ExternalBuffer.SRV;
+	OutExternalBuffer = ExternalBuffer;
+}
+#define CONVERT_HAIRSSTRANDS_VF_PARAMETERS(OutName, ExternalBuffer) ConvertHairStrandsVFParameters(GraphBuilder, OutName, OutName##External, OutName##RHISRV, ExternalBuffer);
+
+// Compute/Update the hair strands description which will be used for rendering (VF) / voxelization & co.
+static FHairGroupPublicData::FVertexFactoryInput InternalComputeHairStrandsVertexInputData(FRDGBuilder* GraphBuilder, const FHairGroupInstance* Instance, EGroomViewMode ViewMode)
+{
+	FHairGroupPublicData::FVertexFactoryInput OutVFInput;
+	if (!Instance || Instance->GeometryType != EHairGeometryType::Strands)
+	{
+		const FUintVector4 InvalidOffset = FUintVector4(HAIR_ATTRIBUTE_INVALID_OFFSET,HAIR_ATTRIBUTE_INVALID_OFFSET,HAIR_ATTRIBUTE_INVALID_OFFSET,HAIR_ATTRIBUTE_INVALID_OFFSET);
+		for (uint32 OffsetIt = 0; OffsetIt < HAIR_CURVE_ATTRIBUTE_OFFSET_COUNT; ++OffsetIt)
+		{
+			OutVFInput.Strands.Common.CurveAttributeOffsets[OffsetIt] = InvalidOffset;
+		}
+		for (uint32 OffsetIt = 0; OffsetIt < HAIR_POINT_ATTRIBUTE_OFFSET_COUNT; ++OffsetIt)
+		{
+			OutVFInput.Strands.Common.PointAttributeOffsets[OffsetIt] = InvalidOffset;
+		}
+		return OutVFInput;
+	}
+
+	const bool bSupportHasGuides = Instance->Guides.IsValid();
+	const bool bSupportDeformation = Instance->Strands.DeformedResource != nullptr;
+	bool bUseGuideAttributeOffsets = false;
+
+	// 1. Guides
+	if (ViewMode == EGroomViewMode::SimHairStrands && bSupportHasGuides)
+	{
+		CONVERT_HAIRSSTRANDS_VF_PARAMETERS(OutVFInput.Strands.PositionBuffer,			Instance->Guides.DeformedResource->GetBuffer(FHairStrandsDeformedResource::EFrameType::Current));
+		CONVERT_HAIRSSTRANDS_VF_PARAMETERS(OutVFInput.Strands.PrevPositionBuffer,		Instance->Guides.DeformedResource->GetBuffer(FHairStrandsDeformedResource::EFrameType::Previous));
+		CONVERT_HAIRSSTRANDS_VF_PARAMETERS(OutVFInput.Strands.TangentBuffer,			Instance->Guides.DeformedResource->TangentBuffer);
+		CONVERT_HAIRSSTRANDS_VF_PARAMETERS(OutVFInput.Strands.CurveAttributeBuffer,		Instance->Guides.RestResource->CurveAttributeBuffer);
+		CONVERT_HAIRSSTRANDS_VF_PARAMETERS(OutVFInput.Strands.PointAttributeBuffer,		Instance->Guides.RestResource->PointAttributeBuffer);
+		CONVERT_HAIRSSTRANDS_VF_PARAMETERS(OutVFInput.Strands.CurveBuffer,				Instance->Guides.RestResource->CurveBuffer);
+		CONVERT_HAIRSSTRANDS_VF_PARAMETERS(OutVFInput.Strands.PointToCurveBuffer,		Instance->Guides.RestResource->PointToCurveBuffer);
+		CONVERT_HAIRSSTRANDS_VF_PARAMETERS(OutVFInput.Strands.PositionOffsetBuffer,		Instance->Guides.DeformedResource->GetPositionOffsetBuffer(FHairStrandsDeformedResource::EFrameType::Current));
+		CONVERT_HAIRSSTRANDS_VF_PARAMETERS(OutVFInput.Strands.PrevPositionOffsetBuffer,	Instance->Guides.DeformedResource->GetPositionOffsetBuffer(FHairStrandsDeformedResource::EFrameType::Previous));
+
+		OutVFInput.Strands.Common.PositionOffset = FVector3f(Instance->Guides.DeformedResource->GetPositionOffset(FHairStrandsDeformedResource::EFrameType::Current)); // LWC precision loss
+		OutVFInput.Strands.Common.PrevPositionOffset = FVector3f(Instance->Guides.DeformedResource->GetPositionOffset(FHairStrandsDeformedResource::EFrameType::Previous)); // LWC precision loss
+		OutVFInput.Strands.Common.PointCount = Instance->Guides.RestResource->GetPointCount();
+		OutVFInput.Strands.Common.CurveCount = Instance->Guides.RestResource->GetCurveCount();
+		bUseGuideAttributeOffsets = true;
+	}
+	// 2. Render Strands with deformation
+	else if (bSupportDeformation)
+	{
+		const bool bHasValidMotionVector =
+			Instance->Strands.DeformedResource->GetUniqueViewID(FHairStrandsDeformedResource::Current) ==
+			Instance->Strands.DeformedResource->GetUniqueViewID(FHairStrandsDeformedResource::Previous);
+		if (bHasValidMotionVector)
+		{
+			CONVERT_HAIRSSTRANDS_VF_PARAMETERS(OutVFInput.Strands.PositionBuffer,			Instance->Strands.DeformedResource->GetBuffer(FHairStrandsDeformedResource::EFrameType::Current));
+			CONVERT_HAIRSSTRANDS_VF_PARAMETERS(OutVFInput.Strands.PrevPositionBuffer,		Instance->Strands.DeformedResource->GetBuffer(FHairStrandsDeformedResource::EFrameType::Previous));
+			CONVERT_HAIRSSTRANDS_VF_PARAMETERS(OutVFInput.Strands.PositionOffsetBuffer,		Instance->Strands.DeformedResource->GetPositionOffsetBuffer(FHairStrandsDeformedResource::EFrameType::Current));
+			CONVERT_HAIRSSTRANDS_VF_PARAMETERS(OutVFInput.Strands.PrevPositionOffsetBuffer, Instance->Strands.DeformedResource->GetPositionOffsetBuffer(FHairStrandsDeformedResource::EFrameType::Previous));
+
+			OutVFInput.Strands.Common.PositionOffset	 = FVector3f(Instance->Strands.DeformedResource->GetPositionOffset(FHairStrandsDeformedResource::EFrameType::Current)); // LWC precision loss
+			OutVFInput.Strands.Common.PrevPositionOffset = FVector3f(Instance->Strands.DeformedResource->GetPositionOffset(FHairStrandsDeformedResource::EFrameType::Previous)); // LWC precision loss
+		}
+		else
+		{
+			CONVERT_HAIRSSTRANDS_VF_PARAMETERS(OutVFInput.Strands.PositionBuffer,			Instance->Strands.DeformedResource->GetBuffer(FHairStrandsDeformedResource::EFrameType::Current));
+			CONVERT_HAIRSSTRANDS_VF_PARAMETERS(OutVFInput.Strands.PrevPositionBuffer,		Instance->Strands.DeformedResource->GetBuffer(FHairStrandsDeformedResource::EFrameType::Current));
+			CONVERT_HAIRSSTRANDS_VF_PARAMETERS(OutVFInput.Strands.PositionOffsetBuffer,		Instance->Strands.DeformedResource->GetPositionOffsetBuffer(FHairStrandsDeformedResource::EFrameType::Current));
+			CONVERT_HAIRSSTRANDS_VF_PARAMETERS(OutVFInput.Strands.PrevPositionOffsetBuffer,	Instance->Strands.DeformedResource->GetPositionOffsetBuffer(FHairStrandsDeformedResource::EFrameType::Current));
+
+			OutVFInput.Strands.Common.PositionOffset	 = FVector3f(Instance->Strands.DeformedResource->GetPositionOffset(FHairStrandsDeformedResource::EFrameType::Current)); // LWC precision loss
+			OutVFInput.Strands.Common.PrevPositionOffset = FVector3f(Instance->Strands.DeformedResource->GetPositionOffset(FHairStrandsDeformedResource::EFrameType::Current)); // LWC precision loss
+		}
+		CONVERT_HAIRSSTRANDS_VF_PARAMETERS(OutVFInput.Strands.TangentBuffer,			Instance->Strands.DeformedResource->TangentBuffer);
+		CONVERT_HAIRSSTRANDS_VF_PARAMETERS(OutVFInput.Strands.CurveAttributeBuffer,		(ViewMode == EGroomViewMode::RenderHairStrands || ViewMode == EGroomViewMode::Cluster || ViewMode == EGroomViewMode::ClusterAABB) ? Instance->Strands.DebugCurveAttributeBuffer : Instance->Strands.RestResource->CurveAttributeBuffer);
+		CONVERT_HAIRSSTRANDS_VF_PARAMETERS(OutVFInput.Strands.PointAttributeBuffer,		Instance->Strands.RestResource->PointAttributeBuffer);
+		CONVERT_HAIRSSTRANDS_VF_PARAMETERS(OutVFInput.Strands.PointToCurveBuffer,		Instance->Strands.RestResource->PointToCurveBuffer);
+		CONVERT_HAIRSSTRANDS_VF_PARAMETERS(OutVFInput.Strands.CurveBuffer,				Instance->Strands.RestResource->CurveBuffer);
+
+		OutVFInput.Strands.Common.PointCount = Instance->HairGroupPublicData->GetActiveStrandsPointCount();
+		OutVFInput.Strands.Common.CurveCount = Instance->HairGroupPublicData->GetActiveStrandsCurveCount();
+	}
+	// 3. Render Strands in rest position: used when there are no skinning (rigid binding or no binding), no simulation, and no RBF
+	else
+	{
+		CONVERT_HAIRSSTRANDS_VF_PARAMETERS(OutVFInput.Strands.PositionBuffer,			Instance->Strands.RestResource->PositionBuffer);
+		CONVERT_HAIRSSTRANDS_VF_PARAMETERS(OutVFInput.Strands.PrevPositionBuffer,		Instance->Strands.RestResource->PositionBuffer);
+		CONVERT_HAIRSSTRANDS_VF_PARAMETERS(OutVFInput.Strands.TangentBuffer,			Instance->Strands.RestResource->TangentBuffer);
+		CONVERT_HAIRSSTRANDS_VF_PARAMETERS(OutVFInput.Strands.CurveAttributeBuffer,		(ViewMode == EGroomViewMode::Cluster || ViewMode == EGroomViewMode::ClusterAABB) ? Instance->Strands.DebugCurveAttributeBuffer : Instance->Strands.RestResource->CurveAttributeBuffer);
+		CONVERT_HAIRSSTRANDS_VF_PARAMETERS(OutVFInput.Strands.PointAttributeBuffer,		Instance->Strands.RestResource->PointAttributeBuffer);
+		CONVERT_HAIRSSTRANDS_VF_PARAMETERS(OutVFInput.Strands.PointToCurveBuffer,		Instance->Strands.RestResource->PointToCurveBuffer);
+		CONVERT_HAIRSSTRANDS_VF_PARAMETERS(OutVFInput.Strands.CurveBuffer,				Instance->Strands.RestResource->CurveBuffer);
+		CONVERT_HAIRSSTRANDS_VF_PARAMETERS(OutVFInput.Strands.PositionOffsetBuffer,		Instance->Strands.RestResource->PositionOffsetBuffer);
+		CONVERT_HAIRSSTRANDS_VF_PARAMETERS(OutVFInput.Strands.PrevPositionOffsetBuffer, Instance->Strands.RestResource->PositionOffsetBuffer);
+
+		OutVFInput.Strands.Common.PositionOffset 	= FVector3f(Instance->Strands.RestResource->GetPositionOffset()); // LWC Precision loss
+		OutVFInput.Strands.Common.PrevPositionOffset= FVector3f(Instance->Strands.RestResource->GetPositionOffset()); // LWC Precision loss
+		OutVFInput.Strands.Common.PointCount 		= Instance->HairGroupPublicData->GetActiveStrandsPointCount();
+		OutVFInput.Strands.Common.CurveCount 		= Instance->HairGroupPublicData->GetActiveStrandsCurveCount();
+	}
+
+	OutVFInput.Strands.Common.Radius = (GStrandHairWidth > 0 ? GStrandHairWidth : Instance->Strands.Modifier.HairWidth) * 0.5f;
+	OutVFInput.Strands.Common.RootScale = Instance->Strands.Modifier.HairRootScale;
+	OutVFInput.Strands.Common.TipScale = Instance->Strands.Modifier.HairTipScale;
+	OutVFInput.Strands.Common.RaytracingRadiusScale = (GHairRaytracingRadiusScale > 0 ? GHairRaytracingRadiusScale : Instance->Strands.Modifier.HairRaytracingRadiusScale);
+	OutVFInput.Strands.Common.Length = Instance->Strands.Modifier.HairLength;
+	OutVFInput.Strands.Common.Density = Instance->Strands.Modifier.HairShadowDensity;
+	OutVFInput.Strands.Common.bScatterSceneLighting = Instance->Strands.Modifier.bScatterSceneLighting;
+	OutVFInput.Strands.Common.bStableRasterization = Instance->Strands.Modifier.bUseStableRasterization;
+	OutVFInput.Strands.Common.bRaytracingGeometry = false;
+	OutVFInput.Strands.Common.RaytracingProceduralSplits = GetHairRaytracingProceduralSplits();
+
+	// Add local transform assignement here ...
+
+	if (bUseGuideAttributeOffsets)
+	{
+		PACK_CURVE_HAIR_ATTRIBUTE_OFFSETS(OutVFInput.Strands.Common.CurveAttributeOffsets, Instance->Guides.Data->Header.CurveAttributeOffsets);
+		PACK_POINT_HAIR_ATTRIBUTE_OFFSETS(OutVFInput.Strands.Common.PointAttributeOffsets, Instance->Guides.Data->Header.PointAttributeOffsets);
+	}
+	else
+	{
+		PACK_CURVE_HAIR_ATTRIBUTE_OFFSETS(OutVFInput.Strands.Common.CurveAttributeOffsets, Instance->Strands.Data->Header.CurveAttributeOffsets);
+		PACK_POINT_HAIR_ATTRIBUTE_OFFSETS(OutVFInput.Strands.Common.PointAttributeOffsets, Instance->Strands.Data->Header.PointAttributeOffsets);
+	}
+
+#if RHI_RAYTRACING
+	// Flag bUseRaytracingGeometry only if RT geometry has been allocated for RayTracing view (not for PathTracing view). 
+	// This flag is used later for selecting if voxelization needs to flags voxel has shadow caster or if shadow casting is handled 
+	// by the RT geometry.
+	OutVFInput.Strands.Common.bRaytracingGeometry = Instance->Strands.RenRaytracingResource != nullptr && (Instance->Strands.ViewRayTracingMask & EHairViewRayTracingMask::RayTracing) != 0;
+#endif
+
+	return OutVFInput;
+}
+
+FHairGroupPublicData::FVertexFactoryInput ComputeHairStrandsVertexInputData(const FHairGroupInstance* Instance, EGroomViewMode ViewMode)
+{
+	return InternalComputeHairStrandsVertexInputData(nullptr, Instance, ViewMode);
+}
+
+void CreateHairStrandsDebugAttributeBuffer(FRDGBuilder& GraphBuilder, FRDGExternalBuffer* DebugAttributeBuffer, uint32 VertexCount, const FName& OwnerName);
+
+EGroomCacheType GetHairInstanceCacheType(const FHairGroupInstance* Instance)
+{
+	const bool bHasValidGroomCacheBuffers = Instance->Debug.GroomCacheBuffers.IsValid() && 
+		Instance->Debug.GroomCacheBuffers->GetCurrentFrameBuffer().GroupsData.IsValidIndex(Instance->Debug.GroupIndex) && 
+		Instance->Debug.GroomCacheBuffers->GetNextFrameBuffer().GroupsData.IsValidIndex(Instance->Debug.GroupIndex);
+	return bHasValidGroomCacheBuffers ? Instance->Debug.GroomCacheType : EGroomCacheType::None;
+}
+
+void ComputeHairStrandsInterpolation(
+	FRDGBuilder& GraphBuilder,
+	FGlobalShaderMap* ShaderMap,
+	const uint32 ViewUniqueID,
+	const uint32 ViewRayTracingMask,
+	const EGroomViewMode ViewMode,
+	const FVector& TranslatedWorldOffset,
+	const FShaderPrintData* ShaderPrintData,
+	FHairGroupInstance* Instance,
+	int32 MeshLODIndex,
+	FHairStrandClusterData* InClusterData)
+{	
+	if (!Instance)
+	{
+		return;
+	}
+
+	// Reset
+	Instance->HairGroupPublicData->VFInput.Strands	= FHairGroupPublicData::FVertexFactoryInput::FStrands();
+	Instance->HairGroupPublicData->VFInput.Cards	= FHairGroupPublicData::FVertexFactoryInput::FCards();
+	Instance->HairGroupPublicData->VFInput.Meshes	= FHairGroupPublicData::FVertexFactoryInput::FMeshes();
+	const EHairGeometryType InstanceGeometryType	= Instance->GeometryType;
+
+	FRDGExternalAccessQueue ExternalAccessQueue;
+
+	if (InstanceGeometryType == EHairGeometryType::Strands)
+	{
+		DECLARE_GPU_STAT(HairStrandsInterpolation);
+		RDG_EVENT_SCOPE(GraphBuilder, "HairInterpolation(Strands)");
+		RDG_GPU_STAT_SCOPE(GraphBuilder, HairStrandsInterpolation);
+
+		const EGroomCacheType ActiveGroomCacheType = GetHairInstanceCacheType(Instance);
+
+		// Debug mode:
+		// * None	: Display hair normally
+		// * Sim	: Show sim strands
+		// * Render : Show rendering strands with sim color influence
+		if (ViewMode == EGroomViewMode::SimHairStrands && Instance->Guides.DeformedResource)
+		{
+			// Disable culling when drawing only guides, as the culling output has been computed for the strands, not for the guides.
+			Instance->HairGroupPublicData->SetCullingResultAvailable(false);
+
+			// If groom guide cache, is enabled, used the deform position for debug visualization
+			if (ActiveGroomCacheType == EGroomCacheType::Guides)
+			{
+				FScopeLock Lock(Instance->Debug.GroomCacheBuffers->GetCriticalSection());
+				FGroomCacheGroupData* GroomCacheData0 = const_cast<FGroomCacheGroupData*>(&Instance->Debug.GroomCacheBuffers->GetCurrentFrameBuffer().GroupsData[Instance->Debug.GroupIndex]);
+				FGroomCacheGroupData* GroomCacheData1 = const_cast<FGroomCacheGroupData*>(&Instance->Debug.GroomCacheBuffers->GetNextFrameBuffer().GroupsData[Instance->Debug.GroupIndex]);
+
+				const float InterpolationFactor = Instance->Debug.GroomCacheBuffers->GetInterpolationFactor();
+				Instance->Guides.DeformedResource->SetPositionOffset(FHairStrandsDeformedResource::EFrameType::Current, FMath::Lerp(GroomCacheData0->BoundingBox.GetCenter(), GroomCacheData1->BoundingBox.GetCenter(), InterpolationFactor));
+
+				FGroomCacheResources CacheResources0 = CreateGroomCacheBuffer(GraphBuilder, GroomCacheData0->VertexData);
+				FGroomCacheResources CacheResources1 = CreateGroomCacheBuffer(GraphBuilder, GroomCacheData1->VertexData);
+
+				AddHairStrandUpdatePositionOffsetPass(
+					GraphBuilder,
+					ShaderMap,
+					MeshLODIndex,
+					Instance->Guides.DeformedRootResource,
+					Instance->Guides.DeformedResource);
+
+				// Pass to upload GroomCache guide positions
+				AddGroomCacheUpdatePass(
+					GraphBuilder,
+					ShaderMap,
+					Instance->Guides.RestResource->GetPointCount(),
+					InterpolationFactor,
+					CacheResources0,
+					CacheResources1,
+					RegisterAsSRV(GraphBuilder, Instance->Guides.RestResource->PositionBuffer),
+					RegisterAsSRV(GraphBuilder, Instance->Guides.DeformedResource->GetPositionOffsetBuffer(FHairStrandsDeformedResource::EFrameType::Current)),
+					RegisterAsUAV(GraphBuilder, Instance->Guides.DeformedResource->GetBuffer(FHairStrandsDeformedResource::Current)));
+			}
+
+			AddHairTangentPass(
+				GraphBuilder,
+				ShaderMap,
+				Instance->Guides.RestResource->GetPointCount(),
+				Instance->HairGroupPublicData,
+				RegisterAsSRV(GraphBuilder, Instance->Guides.DeformedResource->GetBuffer(FHairStrandsDeformedResource::Current)),
+				Register(GraphBuilder, Instance->Guides.DeformedResource->TangentBuffer, ERDGImportedBufferFlags::CreateUAV));
+
+			Instance->HairGroupPublicData->VFInput = InternalComputeHairStrandsVertexInputData(&GraphBuilder, Instance, ViewMode);
+		}
+		else
+		{
+			const bool bNeedDeformation = Instance->Strands.DeformedResource != nullptr;
+			const EHairAABBUpdateType UpdateType = bNeedDeformation ? EHairAABBUpdateType::UpdateClusterAABB : EHairAABBUpdateType::UpdateGroupAABB;
+
+			// 1. Clear cluster AABBs (used optionally  for voxel allocation & for culling)
+			FRDGHairStrandsCullingData CullingData = ImportCullingData(GraphBuilder, Instance->HairGroupPublicData);
+			{
+				AddClearClusterAABBPass(
+					GraphBuilder,
+					ShaderMap,
+					UpdateType,
+					CullingData.ClusterCount,
+					CullingData.ClusterAABBBuffer,
+					CullingData.GroupAABBBuffer);
+			}
+
+			// "WITH_EDITOR && PatchMode == EHairPatchAttribute::GuideInflucence" is a special path when visualizing groom within the groom editor, so that we can diplay guides even when there is no simulation or global interpolation enabled
+			EHairPatchAttribute PatchMode = EHairPatchAttribute::None;
+			#if WITH_EDITOR
+			if (ViewMode == EGroomViewMode::RenderHairStrands && bNeedDeformation)				{ PatchMode = EHairPatchAttribute::GuideInflucence; }
+			if (ViewMode == EGroomViewMode::Cluster || ViewMode == EGroomViewMode::ClusterAABB)	{ PatchMode = EHairPatchAttribute::ClusterInfluence; }
+			#endif
+
+			// 2. Deform hair if needed (e.g.: skinning, simulation, RBF) and recompute tangent
+			FRDGBufferSRVRef Strands_PositionSRV = nullptr;
+			FRDGBufferSRVRef Strands_PositionOffsetSRV = nullptr;
+			FRDGBufferSRVRef Strands_PrevPositionOffsetSRV = nullptr;
+			FRDGBufferSRVRef Strands_TangentSRV = nullptr;
+
+			const bool bValidGuide		= bNeedDeformation && (Instance->Guides.bIsSimulationEnable || Instance->Guides.bHasGlobalInterpolation || Instance->Guides.bIsDeformationEnable || Instance->Guides.bIsSimulationCacheEnable);// || (WITH_EDITOR && PatchMode == EHairPatchAttribute::GuideInflucence);
+			const bool bUseSingleGuide	= bNeedDeformation && bValidGuide && Instance->Strands.InterpolationResource->UseSingleGuide();
+			const bool bHasSkinning		= bNeedDeformation && Instance->BindingType == EHairBindingType::Skinning;
+
+			const uint32 ActivePointCount = Instance->HairGroupPublicData->GetActiveStrandsPointCount();
+			const uint32 ActiveCurveCount = Instance->HairGroupPublicData->GetActiveStrandsCurveCount();
+
+			if (bNeedDeformation)
+			{
+				FRDGImportedBuffer Strands_DeformedPosition = Register(GraphBuilder, Instance->Strands.DeformedResource->GetBuffer(FHairStrandsDeformedResource::Current), ERDGImportedBufferFlags::CreateViews);
+				FRDGImportedBuffer Strands_DeformedPrevPosition = Register(GraphBuilder, Instance->Strands.DeformedResource->GetBuffer(FHairStrandsDeformedResource::Previous), ERDGImportedBufferFlags::CreateViews);
+				FRDGImportedBuffer Strands_DeformedTangent = Register(GraphBuilder, Instance->Strands.DeformedResource->TangentBuffer, ERDGImportedBufferFlags::CreateViews);
+				FRDGBufferSRVRef Strands_DeformerPositionSRV = Register(GraphBuilder, Instance->Strands.DeformedResource->DeformerBuffer, ERDGImportedBufferFlags::CreateSRV).SRV;
+
+				// Trach on which view the position has been update, so that we can ensure motion vector are coherent
+				Instance->Strands.DeformedResource->GetUniqueViewID(FHairStrandsDeformedResource::Current) = ViewUniqueID;
+
+				Strands_PositionSRV = Strands_DeformedPosition.SRV;
+				Strands_PositionOffsetSRV = RegisterAsSRV(GraphBuilder, Instance->Strands.DeformedResource->GetPositionOffsetBuffer(FHairStrandsDeformedResource::EFrameType::Current));
+				Strands_TangentSRV = Strands_DeformedTangent.SRV;
+
+				if (ActiveGroomCacheType == EGroomCacheType::Guides)
+				{
+					FScopeLock Lock(Instance->Debug.GroomCacheBuffers->GetCriticalSection()); // This is not ideally it will block the rendering thread / game thread
+
+					FGroomCacheGroupData* GroomCacheData0 = const_cast<FGroomCacheGroupData*>(&Instance->Debug.GroomCacheBuffers->GetCurrentFrameBuffer().GroupsData[Instance->Debug.GroupIndex]);
+					FGroomCacheGroupData* GroomCacheData1 = const_cast<FGroomCacheGroupData*>(&Instance->Debug.GroomCacheBuffers->GetNextFrameBuffer().GroupsData[Instance->Debug.GroupIndex]);
+					const float InterpolationFactor = Instance->Debug.GroomCacheBuffers->GetInterpolationFactor();
+					const FVector OffsetPosition = FMath::Lerp(GroomCacheData0->BoundingBox.GetCenter(), GroomCacheData1->BoundingBox.GetCenter(), InterpolationFactor);
+
+					FGroomCacheResources CacheResources0 = CreateGroomCacheBuffer(GraphBuilder, GroomCacheData0->VertexData);
+					FGroomCacheResources CacheResources1 = CreateGroomCacheBuffer(GraphBuilder, GroomCacheData1->VertexData);
+
+					Instance->Guides.DeformedResource->SetPositionOffset(FHairStrandsDeformedResource::EFrameType::Current, OffsetPosition);
+
+					AddHairStrandUpdatePositionOffsetPass(
+						GraphBuilder,
+						ShaderMap,
+						MeshLODIndex,
+						Instance->Guides.DeformedRootResource,
+						Instance->Guides.DeformedResource);
+
+					// Apply the same offset to the render strands for proper rendering
+					Instance->Strands.DeformedResource->SetPositionOffset(FHairStrandsDeformedResource::EFrameType::Current, OffsetPosition);
+
+					AddHairStrandUpdatePositionOffsetPass(
+						GraphBuilder,
+						ShaderMap,
+						MeshLODIndex,
+						Instance->Strands.DeformedRootResource,
+						Instance->Strands.DeformedResource);
+
+					// Pass to upload GroomCache guide positions
+					AddGroomCacheUpdatePass(
+						GraphBuilder,
+						ShaderMap,
+						Instance->Guides.RestResource->GetPointCount(),
+						InterpolationFactor,
+						CacheResources0,
+						CacheResources1,
+						RegisterAsSRV(GraphBuilder, Instance->Guides.RestResource->PositionBuffer),
+						RegisterAsSRV(GraphBuilder, Instance->Guides.DeformedResource->GetPositionOffsetBuffer(FHairStrandsDeformedResource::EFrameType::Current)),
+						RegisterAsUAV(GraphBuilder, Instance->Guides.DeformedResource->GetBuffer(FHairStrandsDeformedResource::Current)));
+				}
+				else if (Instance->Guides.bIsDeformationEnable && Instance->Guides.DeformedResource && Instance->DeformedComponent && (Instance->DeformedSection != INDEX_NONE))
+				{
+					if (USkeletalMeshComponent* SkeletalMeshComponent = Cast<USkeletalMeshComponent>(Instance->DeformedComponent))
+					{
+						if(FSkeletalMeshObject* SkeletalMeshObject = SkeletalMeshComponent->MeshObject)
+						{
+							const int32 LodIndex = SkeletalMeshObject->GetLOD();
+							FSkeletalMeshLODRenderData const* LodRenderData = &SkeletalMeshObject->GetSkeletalMeshRenderData().LODRenderData[LodIndex];
+				
+							if(LodRenderData->RenderSections.Num() > Instance->DeformedSection) 
+							{
+								FRHIShaderResourceView* BoneBufferSRV = FSkeletalMeshDeformerHelpers::GetBoneBufferForReading(SkeletalMeshObject, LodIndex, Instance->DeformedSection, false);
+
+								// Guides deformation based on the skeletal mesh bones
+								AddDeformSimHairStrandsPass(
+									GraphBuilder,
+									ShaderMap,
+									EDeformationType::OffsetGuide,
+									MeshLODIndex,
+									Instance->Guides.RestResource->GetPointCount(),
+									Instance->Guides.RestRootResource,
+									Instance->Guides.DeformedRootResource,
+									RegisterAsSRV(GraphBuilder, Instance->Guides.RestResource->PositionBuffer),
+									RegisterAsSRV(GraphBuilder, Instance->Guides.RestResource->PointToCurveBuffer),
+									Register(GraphBuilder, Instance->Guides.DeformedResource->GetBuffer(FHairStrandsDeformedResource::Current), ERDGImportedBufferFlags::CreateUAV),
+									Instance->Guides.RestResource->GetPositionOffset(),
+									RegisterAsSRV(GraphBuilder, Instance->Guides.DeformedResource->GetPositionOffsetBuffer(FHairStrandsDeformedResource::Current)),
+									Instance->Guides.bHasGlobalInterpolation,
+									BoneBufferSRV);
+							}
+						}
+					}
+				}
+
+				// 2.1 Compute deformation position based on simulation/skinning/RBF
+				if (ActiveGroomCacheType != EGroomCacheType::Strands)
+				{
+					// 2.1.1 Current Position
+					AddHairStrandsInterpolationPass(
+						GraphBuilder,
+						ShaderMap,
+						ShaderPrintData,
+						Instance,
+						ActivePointCount,
+						MeshLODIndex,
+						Instance->Strands.Modifier.HairLengthScale,
+						Instance->Strands.HairInterpolationType,
+						InstanceGeometryType,
+						CullingData,
+						Instance->Strands.RestResource->GetPositionOffset(),
+						bValidGuide ? Instance->Guides.RestResource->GetPositionOffset() : FVector::ZeroVector,
+						Strands_PositionOffsetSRV,
+						bValidGuide ? RegisterAsSRV(GraphBuilder, Instance->Guides.DeformedResource->GetPositionOffsetBuffer(FHairStrandsDeformedResource::Current)) : nullptr,
+						bHasSkinning ? Instance->Strands.RestRootResource : nullptr,
+						bHasSkinning && bValidGuide ? Instance->Guides.RestRootResource : nullptr,
+						bHasSkinning ? Instance->Strands.DeformedRootResource : nullptr,
+						bHasSkinning && bValidGuide ? Instance->Guides.DeformedRootResource : nullptr,
+						RegisterAsSRV(GraphBuilder, Instance->Strands.RestResource->PositionBuffer),
+						RegisterAsSRV(GraphBuilder, Instance->Strands.RestResource->PointToCurveBuffer),
+						bUseSingleGuide,
+						bValidGuide ? RegisterAsSRV(GraphBuilder, Instance->Strands.InterpolationResource->InterpolationBuffer) : nullptr,
+						bValidGuide ? RegisterAsSRV(GraphBuilder, Instance->Guides.RestResource->PositionBuffer) : nullptr,
+						bValidGuide ? RegisterAsSRV(GraphBuilder, Instance->Guides.DeformedResource->GetBuffer(FHairStrandsDeformedResource::Current)) : nullptr,
+						bValidGuide ? RegisterAsSRV(GraphBuilder, Instance->Strands.InterpolationResource->SimRootPointIndexBuffer) : nullptr,
+						bValidGuide ? RegisterAsSRV(GraphBuilder, Instance->Guides.RestResource->PointToCurveBuffer) : nullptr,
+						Strands_DeformerPositionSRV,
+						Strands_DeformedPosition,
+						FHairStrandsDeformedRootResource::FLOD::Current);
+					
+					// 2.1.2 Previous Position
+					// * Transfer prev. position
+					// * Or recompute interpolation pass for previous positions if needed
+					const bool bRecomputePrevPosition = IsHairVisibilityComputeRasterContinuousLODEnabled();
+					const bool bTransferPrevPosition = 
+						!bRecomputePrevPosition &&
+						Instance->HairGroupPublicData->VFInput.bHasLODSwitch && 
+						GHairStrandsTransferPositionOnLODChange > 0;
+					if (bTransferPrevPosition)
+					{
+						AddTransferPositionPass(GraphBuilder, ShaderMap, ActivePointCount, Strands_DeformedPosition.SRV, Strands_DeformedPrevPosition.UAV);
+						GraphBuilder.SetBufferAccessFinal(Strands_DeformedPrevPosition.Buffer, ERHIAccess::SRVMask);
+					}
+					else if (bRecomputePrevPosition)
+					{
+						Strands_PrevPositionOffsetSRV = RegisterAsSRV(GraphBuilder, Instance->Strands.DeformedResource->GetPositionOffsetBuffer(FHairStrandsDeformedResource::EFrameType::Previous));
+
+						AddHairStrandsInterpolationPass(
+							GraphBuilder,
+							ShaderMap,
+							ShaderPrintData,
+							Instance,
+							ActivePointCount,
+							MeshLODIndex,
+							Instance->Strands.Modifier.HairLengthScale,
+							Instance->Strands.HairInterpolationType,
+							InstanceGeometryType,
+							CullingData,
+							Instance->Strands.RestResource->GetPositionOffset(),
+							bValidGuide ? Instance->Guides.RestResource->GetPositionOffset() : FVector::ZeroVector,
+							Strands_PrevPositionOffsetSRV,
+							bValidGuide ? RegisterAsSRV(GraphBuilder, Instance->Guides.DeformedResource->GetPositionOffsetBuffer(FHairStrandsDeformedResource::Previous)) : nullptr,
+							bHasSkinning ? Instance->Strands.RestRootResource : nullptr,
+							bHasSkinning && bValidGuide ? Instance->Guides.RestRootResource : nullptr,
+							bHasSkinning ? Instance->Strands.DeformedRootResource : nullptr,
+							bHasSkinning && bValidGuide ? Instance->Guides.DeformedRootResource : nullptr,
+							RegisterAsSRV(GraphBuilder, Instance->Strands.RestResource->PositionBuffer),
+							RegisterAsSRV(GraphBuilder, Instance->Strands.RestResource->PointToCurveBuffer),
+							bUseSingleGuide,
+							bValidGuide ? RegisterAsSRV(GraphBuilder, Instance->Strands.InterpolationResource->InterpolationBuffer) : nullptr,
+							bValidGuide ? RegisterAsSRV(GraphBuilder, Instance->Guides.RestResource->PositionBuffer) : nullptr,
+							bValidGuide ? RegisterAsSRV(GraphBuilder, Instance->Guides.DeformedResource->GetBuffer(FHairStrandsDeformedResource::Previous)) : nullptr,
+							bValidGuide ? RegisterAsSRV(GraphBuilder, Instance->Strands.InterpolationResource->SimRootPointIndexBuffer) : nullptr,
+							bValidGuide ? RegisterAsSRV(GraphBuilder, Instance->Guides.RestResource->PointToCurveBuffer) : nullptr,
+							Strands_DeformerPositionSRV,
+							Strands_DeformedPrevPosition,
+							FHairStrandsDeformedRootResource::FLOD::Previous);
+					}
+				}
+				else if (ActiveGroomCacheType == EGroomCacheType::Strands)
+				{
+					FScopeLock Lock(Instance->Debug.GroomCacheBuffers->GetCriticalSection());
+
+					FGroomCacheGroupData* GroomCacheData0 = const_cast<FGroomCacheGroupData*>(&Instance->Debug.GroomCacheBuffers->GetCurrentFrameBuffer().GroupsData[Instance->Debug.GroupIndex]);
+					FGroomCacheGroupData* GroomCacheData1 = const_cast<FGroomCacheGroupData*>(&Instance->Debug.GroomCacheBuffers->GetNextFrameBuffer().GroupsData[Instance->Debug.GroupIndex]);
+					const float InterpolationFactor = Instance->Debug.GroomCacheBuffers->GetInterpolationFactor();
+
+					FGroomCacheResources CacheResources0 = CreateGroomCacheBuffer(GraphBuilder, GroomCacheData0->VertexData);
+					FGroomCacheResources CacheResources1 = CreateGroomCacheBuffer(GraphBuilder, GroomCacheData1->VertexData);
+
+					// Update position offset
+					{
+						const FVector OffsetPosition = FMath::Lerp(GroomCacheData0->BoundingBox.GetCenter(), GroomCacheData1->BoundingBox.GetCenter(), InterpolationFactor);
+						Instance->Strands.DeformedResource->SetPositionOffset(FHairStrandsDeformedResource::EFrameType::Current, OffsetPosition);
+					}
+
+					// Update max radius
+					if (GroomCacheData0->VertexData.PointsRadius.Num() > 0)
+					{
+						Instance->Strands.Modifier.HairWidth = FMath::Lerp(GroomCacheData0->StrandData.MaxRadius, GroomCacheData1->StrandData.MaxRadius, InterpolationFactor) * 2.0f;
+					}
+
+					AddHairStrandUpdatePositionOffsetPass(
+						GraphBuilder,
+						ShaderMap,
+						MeshLODIndex,
+						Instance->Strands.DeformedRootResource,
+						Instance->Strands.DeformedResource);
+
+					// Pass to upload GroomCache strands positions
+					AddGroomCacheUpdatePass(
+						GraphBuilder,
+						ShaderMap,
+						ActivePointCount,
+						InterpolationFactor,
+						CacheResources0,
+						CacheResources1,
+						RegisterAsSRV(GraphBuilder, Instance->Strands.RestResource->PositionBuffer),
+						RegisterAsSRV(GraphBuilder, Instance->Strands.DeformedResource->GetPositionOffsetBuffer(FHairStrandsDeformedResource::EFrameType::Current)),
+						Strands_DeformedPosition.UAV);
+				}
+
+				// 2.2 Update tangent data based on the deformed positions
+				AddHairTangentPass(
+					GraphBuilder,
+					ShaderMap,
+					ActivePointCount,
+					Instance->HairGroupPublicData,
+					Strands_DeformedPosition.SRV,
+					Strands_DeformedTangent);
+			}
+			else
+			{
+				Strands_PositionSRV = RegisterAsSRV(GraphBuilder, Instance->Strands.RestResource->PositionBuffer);
+				Strands_PositionOffsetSRV = RegisterAsSRV(GraphBuilder, Instance->Strands.RestResource->PositionOffsetBuffer);
+				// Generated the static tangent if they haven;t been generated yet
+				Instance->Strands.RestResource->GetTangentBuffer(GraphBuilder, ShaderMap);
+				Strands_TangentSRV = RegisterAsSRV(GraphBuilder, Instance->Strands.RestResource->TangentBuffer);
+			}
+
+			// 2.1 Patch attribute for debug visualization (guide influence or clusters visualization)
+			#if WITH_EDITOR
+			if (PatchMode != EHairPatchAttribute::None)
+			{
+				// Create an debug buffer for storing cluster visualization data. This is only used for debug purpose, hence only enable in editor build.
+				// Special case for debug mode were the attribute buffer is patch with some custom data to show hair properties (strands belonging to the same cluster, ...)
+				if (Instance->Strands.DebugCurveAttributeBuffer.Buffer == nullptr)
+				{
+					CreateHairStrandsDebugAttributeBuffer(GraphBuilder, &Instance->Strands.DebugCurveAttributeBuffer, Instance->Strands.Data->Data.CurveAttributes.GetBulkDataSize(), FName(Instance->Debug.MeshComponentName));
+				}
+				FRDGImportedBuffer OutRenCurveAttributeBuffer = Register(GraphBuilder, Instance->Strands.DebugCurveAttributeBuffer, ERDGImportedBufferFlags::CreateUAV);
+
+				AddPatchAttributePass(
+					GraphBuilder,
+					ShaderMap,
+					ActiveCurveCount,
+					PatchMode,
+					bValidGuide,
+					bUseSingleGuide,
+					Instance->Strands.Data->Header.CurveAttributeOffsets,
+					Register(GraphBuilder, Instance->Strands.RestResource->CurveAttributeBuffer, ERDGImportedBufferFlags::CreateSRV).Buffer,
+					RegisterAsSRV(GraphBuilder, Instance->Strands.RestResource->CurveBuffer),
+					RegisterAsSRV(GraphBuilder, Instance->Strands.ClusterCullingResource->VertexToClusterIdBuffer),
+					bValidGuide ? RegisterAsSRV(GraphBuilder, Instance->Strands.InterpolationResource->InterpolationBuffer) : nullptr,
+					OutRenCurveAttributeBuffer);
+			}
+			#endif
+
+			// 2.2 Update the VF input with the update resources
+			Instance->HairGroupPublicData->VFInput = InternalComputeHairStrandsVertexInputData(&GraphBuilder, Instance, ViewMode);
+
+			// 3. Compute cluster AABBs (used for LODing and voxelization)
+			{
+				// Optim: If an instance does not voxelize it's data, then there is no need for having valid AABB
+				bool bNeedGPUAABB = Instance->Strands.Modifier.bSupportVoxelization && Instance->bCastShadow;
+
+				FHairStrandClusterData::FHairGroup* HairGroupCluster = nullptr;
+				if (CullingData.bCullingResultAvailable)
+				{
+					// Sanity check
+					check(Instance->Strands.bCullingEnable);
+					HairGroupCluster = &InClusterData->HairGroups[Instance->HairGroupPublicData->ClusterDataIndex];
+					if (!HairGroupCluster->bVisible)
+					{
+						bNeedGPUAABB = false;
+					}
+				}
+
+				// When CLOD is enabled, GPU AABB are not supported yet
+				if (IsHairStrandsContinousLODEnabled())
+				{
+					bNeedGPUAABB = false;
+				}
+				Instance->HairGroupPublicData->bClusterAABBValid = false;
+				Instance->HairGroupPublicData->bGroupAABBValid = false;
+
+
+				if (bNeedGPUAABB)
+				{
+					FRDGImportedBuffer Strands_CulledVertexCount = Register(GraphBuilder, Instance->HairGroupPublicData->GetDrawIndirectRasterComputeBuffer(), ERDGImportedBufferFlags::CreateSRV);
+					AddHairClusterAABBPass(
+						GraphBuilder,
+						ShaderMap,
+						TranslatedWorldOffset,
+						ShaderPrintData,
+						UpdateType,
+						Instance,
+						Strands_PositionOffsetSRV,
+						HairGroupCluster,
+						CullingData,
+						Strands_PositionSRV,
+						Strands_CulledVertexCount.SRV);
+
+					Instance->HairGroupPublicData->bClusterAABBValid = UpdateType == EHairAABBUpdateType::UpdateClusterAABB;
+					Instance->HairGroupPublicData->bGroupAABBValid = true;
+				}
+			}
+
+			// 4. Update raytracing geometry (update only if the view mask and the RT geometry mask match)
+			#if RHI_RAYTRACING
+			if (Instance->Strands.RenRaytracingResource && (Instance->Strands.ViewRayTracingMask & ViewRayTracingMask) != 0)
+			{
+				const float HairRadiusRT	= Instance->HairGroupPublicData->VFInput.Strands.Common.RaytracingRadiusScale * Instance->HairGroupPublicData->VFInput.Strands.Common.Radius;
+				const float HairRootScaleRT = Instance->HairGroupPublicData->VFInput.Strands.Common.RootScale;
+				const float HairTipScaleRT	= Instance->HairGroupPublicData->VFInput.Strands.Common.TipScale;
+				const int ProceduralSplits = GetHairRaytracingProceduralSplits();
+
+				const bool bNeedUpdate = Instance->Strands.DeformedResource != nullptr ||
+					Instance->Strands.CachedHairScaledRadius != HairRadiusRT ||
+					Instance->Strands.CachedHairRootScale != HairRootScaleRT ||
+					Instance->Strands.CachedHairTipScale != HairTipScaleRT ||
+					Instance->Strands.CachedProceduralSplits != ProceduralSplits;
+				const bool bNeedBuild = !Instance->Strands.RenRaytracingResource->bIsRTGeometryInitialized;
+				if (bNeedBuild || bNeedUpdate)
+				{
+					const uint32 HairLODIndex = Instance->HairGroupPublicData->GetIntLODIndex();
+					const bool bProceduralPrimitive = Instance->Strands.RenRaytracingResource->bProceduralPrimitive;
+
+					FRDGImportedBuffer Raytracing_PositionBuffer = Register(GraphBuilder, Instance->Strands.RenRaytracingResource->PositionBuffer, ERDGImportedBufferFlags::CreateViews);
+					FRDGImportedBuffer Raytracing_IndexBuffer = Register(GraphBuilder, Instance->Strands.RenRaytracingResource->IndexBuffer, ERDGImportedBufferFlags::CreateViews);
+					AddGenerateRaytracingGeometryPass(
+						GraphBuilder,
+						ShaderMap,
+						ShaderPrintData,
+						ActivePointCount,
+						bProceduralPrimitive,
+						ProceduralSplits,
+						HairRadiusRT,
+						HairRootScaleRT,
+						HairTipScaleRT,
+						Strands_PositionOffsetSRV,
+						CullingData,
+						Strands_PositionSRV,
+						Strands_TangentSRV,
+						Raytracing_PositionBuffer.UAV,
+						Raytracing_IndexBuffer.UAV);
+
+					Instance->Strands.CachedHairScaledRadius= HairRadiusRT;
+					Instance->Strands.CachedHairRootScale	= HairRootScaleRT;
+					Instance->Strands.CachedHairTipScale	= HairTipScaleRT;
+
+					FStrandsResourceBLASParameters* Parameters = GraphBuilder.AllocParameters<FStrandsResourceBLASParameters>();
+					Parameters->PositionBuffer = Raytracing_PositionBuffer.Buffer;
+					Parameters->IndexBuffer = Raytracing_IndexBuffer.Buffer;
+
+					GraphBuilder.AddPass(
+						RDG_EVENT_NAME("HairStrands::UpdateBLAS(Strands)"),
+						Parameters,
+						ERDGPassFlags::NeverCull | ERDGPassFlags::Compute,
+					[Instance, HairLODIndex, bNeedUpdate, bProceduralPrimitive, ProceduralSplits](FRHICommandListImmediate& RHICmdList)
+					{
+						SCOPED_GPU_MASK(RHICmdList, FRHIGPUMask::All());
+
+						const bool bLocalNeedBuild = !Instance->Strands.RenRaytracingResource->bIsRTGeometryInitialized || Instance->Strands.CachedProceduralSplits != ProceduralSplits;
+						if (bLocalNeedBuild)
+						{
+							FBufferRHIRef PositionBuffer(Instance->Strands.RenRaytracingResource->PositionBuffer.Buffer->GetRHI());
+
+							// no index buffer needed for procedural primitive
+							FBufferRHIRef IndexBuffer(bProceduralPrimitive ? nullptr : Instance->Strands.RenRaytracingResource->IndexBuffer.Buffer->GetRHI());
+							Instance->Strands.CachedProceduralSplits = ProceduralSplits;
+							BuildHairAccelerationStructure_Strands(RHICmdList,
+								Instance->Strands.RenRaytracingResource->VertexCount, // HAIR_TODO
+								Instance->Strands.RenRaytracingResource->IndexCount,
+								PositionBuffer,
+								IndexBuffer,
+								&Instance->Strands.RenRaytracingResource->RayTracingGeometry,
+								Instance->Debug.GroomAssetName,
+								Instance->Debug.MeshComponentName,
+								HairLODIndex,
+								bProceduralPrimitive,
+								ProceduralSplits
+							);
+						}
+						else if (bNeedUpdate)
+						{
+							// hair strands can move chaotically during simulation which can really tank ray tracing performance
+							// even though a rebuild is more expensive, we can still come out ahead overall, even with a single ray cast per pixel
+							UpdateHairAccelerationStructure(RHICmdList, &Instance->Strands.RenRaytracingResource->RayTracingGeometry, GHairStrands_Raytracing_ForceRebuildBVH > 0? EAccelerationStructureBuildMode::Build : EAccelerationStructureBuildMode::Update);
+						}
+						Instance->Strands.RenRaytracingResource->bIsRTGeometryInitialized = true;
+					});
+				}
+			}
+			#endif
+		}
+
+		// Sanity check
+		check(Instance->HairGroupPublicData->VFInput.Strands.PositionBuffer.Buffer);
+		Instance->Strands.UniformBuffer.UpdateUniformBufferImmediate(Instance->GetHairStandsUniformShaderParameters(ViewMode));
+	}
+	else if (InstanceGeometryType == EHairGeometryType::Cards)
+	{	
+		DECLARE_GPU_STAT(HairCardsInterpolation);
+		RDG_EVENT_SCOPE(GraphBuilder, "HairInterpolation(Cards)");
+		RDG_GPU_STAT_SCOPE(GraphBuilder, HairCardsInterpolation);
+
+		const uint32 HairLODIndex = Instance->HairGroupPublicData->GetIntLODIndex();
+		const bool bIsCardsValid = Instance->Cards.IsValid(HairLODIndex);
+		if (bIsCardsValid)
+		{
+			const bool bValidGuide = Instance->Guides.bIsSimulationEnable || Instance->Guides.bHasGlobalInterpolation || Instance->Guides.bIsDeformationEnable || Instance->Guides.bIsSimulationCacheEnable;
+			const bool bHasSkinning = Instance->BindingType == EHairBindingType::Skinning && MeshLODIndex >= 0;
+			const bool bNeedDeformation = bValidGuide || bHasSkinning;
+
+			FHairGroupInstance::FCards::FLOD& LOD = Instance->Cards.LODs[HairLODIndex];
+
+			if (bNeedDeformation)
+			{
+				check(LOD.Guides.Data);
+
+				const EHairCardsSimulationType CardsSimulationType = 
+					bHasSkinning || bValidGuide ?
+					GetHairCardsSimulationType() :
+					EHairCardsSimulationType::None;
+
+				// 1. Cards are deformed based on guides motion (simulation or RBF applied on guides)
+				if (CardsSimulationType == EHairCardsSimulationType::Guide)
+				{
+					FRDGImportedBuffer Guides_DeformedPositionBuffer = Register(GraphBuilder, LOD.Guides.DeformedResource->GetBuffer(FHairStrandsDeformedResource::Current), ERDGImportedBufferFlags::CreateViews);
+
+					const bool bUseSingleGuide = LOD.Guides.InterpolationResource->UseSingleGuide();
+
+					FRDGHairStrandsCullingData CullingData;
+					AddHairStrandsInterpolationPass(
+						GraphBuilder,
+						ShaderMap,
+						ShaderPrintData,
+						Instance,
+						LOD.Guides.RestResource->GetPointCount(),
+						MeshLODIndex,
+						1.0f,
+						LOD.Guides.HairInterpolationType,
+						InstanceGeometryType,
+						CullingData,
+						LOD.Guides.RestResource->GetPositionOffset(),
+						bValidGuide ? Instance->Guides.RestResource->GetPositionOffset() : FVector::ZeroVector,
+						RegisterAsSRV(GraphBuilder, LOD.Guides.DeformedResource->GetPositionOffsetBuffer(FHairStrandsDeformedResource::Current)),
+						bValidGuide ? RegisterAsSRV(GraphBuilder, Instance->Guides.DeformedResource->GetPositionOffsetBuffer(FHairStrandsDeformedResource::Current)) : nullptr,
+						bHasSkinning ? LOD.Guides.RestRootResource : nullptr ,
+						bHasSkinning && bValidGuide ? Instance->Guides.RestRootResource : nullptr,
+						bHasSkinning ? LOD.Guides.DeformedRootResource : nullptr,
+						bHasSkinning && bValidGuide ? Instance->Guides.DeformedRootResource : nullptr,
+						RegisterAsSRV(GraphBuilder, LOD.Guides.RestResource->PositionBuffer),
+						RegisterAsSRV(GraphBuilder, LOD.Guides.RestResource->PointToCurveBuffer),
+						bUseSingleGuide,
+						bValidGuide ? RegisterAsSRV(GraphBuilder, LOD.Guides.InterpolationResource->InterpolationBuffer) : nullptr,
+						bValidGuide ? RegisterAsSRV(GraphBuilder, Instance->Guides.RestResource->PositionBuffer) : nullptr,
+						bValidGuide ? RegisterAsSRV(GraphBuilder, Instance->Guides.DeformedResource->GetBuffer(FHairStrandsDeformedResource::Current)) : nullptr,
+						RegisterAsSRV(GraphBuilder, LOD.Guides.InterpolationResource->SimRootPointIndexBuffer),
+						bValidGuide ? RegisterAsSRV(GraphBuilder, Instance->Guides.RestResource->PointToCurveBuffer) : nullptr,
+						nullptr,
+						Guides_DeformedPositionBuffer,
+						FHairStrandsDeformedRootResource::FLOD::Current); // <- this should be optional
+
+					AddHairCardsDeformationPass(
+						GraphBuilder,
+						ShaderMap,
+						ShaderPrintData,
+						Instance,
+						MeshLODIndex);
+				}
+				// 2. Cards are deformed only based on skel. mesh RBF data)
+				else if (CardsSimulationType == EHairCardsSimulationType::RBF)
+				{
+					AddHairCardsRBFInterpolationPass(
+						GraphBuilder,
+						ShaderMap,
+						MeshLODIndex,
+						LOD.RestResource,
+						LOD.DeformedResource,
+						Instance->Guides.RestRootResource,
+						Instance->Guides.DeformedRootResource);
+				}
+
+				if (LOD.DeformedResource)
+				{
+					ExternalAccessQueue.Add(Register(GraphBuilder, LOD.DeformedResource->GetBuffer(FHairCardsDeformedResource::Current), ERDGImportedBufferFlags::None).Buffer, ERHIAccess::SRVMask);
+					ExternalAccessQueue.Add(Register(GraphBuilder, LOD.DeformedResource->GetBuffer(FHairCardsDeformedResource::Previous), ERDGImportedBufferFlags::None).Buffer, ERHIAccess::SRVMask);
+				}
+			}
+
+			#if RHI_RAYTRACING
+			if (LOD.RaytracingResource && IsRayTracingEnabled())
+			{
+				FCardsOrMeshesResourceBLASParameters* Parameters = GraphBuilder.AllocParameters<FCardsOrMeshesResourceBLASParameters>();
+				Parameters->PositionBuffer = LOD.DeformedResource ? Register(GraphBuilder, LOD.DeformedResource->GetBuffer(FHairCardsDeformedResource::Current), ERDGImportedBufferFlags::None).Buffer : nullptr;
+
+				// * When cards are dynamic, RT geometry is built once and then update.
+				// * When cards are static (i.e., no sim, not attached to  skinning()), in which case RT geometry needs only to be built once. This static geometry is shared between all the instances, 
+				//   and owned by the asset, rathter than the component. In this case the RT geometry will be built only once for all the instances
+				const bool bNeedUpdate = bNeedDeformation;
+				const bool bNeedBuild = !LOD.RaytracingResource->bIsRTGeometryInitialized;
+				if (bNeedBuild || bNeedUpdate)
+				{
+					GraphBuilder.AddPass(
+						RDG_EVENT_NAME("HairStrands::UpdateBLAS(Cards)"),
+						Parameters,
+						ERDGPassFlags::NeverCull | ERDGPassFlags::Compute,
+						[Instance, HairLODIndex, bNeedUpdate](FRHICommandListImmediate& RHICmdList)
+					{
+						SCOPED_GPU_MASK(RHICmdList, FRHIGPUMask::All());
+
+						FHairGroupInstance::FCards::FLOD& LocalLOD = Instance->Cards.LODs[HairLODIndex];
+
+						const bool bLocalNeedBuild = !LocalLOD.RaytracingResource->bIsRTGeometryInitialized;
+						if (bLocalNeedBuild)
+						{
+							BuildHairAccelerationStructure_Cards(RHICmdList, LocalLOD.RestResource, LocalLOD.DeformedResource, &LocalLOD.RaytracingResource->RayTracingGeometry, Instance->Debug.GroomAssetName, Instance->Debug.MeshComponentName, HairLODIndex);
+							LocalLOD.RaytracingResource->bIsRTGeometryInitialized = true;
+						}
+						else if (bNeedUpdate)
+						{
+							// TODO: evaluate perf tradeoff of rebuild vs refit
+							UpdateHairAccelerationStructure(RHICmdList, &LocalLOD.RaytracingResource->RayTracingGeometry, EAccelerationStructureBuildMode::Update);
+						}
+					});
+				}
+			}
+			#endif
+		}
+	}
+	else if (InstanceGeometryType == EHairGeometryType::Meshes)
+	{
+		DECLARE_GPU_STAT(HairMeshesInterpolation);
+		RDG_EVENT_SCOPE(GraphBuilder, "HairInterpolation(Meshes)");
+		RDG_GPU_STAT_SCOPE(GraphBuilder, HairMeshesInterpolation);
+
+		const uint32 HairLODIndex = Instance->HairGroupPublicData->GetIntLODIndex();
+		if (Instance->Meshes.IsValid(HairLODIndex))
+		{
+			const bool bNeedDeformation = Instance->Meshes.LODs[HairLODIndex].DeformedResource != nullptr;
+			if (bNeedDeformation)
+			{
+				FHairGroupInstance::FMeshes::FLOD& MeshesInstance = Instance->Meshes.LODs[HairLODIndex];
+
+				check(Instance->BindingType == EHairBindingType::Skinning)
+				check(Instance->Guides.IsValid());
+				check(Instance->Guides.HasValidRootData());
+				check(Instance->Guides.DeformedRootResource);
+				check(MeshLODIndex == -1 || Instance->Guides.DeformedRootResource->IsValid(MeshLODIndex)); //MeshLODIndex -1 indicates that skin cache is disabled and this is a workaround to prevent the editor from crashing - an editor setting guildeline popup exists to inform the user that the skin cache should be enabled. 
+
+				AddHairMeshesRBFInterpolationPass(
+					GraphBuilder,
+					ShaderMap,
+					MeshLODIndex,
+					MeshesInstance.RestResource,
+					MeshesInstance.DeformedResource,
+					Instance->Guides.RestRootResource,
+					Instance->Guides.DeformedRootResource);
+
+				ExternalAccessQueue.Add(Register(GraphBuilder, MeshesInstance.DeformedResource->GetBuffer(FHairMeshesDeformedResource::Current), ERDGImportedBufferFlags::None).Buffer, ERHIAccess::SRVMask);
+			}
+
+			#if RHI_RAYTRACING
+			FHairGroupInstance::FMeshes::FLOD& LOD = Instance->Meshes.LODs[HairLODIndex];
+			if (LOD.RaytracingResource)
+			{
+				FCardsOrMeshesResourceBLASParameters* Parameters = GraphBuilder.AllocParameters<FCardsOrMeshesResourceBLASParameters>();
+				Parameters->PositionBuffer = LOD.DeformedResource ? Register(GraphBuilder, LOD.DeformedResource->GetBuffer(FHairMeshesDeformedResource::Current), ERDGImportedBufferFlags::None).Buffer : nullptr;
+
+				const bool bNeedUpdate = bNeedDeformation;
+				const bool bNeedBuid = !LOD.RaytracingResource->bIsRTGeometryInitialized;
+				if (bNeedBuid || bNeedUpdate)
+				{
+					GraphBuilder.AddPass(
+						RDG_EVENT_NAME("HairStrands::UpdateBLAS(Meshes)"),
+						Parameters,
+						ERDGPassFlags::NeverCull | ERDGPassFlags::Compute,
+						[Instance, HairLODIndex, bNeedUpdate](FRHICommandListImmediate& RHICmdList)
+					{
+						SCOPED_GPU_MASK(RHICmdList, FRHIGPUMask::All());
+
+						FHairGroupInstance::FMeshes::FLOD& LocalLOD = Instance->Meshes.LODs[HairLODIndex];				
+						const bool bLocalNeedBuild = !LocalLOD.RaytracingResource->bIsRTGeometryInitialized;
+						if (bLocalNeedBuild)
+						{
+							BuildHairAccelerationStructure_Meshes(RHICmdList, LocalLOD.RestResource, LocalLOD.DeformedResource,  &LocalLOD.RaytracingResource->RayTracingGeometry, Instance->Debug.GroomAssetName, HairLODIndex);
+						}
+						else if (bNeedUpdate)
+						{
+							// TODO: evaluate perf tradeoff of rebuild vs refit
+							UpdateHairAccelerationStructure(RHICmdList, &LocalLOD.RaytracingResource->RayTracingGeometry, EAccelerationStructureBuildMode::Update);
+						}
+						LocalLOD.RaytracingResource->bIsRTGeometryInitialized = true;
+					});
+				}
+			}
+			#endif
+		}
+	}
+
+	Instance->HairGroupPublicData->VFInput.GeometryType = InstanceGeometryType;
+	Instance->HairGroupPublicData->VFInput.LocalToWorldTransform = Instance->GetCurrentLocalToWorld();
+	Instance->HairGroupPublicData->bSupportVoxelization = Instance->Strands.Modifier.bSupportVoxelization && Instance->bCastShadow;
+
+	ExternalAccessQueue.Submit(GraphBuilder);
+}
+
+void ResetHairStrandsInterpolation(
+	FRDGBuilder& GraphBuilder,
+	FGlobalShaderMap* ShaderMap,
+	FHairGroupInstance* Instance,
+	int32 MeshLODIndex)
+{
+	if (!Instance || 
+		(Instance && (Instance->Guides.bIsSimulationEnable || Instance->Guides.bIsDeformationEnable || Instance->Guides.bIsSimulationCacheEnable)) ||
+		(Instance && !Instance->Guides.bHasGlobalInterpolation && !Instance->Guides.bIsSimulationEnable && !Instance->Guides.bIsDeformationEnable && !Instance->Guides.bIsSimulationCacheEnable) ||
+		!IsHairStrandsBindingEnable()) return;
+
+	DECLARE_GPU_STAT(HairStrandsGuideDeform);
+	RDG_EVENT_SCOPE(GraphBuilder, "HairStrandsGuideDeform");
+	RDG_GPU_STAT_SCOPE(GraphBuilder, HairStrandsGuideDeform);
+
+	FRDGExternalBuffer RawDeformedPositionBuffer = Instance->Guides.DeformedResource->GetBuffer(FHairStrandsDeformedResource::Current);
+	FRDGImportedBuffer DeformedPositionBuffer = Register(GraphBuilder, RawDeformedPositionBuffer, ERDGImportedBufferFlags::CreateUAV);
+
+	AddDeformSimHairStrandsPass(
+		GraphBuilder,
+		ShaderMap,
+		EDeformationType::OffsetGuide,
+		MeshLODIndex,
+		Instance->Guides.RestResource->GetPointCount(),
+		Instance->Guides.RestRootResource,
+		Instance->Guides.DeformedRootResource,
+		RegisterAsSRV(GraphBuilder, Instance->Guides.RestResource->PositionBuffer),
+		RegisterAsSRV(GraphBuilder, Instance->Guides.RestResource->PointToCurveBuffer),
+		DeformedPositionBuffer,
+		Instance->Guides.RestResource->GetPositionOffset(),
+		RegisterAsSRV(GraphBuilder, Instance->Guides.DeformedResource->GetPositionOffsetBuffer(FHairStrandsDeformedResource::Current)),
+		Instance->Guides.bHasGlobalInterpolation,
+		nullptr);
+}
diff --git a/Engine/Plugins/Runtime/HairStrands/Source/HairStrandsCore/Private/HairStrandsRendering.cpp b/Engine/Plugins/Runtime/HairStrands/Source/HairStrandsCore/Private/HairStrandsRendering.cpp
deleted file mode 100644
diff --git a/Engine/Plugins/Runtime/HairStrands/Source/HairStrandsCore/Public/HairCardsVertexFactory.h b/Engine/Plugins/Runtime/HairStrands/Source/HairStrandsCore/Public/HairCardsVertexFactory.h
--- a/Engine/Plugins/Runtime/HairStrands/Source/HairStrandsCore/Public/HairCardsVertexFactory.h
+++ b/Engine/Plugins/Runtime/HairStrands/Source/HairStrandsCore/Public/HairCardsVertexFactory.h
@@ -8,7 +8,6 @@
 #include "VertexFactory.h"
 #include "RenderGraphResources.h"
 #include "HairCardsDatas.h"
-#include "HairStrandsRendering.h"
 #include "HairStrandsInterface.h"
 #include "PrimitiveSceneProxy.h"
 
diff --git a/Engine/Plugins/Runtime/HairStrands/Source/HairStrandsCore/Public/HairStrandsInterpolation.h b/Engine/Plugins/Runtime/HairStrands/Source/HairStrandsCore/Public/HairStrandsInterpolation.h
new file mode 100644
--- /dev/null
+++ b/Engine/Plugins/Runtime/HairStrands/Source/HairStrandsCore/Public/HairStrandsInterpolation.h
@@ -0,0 +1,58 @@
+// Copyright Epic Games, Inc. All Rights Reserved.
+
+#pragma once
+
+#include "Containers/ContainersFwd.h"
+#include "HAL/Platform.h"
+#include "Math/MathFwd.h"
+
+#if UE_ENABLE_INCLUDE_ORDER_DEPRECATED_IN_5_3
+#include "CoreMinimal.h"
+#include "UObject/ObjectMacros.h"
+#include "Engine/EngineTypes.h"
+#include "RenderGraphResources.h"
+#include "HairStrandsInterface.h"
+#include "Shader.h"
+#include "GroomDesc.h"
+#endif
+
+class FGlobalShaderMap;
+class FRDGBuilder;
+class FSkeletalMeshRenderData;
+class UGroomBindingAsset;
+
+struct FHairGroupInstance;
+struct FHairStrandClusterData;
+struct FRWBuffer;
+struct FShaderPrintData;
+
+enum class EGroomViewMode : uint8;
+
+// Reset the interpolation data. This needs to be called prior to ComputeHairStrandsInterpolation 
+// and prior to the actual hair simulation in order to insure that:
+//  1) when hair simulation is enabled, the first frame is correct
+//  2) when hair simulation is enabled/disabled (i.e., toggle/change) 
+//     we reset to deform buffer to rest state)
+void ResetHairStrandsInterpolation(
+	FRDGBuilder& GraphBuilder,
+	FGlobalShaderMap* ShaderMap,
+	FHairGroupInstance* Instance,
+	int32 LODIndex);
+
+void ComputeHairStrandsInterpolation(
+	FRDGBuilder& GraphBuilder,
+	FGlobalShaderMap* ShaderMap,
+	const uint32 ViewUniqueID,
+	const uint32 ViewRayTracingMask,
+	const EGroomViewMode ViewMode,
+	const FVector& TranslatedWorldOffset,
+	const FShaderPrintData* ShaderPrintData,
+	FHairGroupInstance* Instance,
+	int32 LODIndex,
+	FHairStrandClusterData* ClusterData);
+
+void RegisterClusterData(
+	FHairGroupInstance* Instance,
+	FHairStrandClusterData* InClusterData);
+
+HAIRSTRANDSCORE_API void ComputeInterpolationWeights(UGroomBindingAsset* BindingAsset, FSkeletalMeshRenderData* TargetRenderData, TArray<FRWBuffer>& TransferedPositions);
diff --git a/Engine/Plugins/Runtime/HairStrands/Source/HairStrandsCore/Public/HairStrandsRendering.h b/Engine/Plugins/Runtime/HairStrands/Source/HairStrandsCore/Public/HairStrandsRendering.h
deleted file mode 100644
diff --git a/Engine/Plugins/Runtime/HairStrands/Source/HairStrandsCore/Public/HairStrandsVertexFactory.h b/Engine/Plugins/Runtime/HairStrands/Source/HairStrandsCore/Public/HairStrandsVertexFactory.h
--- a/Engine/Plugins/Runtime/HairStrands/Source/HairStrandsCore/Public/HairStrandsVertexFactory.h
+++ b/Engine/Plugins/Runtime/HairStrands/Source/HairStrandsCore/Public/HairStrandsVertexFactory.h
@@ -8,13 +8,13 @@
 #include "VertexFactory.h"
 #include "RenderGraphResources.h"
 #include "HairStrandsDatas.h"
-#include "HairStrandsRendering.h"
 #include "HairStrandsInterface.h"
 #include "PrimitiveSceneProxy.h"
 
 class FMaterial;
 class FSceneView;
 struct FMeshBatchElement;
+struct FHairGroupInstance;
 
 /**
  * A vertex factory which simply transforms explicit vertex attributes from local to world space.
diff --git a/Engine/Plugins/Runtime/HairStrands/Source/HairStrandsEditor/Private/GroomActions.cpp b/Engine/Plugins/Runtime/HairStrands/Source/HairStrandsEditor/Private/GroomActions.cpp
--- a/Engine/Plugins/Runtime/HairStrands/Source/HairStrandsEditor/Private/GroomActions.cpp
+++ b/Engine/Plugins/Runtime/HairStrands/Source/HairStrandsEditor/Private/GroomActions.cpp
@@ -5,7 +5,6 @@
 
 #include "EditorFramework/AssetImportData.h"
 #include "HairStrandsCore.h"
-#include "HairStrandsRendering.h"
 #include "GeometryCache.h"
 #include "GroomAssetImportData.h"
 #include "GroomBuilder.h"
