Fix all hair position reading/writing with proper pack/unpack function.

diff --git a/Engine/Plugins/Runtime/HairStrands/Shaders/Private/NiagaraDataInterfaceHairStrands.ush b/Engine/Plugins/Runtime/HairStrands/Shaders/Private/NiagaraDataInterfaceHairStrands.ush
--- a/Engine/Plugins/Runtime/HairStrands/Shaders/Private/NiagaraDataInterfaceHairStrands.ush
+++ b/Engine/Plugins/Runtime/HairStrands/Shaders/Private/NiagaraDataInterfaceHairStrands.ush
@@ -9,13 +9,8 @@ NiagaraDataInterfaceHairStrands.ush
 static const float M_PI = 3.14159265;
 static const float3 MAX_POS = float3(1e+8,1e+8,1e+8);
 
-float2 ExtractBarycentrics(in uint PackedCoordinates)
+float3 Uint3ToFloat3( in uint3 Uint3Vector)
 {
-	return float2(f16tof32(PackedCoordinates & 0xFFFF), f16tof32((PackedCoordinates >> 16) & 0xFFFF));
-}
-
- float3 Uint3ToFloat3( in uint3 Uint3Vector)
- {
 	uint3 SignedVector = Uint3Vector;
 	SignedVector.x ^= (((SignedVector.x >> 31)-1) | 0x80000000);
 	SignedVector.y ^= (((SignedVector.y >> 31)-1) | 0x80000000);
diff --git a/Engine/Plugins/Runtime/HairStrands/Shaders/Private/NiagaraDataInterfaceHairStrandsTemplate.ush b/Engine/Plugins/Runtime/HairStrands/Shaders/Private/NiagaraDataInterfaceHairStrandsTemplate.ush
--- a/Engine/Plugins/Runtime/HairStrands/Shaders/Private/NiagaraDataInterfaceHairStrandsTemplate.ush
+++ b/Engine/Plugins/Runtime/HairStrands/Shaders/Private/NiagaraDataInterfaceHairStrandsTemplate.ush
@@ -144,13 +144,13 @@ float ComputeStrandLength_{ParameterName}(in int StrandIndex)
 		int PointOffset = Curve.PointIndex;
 		const int EdgeCount = Curve.PointCount-1;
 
-		float3 PointNext = f16tof32({ParameterName}_RestPositionBuffer[PointOffset++].xyz);
+		float3 PointNext = UnpackHairControlPointPosition({ParameterName}_RestPositionBuffer[PointOffset++]);
 		float3 PointPrev = PointNext;
 
 		for (int EdgeIndex = 0; EdgeIndex < EdgeCount; ++EdgeIndex)
 		{
 			PointPrev = PointNext;
-			PointNext = f16tof32({ParameterName}_RestPositionBuffer[PointOffset++].xyz);
+			PointNext = UnpackHairControlPointPosition({ParameterName}_RestPositionBuffer[PointOffset++]);
 			StrandLength += length(PointNext-PointPrev);
 		}
 	}
@@ -211,19 +211,19 @@ void SetNodePosition_{ParameterName}(in bool RootNode, const float NodePoint, co
 		const int FirstNext = PointPrev+1;
 
 		const float FirstAlpha = FirstPoint - (float)FirstPrev;
-		const float3 FirstPosition =  f16tof32({ParameterName}_RestPositionBuffer[PointOffset+FirstPrev].xyz) * (1.0-FirstAlpha) + 
-						   f16tof32({ParameterName}_RestPositionBuffer[PointOffset+FirstNext].xyz) * FirstAlpha;
-		const float3 DiffPosition = FirstPosition - f16tof32({ParameterName}_RestPositionBuffer[PointOffset].xyz);
+		const float3 FirstPosition =  UnpackHairControlPointPosition({ParameterName}_RestPositionBuffer[PointOffset+FirstPrev]) * (1.0-FirstAlpha) + 
+						 			  UnpackHairControlPointPosition({ParameterName}_RestPositionBuffer[PointOffset+FirstNext]) * FirstAlpha;
+		const float3 DiffPosition = FirstPosition - UnpackHairControlPointPosition({ParameterName}_RestPositionBuffer[PointOffset]);
 		const float3 EdgeDirection =  normalize(DiffPosition);
 
 		const float EdgeLength = ComputeStrandLength_{ParameterName}(StrandIndex) * CoordScale;
-		OutNodePosition = 	f16tof32({ParameterName}_RestPositionBuffer[PointOffset+PointPrev].xyz) - EdgeLength * normalize(EdgeDirection);					 
+		OutNodePosition = 	UnpackHairControlPointPosition({ParameterName}_RestPositionBuffer[PointOffset+PointPrev]) - EdgeLength * normalize(EdgeDirection);					 
 	}
 	else
 	{
 		const float NodeAlpha = NodePoint - (float)PointPrev;
-		OutNodePosition =  f16tof32({ParameterName}_RestPositionBuffer[PointOffset+PointPrev].xyz) * (1.0-NodeAlpha) + 
-						   f16tof32({ParameterName}_RestPositionBuffer[PointOffset+PointNext].xyz) * NodeAlpha;
+		OutNodePosition =  UnpackHairControlPointPosition({ParameterName}_RestPositionBuffer[PointOffset+PointPrev]) * (1.0-NodeAlpha) + 
+						   UnpackHairControlPointPosition({ParameterName}_RestPositionBuffer[PointOffset+PointNext]) * NodeAlpha;
 	}
 	OutNodePosition = mul(float4( OutNodePosition+{ParameterName}_RestPositionOffset, 1.0), {ParameterName}_WorldTransform).xyz;
 }
@@ -666,17 +666,20 @@ void UpdateDeformedPositions_{ParameterName}(in float3 LocalDisplace, in int Pos
 {
 	if( {ParameterName}_InterpolationMode >= 1)
 	{
-		const float3 LocalTrianglePosition = TriangleLocalPosition_{ParameterName}(RestTrianglePosition, RestTriangleOrientation, 
-			f16tof32({ParameterName}_RestPositionBuffer[PositionIndex].xyz) + {ParameterName}_RestPositionOffset) + LocalDisplace;
+		uint4 Packed = { ParameterName }_RestPositionBuffer[PositionIndex];
+
+		const float3 LocalTrianglePosition = TriangleLocalPosition_{ParameterName}(RestTrianglePosition, RestTriangleOrientation, UnpackHairControlPointPosition(Packed, {ParameterName}_RestPositionOffset) + LocalDisplace);
 			
-		const float3 LocalComponentPosition = TriangleWorldPosition_{ParameterName}(DeformedTrianglePosition, DeformedTriangleOrientation, LocalTrianglePosition) - {ParameterName}_DeformedPositionOffset[0].xyz;
-		{ParameterName}_DeformedPositionBuffer[PositionIndex] = uint4(f32tof16(LocalComponentPosition), {ParameterName}_RestPositionBuffer[PositionIndex].w);
+		const float3 LocalComponentPosition = TriangleWorldPosition_{ParameterName}(DeformedTrianglePosition, DeformedTriangleOrientation, LocalTrianglePosition);
+		{ParameterName}_DeformedPositionBuffer[PositionIndex] = PackHairControlPointPosition(Packed, LocalComponentPosition, {ParameterName}_DeformedPositionOffset[0].xyz);
 		
 	}
 	else
 	{
-		const float3 LocalComponentPosition = f16tof32({ParameterName}_RestPositionBuffer[PositionIndex].xyz) + LocalDisplace + {ParameterName}_RestPositionOffset - {ParameterName}_DeformedPositionOffset[0].xyz;
-		{ParameterName}_DeformedPositionBuffer[PositionIndex] = uint4(f32tof16(LocalComponentPosition), {ParameterName}_RestPositionBuffer[PositionIndex].w);	
+		uint4 Packed = {ParameterName}_RestPositionBuffer[PositionIndex];
+
+		const float3 LocalComponentPosition = UnpackHairControlPointPosition(Packed, {ParameterName}_RestPositionOffset) + LocalDisplace;
+		{ ParameterName }_DeformedPositionBuffer[PositionIndex] = PackHairControlPointPosition(Packed, LocalComponentPosition, {ParameterName}_DeformedPositionOffset[0].xyz);
 	}
 }
 
@@ -723,7 +726,8 @@ void UpdatePointPosition_{ParameterName}(in float3 NodePosition, in float3 RestP
 // Get the point position
 void GetPointPosition_{ParameterName}(in int PointIndex, out float3 OutPointPosition)
 {
-	OutPointPosition = mul(float4(f16tof32({ParameterName}_DeformedPositionBuffer[PointIndex].xyz) + {ParameterName}_DeformedPositionOffset[0].xyz, 1.0), {ParameterName}_WorldTransform).xyz;
+	const float3 Position = UnpackHairControlPointPosition({ParameterName}_DeformedPositionBuffer[PointIndex], {ParameterName}_DeformedPositionOffset[0].xyz);
+	OutPointPosition = mul(float4(Position, 1.0), {ParameterName}_WorldTransform).xyz;
 }
 
 /* -----------------------------------------------------------------
diff --git a/Engine/Plugins/Runtime/HairStrands/Shaders/Private/NiagaraGroomInterpolationUtils.ush b/Engine/Plugins/Runtime/HairStrands/Shaders/Private/NiagaraGroomInterpolationUtils.ush
--- a/Engine/Plugins/Runtime/HairStrands/Shaders/Private/NiagaraGroomInterpolationUtils.ush
+++ b/Engine/Plugins/Runtime/HairStrands/Shaders/Private/NiagaraGroomInterpolationUtils.ush
@@ -47,10 +47,4 @@ float3 DIHairStrands_TriangleWorldPosition(in float3 TrianglePosition, in float4
 float4 DIHairStrands_TriangleWorldOrientation(in float3 TrianglePosition, in float4 TriangleRotation, in float4 LocalOrientation)
 {
 	return NormalizeQuat(MultiplyQuat(TriangleRotation, LocalOrientation));
-}
-
-float2 ExtractBarycentrics(in uint PackedCoordinates)
-{
-	return float2(f16tof32(PackedCoordinates & 0xFFFF), f16tof32((PackedCoordinates >> 16) & 0xFFFF));
-}
-
+}
\ No newline at end of file
diff --git a/Engine/Plugins/Runtime/HairStrands/Shaders/Private/NiagaraInterpolateGroomGuides.usf b/Engine/Plugins/Runtime/HairStrands/Shaders/Private/NiagaraInterpolateGroomGuides.usf
--- a/Engine/Plugins/Runtime/HairStrands/Shaders/Private/NiagaraInterpolateGroomGuides.usf
+++ b/Engine/Plugins/Runtime/HairStrands/Shaders/Private/NiagaraInterpolateGroomGuides.usf
@@ -10,6 +10,7 @@ NiagaraInterpolateGroomGuides.usf
 #define WORLD_INTERPOLATION 1
 
 #include "/Engine/Public/Platform.ush"
+#include "/Engine/Private/HairStrands/HairStrandsPack.ush"
 
 #if PERMUTATION_INTERPOLATION == LOCAL_INTERPOLATION
 
@@ -109,7 +110,7 @@ void MainCS(uint3 DispatchThreadId : SV_DispatchThreadID)
 
 #if PERMUTATION_INTERPOLATION == LOCAL_INTERPOLATION
 		
-		const float2 ProjectionUV = ExtractBarycentrics(RootBarycentricCoordinatesBuffer[StrandIndex]);
+		const float2 ProjectionUV = UnpackBarycentrics(RootBarycentricCoordinatesBuffer[StrandIndex]).xy;
 
 		float3 RestTrianglePosition = float3(0, 0, 0), DeformedTrianglePosition = float3(0, 0, 0);
 		float4 RestTriangleOrientation = QUATERNION_IDENTITY, DeformedTriangleOrientation = QUATERNION_IDENTITY;
@@ -126,11 +127,12 @@ void MainCS(uint3 DispatchThreadId : SV_DispatchThreadID)
 
 		const float3 LocalPointDisplace = (TrianglePrevNodePosition - TrianglePrevRestPosition)  * (1.0 - PointAlpha) + (TriangleNextNodePosition - TriangleNextRestPosition) * PointAlpha;
 		
+		const uint4 Packed = RestPositionBuffer[DispatchThreadId.x];
 		const float3 LocalTrianglePosition = DIHairStrands_TriangleLocalPosition(RestTrianglePosition, RestTriangleOrientation, 
-					f16tof32(RestPositionBuffer[DispatchThreadId.x].xyz) + RestPositionOffset) + LocalPointDisplace;
+					UnpackHairControlPointPosition(Packed, RestPositionOffset) + LocalPointDisplace);
 		
 		const float3 LocalComponentPosition = DIHairStrands_TriangleWorldPosition(DeformedTrianglePosition, DeformedTriangleOrientation, LocalTrianglePosition) - DeformedPositionOffset[0].xyz;
-		DeformedPositionBuffer[DispatchThreadId.x] = uint4(f32tof16(LocalComponentPosition), RestPositionBuffer[DispatchThreadId.x].w);
+		DeformedPositionBuffer[DispatchThreadId.x] = PackHairControlPointPosition(Packed, LocalComponentPosition);
 
 #endif
 		
@@ -138,8 +140,9 @@ void MainCS(uint3 DispatchThreadId : SV_DispatchThreadID)
 		
 		const float3 LocalPointDisplace = (LocalPrevNodePosition - LocalPrevRestPosition)  * (1.0 - PointAlpha) + (LocalNextNodePosition - LocalNextRestPosition) * PointAlpha;
 
-		const float3 LocalComponentPosition = f16tof32(RestPositionBuffer[DispatchThreadId.x].xyz) + LocalPointDisplace + RestPositionOffset - DeformedPositionOffset[0].xyz;
-		DeformedPositionBuffer[DispatchThreadId.x] = uint4(f32tof16(LocalComponentPosition), RestPositionBuffer[DispatchThreadId.x].w);
+		const uint4 Packed = RestPositionBuffer[DispatchThreadId.x];
+		const float3 LocalComponentPosition = UnpackHairControlPointPosition(RestPositionBuffer[DispatchThreadId.x], RestPositionOffset) + LocalPointDisplace;
+		DeformedPositionBuffer[DispatchThreadId.x] = PackHairControlPointPosition(Packed, LocalComponentPosition, DeformedPositionOffset[0].xyz);
 		
 #endif
 	}
diff --git a/Engine/Shaders/Private/HairStrands/HairCardsDeformation.usf b/Engine/Shaders/Private/HairStrands/HairCardsDeformation.usf
--- a/Engine/Shaders/Private/HairStrands/HairCardsDeformation.usf
+++ b/Engine/Shaders/Private/HairStrands/HairCardsDeformation.usf
@@ -1,6 +1,7 @@
 // Copyright Epic Games, Inc. All Rights Reserved.
 
 #include "../Common.ush"
+#include "HairStrandsPack.ush"
 #include "HairStrandsVertexFactoryCommon.ush"
 #include "HairStrandsBindingCommon.ush"
 
@@ -46,10 +47,10 @@ float3 ComputeDynamicGeometryOffset(
 	float3 RootBarycentric,
 	float3 GuideDeformedPositionOffset)
 {
-	const float3 RestGuidePoint = f16tof32(GuideRestPositionBuffer[GuideIndex].xyz) + GuideRestPositionOffset;
+	const float3 RestGuidePoint = UnpackHairControlPointPosition(GuideRestPositionBuffer[GuideIndex], GuideRestPositionOffset);
 	const float3 LocalRestGuidePoint = ToTriangle(RestGuidePoint, RootBarycentric, RestTriangle);
 
-	const float3 DeformedGuidePoint = f16tof32(GuideDeformedPositionBuffer[GuideIndex].xyz) + GuideDeformedPositionOffset;
+	const float3 DeformedGuidePoint = UnpackHairControlPointPosition(GuideDeformedPositionBuffer[GuideIndex], GuideDeformedPositionOffset);
 	const float3 LocalDeformedGuidePoint = ToTriangle(DeformedGuidePoint, RootBarycentric, DeformedTriangle);
 
 	return (LocalDeformedGuidePoint - LocalRestGuidePoint) * GuideVertexWeight;
diff --git a/Engine/Shaders/Private/HairStrands/HairStrandsForwardRaster.usf b/Engine/Shaders/Private/HairStrands/HairStrandsForwardRaster.usf
--- a/Engine/Shaders/Private/HairStrands/HairStrandsForwardRaster.usf
+++ b/Engine/Shaders/Private/HairStrands/HairStrandsForwardRaster.usf
@@ -222,6 +222,7 @@ uint GetControlPointCount()
 	return ControlPointCount.Load(uint3(0,0,0));
 }
 
+// Custom encoding for forward rasterizer. Position are 32bits float. This is temporary.
 FHairControlPoint UnpackHairControlPoint(uint InPrimId)
 {
 	const float4 Packed = ControlPoints[InPrimId];
diff --git a/Engine/Shaders/Private/HairStrands/HairStrandsGuideDeform.usf b/Engine/Shaders/Private/HairStrands/HairStrandsGuideDeform.usf
--- a/Engine/Shaders/Private/HairStrands/HairStrandsGuideDeform.usf
+++ b/Engine/Shaders/Private/HairStrands/HairStrandsGuideDeform.usf
@@ -36,7 +36,7 @@ float3 DeformGuideByTriangle(
 	FHairMeshTriangle DeformedTriangle,
 	float3 RootBarycentric)
 {
-	const float3 WorldRestGuidePoint = f16tof32(SimRestPosePositionBuffer[GuideIndex].xyz) + SimRestOffset;
+	const float3 WorldRestGuidePoint = UnpackHairControlPointPosition(SimRestPosePositionBuffer[GuideIndex], SimRestOffset);
 	const float3 LocalRestGuidePoint = ToTriangle(WorldRestGuidePoint, RootBarycentric, RestTriangle);
 	const float3 WorldDeformedGuidePoint = ToWorld(LocalRestGuidePoint, RootBarycentric, DeformedTriangle);
 
@@ -104,10 +104,10 @@ void MainCS(uint DispatchThreadId : SV_DispatchThreadID)
 	
 	if (VertexIndex < VertexCount)
 	{
-		const uint4  Data 				= SimRestPosePositionBuffer[VertexIndex];
-		const float3 ControlPoint		= f16tof32(Data.xyz);
-		const float3 DisplacedPosition	= DisplacePosition(VertexIndex, ControlPoint);
-		OutSimDeformedPositionBuffer[VertexIndex] = uint4(f32tof16(DisplacedPosition), Data.w);
+		const uint4  PackedCP			= SimRestPosePositionBuffer[VertexIndex];
+		const float3 CPPosition			= UnpackHairControlPointPosition(PackedCP);
+		const float3 DisplacedPosition	= DisplacePosition(VertexIndex, CPPosition);		
+		OutSimDeformedPositionBuffer[VertexIndex] = PackHairControlPointPosition(PackedCP, DisplacedPosition);
 	}
 }
 
diff --git a/Engine/Shaders/Private/HairStrands/HairStrandsInterpolation.usf b/Engine/Shaders/Private/HairStrands/HairStrandsInterpolation.usf
--- a/Engine/Shaders/Private/HairStrands/HairStrandsInterpolation.usf
+++ b/Engine/Shaders/Private/HairStrands/HairStrandsInterpolation.usf
@@ -137,9 +137,9 @@ float3 ComputeStaticGeometryOffset(
 	float GuideVertexWeight, 
 	float3 OutSimHairPositionOffset)
 {
-	const float3 RestGuidePoint = f16tof32(SimRestPosePositionBuffer[GuideIndex].xyz) + InSimHairPositionOffset;
+	const float3 RestGuidePoint = UnpackHairControlPointPosition(SimRestPosePositionBuffer[GuideIndex], InSimHairPositionOffset);
 
-	const float3 DeformedGuidePoint = f16tof32(SimDeformedPositionBuffer[GuideIndex].xyz) + OutSimHairPositionOffset;
+	const float3 DeformedGuidePoint = UnpackHairControlPointPosition(SimDeformedPositionBuffer[GuideIndex], OutSimHairPositionOffset);
 
 	return (DeformedGuidePoint - RestGuidePoint) * GuideVertexWeight;
 }
@@ -150,9 +150,9 @@ float3 ComputeLocalGeometryOffset(
 {
 	const uint RootIndex = SimRootPointIndexBuffer[GuideIndex];
 
-	const float3 RestGuidePoint = f16tof32(SimRestPosePositionBuffer[GuideIndex].xyz) - f16tof32(SimRestPosePositionBuffer[RootIndex].xyz) ;
+	const float3 RestGuidePoint = UnpackHairControlPointPosition(SimRestPosePositionBuffer[GuideIndex]) - UnpackHairControlPointPosition(SimRestPosePositionBuffer[RootIndex]);
 
-	const float3 DeformedGuidePoint = f16tof32(SimDeformedPositionBuffer[GuideIndex].xyz) - f16tof32(SimDeformedPositionBuffer[RootIndex].xyz) ;
+	const float3 DeformedGuidePoint = UnpackHairControlPointPosition(SimDeformedPositionBuffer[GuideIndex]) - UnpackHairControlPointPosition(SimDeformedPositionBuffer[RootIndex]);
 
 	return (DeformedGuidePoint - RestGuidePoint) * GuideVertexWeight;
 }
@@ -165,10 +165,10 @@ float3 ComputeDynamicGeometryOffset(
 	float3 RootBarycentric,
 	float3 OutSimHairPositionOffset)
 {
-	const float3 RestGuidePoint = f16tof32(SimRestPosePositionBuffer[GuideIndex].xyz) + InSimHairPositionOffset;
+	const float3 RestGuidePoint = UnpackHairControlPointPosition(SimRestPosePositionBuffer[GuideIndex], InSimHairPositionOffset);
 	const float3 LocalRestGuidePoint = ToTriangle(RestGuidePoint, RootBarycentric, RestTriangle);
 
-	const float3 DeformedGuidePoint = f16tof32(SimDeformedPositionBuffer[GuideIndex].xyz) + OutSimHairPositionOffset;
+	const float3 DeformedGuidePoint = UnpackHairControlPointPosition(SimDeformedPositionBuffer[GuideIndex], OutSimHairPositionOffset);
 	const float3 LocalDeformedGuidePoint = ToTriangle(DeformedGuidePoint, RootBarycentric, DeformedTriangle);
 
 	return (LocalDeformedGuidePoint - LocalRestGuidePoint) * GuideVertexWeight;
@@ -180,9 +180,9 @@ float3 ComputeTranslateGeometryOffset(
 {
 	const uint RootIndex = SimRootPointIndexBuffer[GuideIndex];
 
-	const float3 RestGuidePoint = f16tof32(SimRestPosePositionBuffer[GuideIndex].xyz) - f16tof32(SimRestPosePositionBuffer[RootIndex].xyz) ;
+	const float3 RestGuidePoint = UnpackHairControlPointPosition(SimRestPosePositionBuffer[GuideIndex]) - UnpackHairControlPointPosition(SimRestPosePositionBuffer[RootIndex]) ;
 
-	const float3 DeformedGuidePoint = f16tof32(SimDeformedPositionBuffer[GuideIndex].xyz) - f16tof32(SimDeformedPositionBuffer[RootIndex].xyz) ;
+	const float3 DeformedGuidePoint = UnpackHairControlPointPosition(SimDeformedPositionBuffer[GuideIndex]) - UnpackHairControlPointPosition(SimDeformedPositionBuffer[RootIndex]) ;
 
 	return (DeformedGuidePoint - RestGuidePoint) * GuideVertexWeight;
 }
@@ -195,10 +195,10 @@ float3 ComputeRotateGeometryOffset(
 {	
 	const uint RootIndex = SimRootPointIndexBuffer[GuideIndex];
 
-	const float3 RestGuidePoint = f16tof32(SimRestPosePositionBuffer[GuideIndex].xyz) - f16tof32(SimRestPosePositionBuffer[RootIndex].xyz) ;
+	const float3 RestGuidePoint = UnpackHairControlPointPosition(SimRestPosePositionBuffer[GuideIndex]) - UnpackHairControlPointPosition(SimRestPosePositionBuffer[RootIndex]) ;
 	const float3 LocalRestGuidePoint = VectorToTriangle(RestGuidePoint, RestTriangle);
 
-	const float3 DeformedGuidePoint = f16tof32(SimDeformedPositionBuffer[GuideIndex].xyz) - f16tof32(SimDeformedPositionBuffer[RootIndex].xyz) ;
+	const float3 DeformedGuidePoint = UnpackHairControlPointPosition(SimDeformedPositionBuffer[GuideIndex]) - UnpackHairControlPointPosition(SimDeformedPositionBuffer[RootIndex]) ;
 	const float3 LocalDeformedGuidePoint = VectorToTriangle(DeformedGuidePoint, DeformedTriangle);
 
 	return (LocalDeformedGuidePoint - LocalRestGuidePoint) * GuideVertexWeight;
diff --git a/Engine/Shaders/Private/HairStrands/HairStrandsPack.ush b/Engine/Shaders/Private/HairStrands/HairStrandsPack.ush
--- a/Engine/Shaders/Private/HairStrands/HairStrandsPack.ush
+++ b/Engine/Shaders/Private/HairStrands/HairStrandsPack.ush
@@ -28,10 +28,10 @@ uint4 PackHairControlPoint(
 
 FHairControlPoint UnpackHairControlPoint(
 	uint4 InPackedData,
-	float3 InVF_PositionOffset,
-	float InVF_Radius,
-	float InVF_RootScale,
-	float InVF_TipScale)
+	float3 InVF_PositionOffset=0,
+	float InVF_Radius=1,
+	float InVF_RootScale=1,
+	float InVF_TipScale=1)
 {
 	const uint PackedAlpha	= asuint(InPackedData.a);
 
@@ -52,12 +52,12 @@ float UnpackHairControlPointCoordU(uint4 InPackedData)
 	return ((PackedAlpha & 0xFF00) >> 8) / 255.f;
 }
 
-float3 UnpackHairControlPointPosition(uint4 InPackedData, float3 InVF_PositionOffset)
+float3 UnpackHairControlPointPosition(uint4 InPackedData, float3 InVF_PositionOffset=0)
 {
 	return f16tof32(InPackedData.xyz) + InVF_PositionOffset;
 }
 
-uint4 PackHairControlPointPosition(uint4 InPackedData, float3 InNewPosition, float3 InPositionOffset)
+uint4 PackHairControlPointPosition(uint4 InPackedData, float3 InNewPosition, float3 InPositionOffset=0)
 {
 	return uint4(f32tof16(InNewPosition.xyz - InPositionOffset), InPackedData.w);
 }
diff --git a/Engine/Shaders/Private/HairStrands/HairStrandsTangent.usf b/Engine/Shaders/Private/HairStrands/HairStrandsTangent.usf
--- a/Engine/Shaders/Private/HairStrands/HairStrandsTangent.usf
+++ b/Engine/Shaders/Private/HairStrands/HairStrandsTangent.usf
@@ -2,6 +2,7 @@
 
 #include "../Common.ush"
 #include "../MonteCarlo.ush"
+#include "HairStrandsPack.ush"
 #include "HairStrandsClusterCommon.ush"
 
 #if SHADER_TANGENT
@@ -54,31 +55,22 @@ void MainCS(uint DispatchThreadId : SV_DispatchThreadID)
 
 	if (IndexCurr < MaxVertexCount)
 	{
-		const uint4 DataCurr = PositionBuffer[IndexCurr];
-		const uint4 DataPrev = PositionBuffer[IndexPrev];
-		const uint4 DataNext = PositionBuffer[IndexNext];
-
-		const float3 PCurr = f16tof32(DataCurr.xyz);
-		float3 PPrev = f16tof32(DataPrev.xyz);
-		float3 PNext = f16tof32(DataNext.xyz);
-
-		// For more details about the vertex type, see StrandHairVertexFactor.ush
-		const uint TypeCurr = asuint(DataCurr.a) & 0x3;
-		const uint TypePrev = asuint(DataPrev.a) & 0x3;
-		const uint TypeNext = asuint(DataNext.a) & 0x3;
+		const FHairControlPoint CPCurr = UnpackHairControlPoint(PositionBuffer[IndexCurr]);
+		const FHairControlPoint CPPrev = UnpackHairControlPoint(PositionBuffer[IndexPrev]);
+		const FHairControlPoint CPNext = UnpackHairControlPoint(PositionBuffer[IndexNext]);
 
 		float3 T0 = 0;
-		const bool bIsPrevValid = TypePrev == 0 || TypePrev == 1;
+		const bool bIsPrevValid = CPPrev.Type == HAIR_CONTROLPOINT_INSIDE || CPPrev.Type == HAIR_CONTROLPOINT_START;
 		if (bIsPrevValid)
 		{
-			T0 = normalize(PCurr - PPrev);
+			T0 = normalize(CPCurr.Position - CPPrev.Position);
 		}
 
 		float3 T1 = 0;
-		const bool bIsNextValid = TypeNext == 0 || TypeNext == 2;
+		const bool bIsNextValid = CPNext.Type == HAIR_CONTROLPOINT_INSIDE || CPNext.Type == HAIR_CONTROLPOINT_END;
 		if (bIsNextValid)
 		{
-			T1 = normalize(PNext - PCurr);
+			T1 = normalize(CPNext.Position - CPCurr.Position);
 		}
 
 		const float3 T = normalize(T0 + T1);
