Fix serialization issue for groom legacy which does not contain any hair description.

diff --git a/Engine/Plugins/Runtime/HairStrands/Source/HairStrandsCore/Private/GroomAsset.cpp b/Engine/Plugins/Runtime/HairStrands/Source/HairStrandsCore/Private/GroomAsset.cpp
--- a/Engine/Plugins/Runtime/HairStrands/Source/HairStrandsCore/Private/GroomAsset.cpp
+++ b/Engine/Plugins/Runtime/HairStrands/Source/HairStrandsCore/Private/GroomAsset.cpp
@@ -711,7 +711,71 @@ static void InternalSerializeGuide(FArchive& Ar, UObject* Owner, FHairGroupPlatf
 static void InternalSerializeStrand(FArchive& Ar, UObject* Owner, FHairGroupPlatformData::FStrands& StrandData, bool bHeader, bool bData);
 static void InternalSerializePlatformData(FArchive& Ar, UObject* Owner, FHairGroupPlatformData& GroupData);
 
-// Serialize is only called for cooking??
+// This dummy serialization function is only intended to support *loading* of legacy content 
+static void DummySerizalizationForLegacyPurpose(FArchive& Ar)
+{
+	// Only support loading of legacy content 
+	check(Ar.IsLoading());
+	Ar.UsingCustomVersion(FReleaseObjectVersion::GUID);
+	Ar.UsingCustomVersion(FUE5ReleaseStreamObjectVersion::GUID);
+
+	uint32 GroupCount = 0;					Ar << GroupCount;
+	for (uint32 GroupIt = 0; GroupIt < GroupCount; ++GroupIt)
+	{
+		// Strands
+		{
+			uint32 CurveCount = 0;			Ar << CurveCount;
+			uint32 PointCount = 0;			Ar << PointCount;
+			float MaxLength = 0;			Ar << MaxLength;
+			float MaxRadius = 0;			Ar << MaxRadius;
+			FBox BoundingBox;				Ar << BoundingBox;
+			uint32 Flags = 0;				Ar << Flags;
+			uint32 AttributeOffsets[7] = { 0 };
+			for (uint8 It = 0; It < 7; ++It)
+			{
+				Ar << AttributeOffsets[It];
+			}
+		}
+
+		// Guides
+		{
+			uint32 CurveCount = 0;			Ar << CurveCount;
+			uint32 PointCount = 0;			Ar << PointCount;
+			float MaxLength = 0;			Ar << MaxLength;
+			float MaxRadius = 0;			Ar << MaxRadius;
+			FBox BoundingBox;				Ar << BoundingBox;
+			uint32 Flags = 0;				Ar << Flags;
+			uint32 AttributeOffsets[7] = { 0 };
+			for (uint8 It = 0; It < 7; ++It)
+			{
+				Ar << AttributeOffsets[It];
+			}
+		}
+
+		// Interpolation
+		{
+			uint32 Flags = 0;				Ar << Flags;
+			uint32 PointCount = 0;			Ar << PointCount;
+			uint32 SimPointCount = 0;		Ar << SimPointCount;
+		}
+
+		// Cluster
+		{
+			uint32 ClusterCount = 0;		Ar << ClusterCount;
+			uint32 ClusterLODCount = 0;		Ar << ClusterLODCount;
+			uint32 VertexCount = 0;			Ar << VertexCount;
+			uint32 VertexLODCount = 0;		Ar << VertexLODCount;
+			TArray<bool> LODVisibility;		Ar << LODVisibility;
+			TArray<float>CPULODScreenSize;	Ar << CPULODScreenSize;
+		}
+
+		// Misc
+		{
+			bool bIsCooked; 				Ar << bIsCooked;
+		}
+	}
+}
+
 void UGroomAsset::Serialize(FArchive& Ar)
 {
 	uint8 ClassDataStripFlags = GenerateClassStripFlags(Ar);
@@ -745,12 +809,17 @@ void UGroomAsset::Serialize(FArchive& Ar)
 			bHasImportedStrands = ClassDataStripFlags & CDSF_ImportedStrands;
 		}
 
-		if (StripFlags.IsEditorDataStripped() || !bHasImportedStrands || (Ar.IsSaving() && !CanRebuildFromDescription()))
+
+		if (StripFlags.IsEditorDataStripped())
 		{
-			// When cooking data or serializing old format to new format,
-			// serialize the computed groom data
+			// When cooking data or loading cooked data
 			InternalSerializePlatformDatas(Ar, this, HairGroupsPlatformData);
 		}
+		else if (Ar.CustomVer(FUE5MainStreamObjectVersion::GUID) < FUE5MainStreamObjectVersion::GroomAssetRemoveInAssetSerialization && Ar.IsLoading() && !bHasImportedStrands)
+		{
+			// Legacy support for groom asset not having strands data
+			DummySerizalizationForLegacyPurpose(Ar);
+		}
 #if WITH_EDITORONLY_DATA
 		else
 		{
diff --git a/Engine/Source/Runtime/Core/Public/UObject/UE5MainStreamObjectVersion.h b/Engine/Source/Runtime/Core/Public/UObject/UE5MainStreamObjectVersion.h
--- a/Engine/Source/Runtime/Core/Public/UObject/UE5MainStreamObjectVersion.h
+++ b/Engine/Source/Runtime/Core/Public/UObject/UE5MainStreamObjectVersion.h
@@ -343,6 +343,9 @@ struct CORE_API FUE5MainStreamObjectVersion
 		// Change the default for facing & alignment to be automatic
 		NiagaraSpriteRendererFacingAlignmentAutoDefault,
 
+		// Change the default for facing & alignment to be automatic
+		GroomAssetRemoveInAssetSerialization,
+
 		// -----<new versions can be added above this line>-------------------------------------------------
 		VersionPlusOne,
 		LatestVersion = VersionPlusOne - 1
