* Fix invalid NoCopy flags used with local array data
* Fix invalid buffer desc creation for dummy data.

diff --git a/Engine/Plugins/Runtime/HairStrands/Source/HairStrandsCore/Private/GroomResources.cpp b/Engine/Plugins/Runtime/HairStrands/Source/HairStrandsCore/Private/GroomResources.cpp
--- a/Engine/Plugins/Runtime/HairStrands/Source/HairStrandsCore/Private/GroomResources.cpp
+++ b/Engine/Plugins/Runtime/HairStrands/Source/HairStrandsCore/Private/GroomResources.cpp
@@ -360,13 +360,9 @@ static FRDGBufferRef InternalCreateVertexBuffer(
 	const void* InitialData,
 	uint64 InitialDataSize,
 	ERDGInitialDataFlags InitialDataFlags,
-	const FName& OwnerName,
-	bool bIsByteAddressBuffer = false)
+	const FName& OwnerName)
 {
-	if (!bIsByteAddressBuffer)
-	{
-		checkf(EnumHasAnyFlags(Desc.Usage, EBufferUsageFlags::VertexBuffer), TEXT("CreateVertexBuffer called with an FRDGBufferDesc underlying type that is not 'VertexBuffer'. Buffer: %s"), Name);
-	}
+	checkf(EnumHasAnyFlags(Desc.Usage, EBufferUsageFlags::VertexBuffer), TEXT("CreateVertexBuffer called with an FRDGBufferDesc underlying type that is not 'VertexBuffer'. Buffer: %s"), Name);
 	FRDGBufferRef Buffer = GraphBuilder.CreateBuffer(Desc, Name, ERDGBufferFlags::MultiFrame);
 	Buffer->SetOwnerName(OwnerName);
 	if (InitialData && InitialDataSize)
@@ -377,7 +373,7 @@ static FRDGBufferRef InternalCreateVertexBuffer(
 }
 
 template<typename FormatType>
-void InternalCreateVertexBufferRDG(FRDGBuilder& GraphBuilder, const typename FormatType::Type* InData, uint32 InDataCount, FRDGExternalBuffer& Out, const TCHAR* DebugName, const FName& OwnerName, EHairResourceUsageType UsageType, ERDGInitialDataFlags InitialDataFlags, bool bIsByteAddressBuffer=false)
+void InternalCreateVertexBufferRDG(FRDGBuilder& GraphBuilder, const typename FormatType::Type* InData, uint32 InDataCount, FRDGExternalBuffer& Out, const TCHAR* DebugName, const FName& OwnerName, EHairResourceUsageType UsageType, ERDGInitialDataFlags InitialDataFlags)
 {
 	FRDGBufferRef Buffer = nullptr;
 
@@ -396,14 +392,13 @@ void InternalCreateVertexBufferRDG(FRDGBuilder& GraphBuilder, const typename For
 		InData,
 		DataSizeInBytes,
 		InitialDataFlags,
-		OwnerName, 
-		bIsByteAddressBuffer);
+		OwnerName);
 
 	ConvertToExternalBufferWithViews(GraphBuilder, Buffer, Out, FormatType::Format);
 }
 
 template<typename FormatType>
-void InternalCreateVertexBufferRDG(FRDGBuilder& GraphBuilder, const TArray<typename FormatType::Type>& InData, FRDGExternalBuffer& Out, const TCHAR* DebugName, const FName& OwnerName, EHairResourceUsageType UsageType, ERDGInitialDataFlags InitialDataFlags= ERDGInitialDataFlags::NoCopy, bool bIsByteAddressBuffer=false)
+void InternalCreateVertexBufferRDG(FRDGBuilder& GraphBuilder, const TArray<typename FormatType::Type>& InData, FRDGExternalBuffer& Out, const TCHAR* DebugName, const FName& OwnerName, EHairResourceUsageType UsageType, ERDGInitialDataFlags InitialDataFlags=ERDGInitialDataFlags::NoCopy)
 {
 	FRDGBufferRef Buffer = nullptr;
 
@@ -423,14 +418,13 @@ void InternalCreateVertexBufferRDG(FRDGBuilder& GraphBuilder, const TArray<typen
 		InData.GetData(),
 		DataSizeInBytes,
 		InitialDataFlags,
-		OwnerName,
-		bIsByteAddressBuffer);
+		OwnerName);
 
 	ConvertToExternalBufferWithViews(GraphBuilder, Buffer, Out, FormatType::Format);
 }
 
 template<typename DataType>
-void InternalCreateVertexBufferRDG(FRDGBuilder& GraphBuilder, const TArray<DataType>& InData, EPixelFormat Format, FRDGExternalBuffer& Out, const TCHAR* DebugName, const FName& OwnerName, EHairResourceUsageType UsageType, bool bIsByteAddressBuffer=false)
+void InternalCreateVertexBufferRDG(FRDGBuilder& GraphBuilder, const TArray<DataType>& InData, EPixelFormat Format, FRDGExternalBuffer& Out, const TCHAR* DebugName, const FName& OwnerName, EHairResourceUsageType UsageType, ERDGInitialDataFlags InitialDataFlags=ERDGInitialDataFlags::NoCopy)
 {
 	FRDGBufferRef Buffer = nullptr;
 
@@ -450,9 +444,8 @@ void InternalCreateVertexBufferRDG(FRDGBuilder& GraphBuilder, const TArray<DataT
 		Desc,
 		InData.GetData(),
 		DataSizeInBytes,
-		ERDGInitialDataFlags::NoCopy,
-		OwnerName,
-		bIsByteAddressBuffer);
+		InitialDataFlags,
+		OwnerName);
 
 	ConvertToExternalBufferWithViews(GraphBuilder, Buffer, Out, Format);
 }
@@ -507,6 +500,26 @@ void InternalCreateStructuredBufferRDG(FRDGBuilder& GraphBuilder, uint32 DataCou
 	ConvertToExternalBufferWithViews(GraphBuilder, Buffer, Out, FormatType::Format);
 }
 
+template<typename DataType>
+void InternalCreateByteAddressBufferRDG(FRDGBuilder& GraphBuilder, const TArray<DataType>& InData, EPixelFormat Format, FRDGExternalBuffer& Out, const TCHAR* DebugName, const FName& OwnerName, EHairResourceUsageType UsageType)
+{
+	const uint32 DataCount = InData.Num();
+	const uint32 DataSizeInBytes = sizeof(DataType) * DataCount;
+	Out.Buffer = nullptr;
+	if (DataSizeInBytes != 0)
+	{
+		const FRDGBufferDesc Desc = ApplyUsage(FRDGBufferDesc::CreateByteAddressDesc(DataSizeInBytes), UsageType);
+		FRDGBufferRef Buffer = GraphBuilder.CreateBuffer(Desc, DebugName, ERDGBufferFlags::MultiFrame);
+		Buffer->SetOwnerName(OwnerName);
+		if (InData.GetData() && DataSizeInBytes)
+		{
+			GraphBuilder.QueueBufferUpload(Buffer, InData.GetData(), DataSizeInBytes, ERDGInitialDataFlags::None);
+		}
+		
+		ConvertToExternalBufferWithViews(GraphBuilder, Buffer, Out, Format);
+	}
+}
+
 /////////////////////////////////////////////////////////////////////////////////////////
 
 static UTexture2D* CreateCardTexture(FIntPoint Resolution)
@@ -922,7 +935,7 @@ void FHairStrandsRestResource::InternalAllocate(FRDGBuilder& GraphBuilder)
 	{
 		TArray<uint32> DummyAttribute;
 		DummyAttribute.Add(0u);
-		InternalCreateVertexBufferRDG(GraphBuilder, DummyAttribute, EPixelFormat::PF_R32_UINT, PointAttributeBuffer, ToHairResourceDebugName(HAIRSTRANDS_RESOUCE_NAME(CurveType, Hair.StrandsRest_PointAttributeBuffer), ResourceName), OwnerName, EHairResourceUsageType::Static, true /*bIsByteAddressBuffer*/);
+		InternalCreateByteAddressBufferRDG(GraphBuilder, DummyAttribute, EPixelFormat::PF_R32_UINT, PointAttributeBuffer, ToHairResourceDebugName(HAIRSTRANDS_RESOUCE_NAME(CurveType, Hair.StrandsRest_PointAttributeBuffer), ResourceName), OwnerName, EHairResourceUsageType::Static);
 		GraphBuilder.UseExternalAccessMode(Register(GraphBuilder, PointAttributeBuffer, ERDGImportedBufferFlags::CreateSRV).Buffer, ERHIAccess::SRVMask);
 	}
 	InternalCreateVertexBufferRDG_FromBulkData<FHairStrandsCurveFormat>(GraphBuilder, BulkData.Data.Curves.Data, CurveCount, CurveBuffer, ToHairResourceDebugName(HAIRSTRANDS_RESOUCE_NAME(CurveType, Hair.StrandsRest_CurveBuffer), ResourceName), OwnerName, EHairResourceUsageType::Static);
