Fix missing creation of default resouces for hair deformer.

diff --git a/Engine/Plugins/Runtime/HairStrands/Source/HairStrandsCore/Private/Deformer/DeformerDataInterfaceGroom.cpp b/Engine/Plugins/Runtime/HairStrands/Source/HairStrandsCore/Private/Deformer/DeformerDataInterfaceGroom.cpp
--- a/Engine/Plugins/Runtime/HairStrands/Source/HairStrandsCore/Private/Deformer/DeformerDataInterfaceGroom.cpp
+++ b/Engine/Plugins/Runtime/HairStrands/Source/HairStrandsCore/Private/Deformer/DeformerDataInterfaceGroom.cpp
@@ -191,7 +191,7 @@ void FOptimusGroomDataProviderProxy::AllocateResources(FRDGBuilder& GraphBuilder
 			R.CurveAttributeBuffer	= RegisterAsSRV(GraphBuilder, Instance->Strands.RestResource->CurveAttributeBuffer);
 			R.PointAttributeBuffer	= RegisterAsSRV(GraphBuilder, Instance->Strands.RestResource->PointAttributeBuffer);
 
-			if (FallbackSRV)
+			if (!FallbackSRV)
 			{
 				FallbackSRV = GraphBuilder.CreateSRV(GraphBuilder.RegisterExternalBuffer(GWhiteVertexBufferWithRDG->Buffer), PF_R16G16B16A16_UINT);
 			}
