Expose more hair attribute (ClumpId, AO, Color, roughness) to hair deformer.

diff --git a/Engine/Plugins/Runtime/HairStrands/Shaders/Private/DeformerDataInterfaceGroom.ush b/Engine/Plugins/Runtime/HairStrands/Shaders/Private/DeformerDataInterfaceGroom.ush
--- a/Engine/Plugins/Runtime/HairStrands/Shaders/Private/DeformerDataInterfaceGroom.ush
+++ b/Engine/Plugins/Runtime/HairStrands/Shaders/Private/DeformerDataInterfaceGroom.ush
@@ -26,10 +26,10 @@ FHairControlPoint ReadControlPointData(uint ControlPointIndex)
 	// Hair deformer work in local space so no need for position offset
 	return UnpackHairControlPoint(
 		{DataInterfaceName}_PositionBuffer[ControlPointIndex], 
-		{DataInterfaceName}_PositionOffsetBuffer[0],
+		{DataInterfaceName}_PositionOffsetBuffer[0].xyz,
 		{DataInterfaceName}_Radius, 
-		1.0f /*RootScale*/, 
-		1.0f /*TipScale*/);
+		{DataInterfaceName}_RootScale, 
+		{DataInterfaceName}_TipScale);
 }
 
 FHairCurve ReadCurveData(uint CurveIndex)
@@ -86,6 +86,26 @@ float ReadSeed_{DataInterfaceName}(uint ControlPointIndex)
 	return GetHairStrandsSeed(ControlPointIndex);
 }
 
+uint ReadClumpId_{DataInterfaceName}(uint ControlPointIndex)
+{
+	return GetHairStrandsClumpID(ControlPointIndex);
+}
+
+float3 ReadColor_{DataInterfaceName}(uint ControlPointIndex)
+{
+	return GetHairStrandsColor(ControlPointIndex, 0);
+}
+
+float ReadRoughness_{DataInterfaceName}(uint ControlPointIndex)
+{
+	return GetHairStrandsRoughness(ControlPointIndex, 0);
+}
+
+float ReadAO_{DataInterfaceName}(uint ControlPointIndex)
+{
+	return GetHairStrandsAO(ControlPointIndex, 0);
+}
+
 float ReadCurveOffsetPoint_{DataInterfaceName}(uint CurveIndex)
 {
 	return ReadCurveData(CurveIndex).PointIndex;
@@ -94,5 +114,4 @@ float ReadCurveOffsetPoint_{DataInterfaceName}(uint CurveIndex)
 float ReadCurveNumPoint_{DataInterfaceName}(uint CurveIndex)
 {
 	return ReadCurveData(CurveIndex).PointCount;
-}
-
+}
\ No newline at end of file
diff --git a/Engine/Plugins/Runtime/HairStrands/Source/HairStrandsCore/Private/Deformer/DeformerDataInterfaceGroom.cpp b/Engine/Plugins/Runtime/HairStrands/Source/HairStrandsCore/Private/Deformer/DeformerDataInterfaceGroom.cpp
--- a/Engine/Plugins/Runtime/HairStrands/Source/HairStrandsCore/Private/Deformer/DeformerDataInterfaceGroom.cpp
+++ b/Engine/Plugins/Runtime/HairStrands/Source/HairStrandsCore/Private/Deformer/DeformerDataInterfaceGroom.cpp
@@ -33,6 +33,10 @@ TArray<FOptimusCDIPinDefinition> UOptimusGroomDataInterface::GetPinDefinitions()
 	Defs.Add({ "Length",           "ReadLength",            ControlPoint,   "ReadLength" });
 	Defs.Add({ "RootUV",           "ReadRootUV",            ControlPoint,   "ReadRootUV" });
 	Defs.Add({ "Seed",             "ReadSeed",              ControlPoint,   "ReadSeed" });
+	Defs.Add({ "ClumpId",          "ReadClumpId",           ControlPoint,   "ReadClumpId" });
+	Defs.Add({ "Color",            "ReadColor",             ControlPoint,   "ReadColor" });
+	Defs.Add({ "Roughness",        "ReadRoughness",         ControlPoint,   "ReadRoughness" });
+	Defs.Add({ "AO",               "ReadAO",                ControlPoint,   "ReadAO" });
 	Defs.Add({ "CurveOffsetPoint", "ReadCurveOffsetPoint",  Curve,          "ReadCurveOffsetPoint" });
 	Defs.Add({ "CurveNumPoint",    "ReadCurveNumPoint",     Curve,          "ReadCurveNumPoint" });
 	return Defs;
@@ -83,6 +87,26 @@ void UOptimusGroomDataInterface::GetSupportedInputs(TArray<FShaderFunctionDefini
 		.AddReturnType(EShaderFundamentalType::Float)
 		.AddParam(EShaderFundamentalType::Uint);
 
+	OutFunctions.AddDefaulted_GetRef()
+		.SetName(TEXT("ReadClumpId"))
+		.AddReturnType(EShaderFundamentalType::Uint)
+		.AddParam(EShaderFundamentalType::Uint);
+
+	OutFunctions.AddDefaulted_GetRef()
+		.SetName(TEXT("ReadColor"))
+		.AddReturnType(EShaderFundamentalType::Float, 3)
+		.AddParam(EShaderFundamentalType::Uint);
+
+	OutFunctions.AddDefaulted_GetRef()
+		.SetName(TEXT("ReadRoughnes"))
+		.AddReturnType(EShaderFundamentalType::Float)
+		.AddParam(EShaderFundamentalType::Uint);
+
+	OutFunctions.AddDefaulted_GetRef()
+		.SetName(TEXT("ReadAO"))
+		.AddReturnType(EShaderFundamentalType::Float)
+		.AddParam(EShaderFundamentalType::Uint);
+
 	OutFunctions.AddDefaulted_GetRef()
 		.SetName(TEXT("ReadCurveOffsetPoint"))
 		.AddReturnType(EShaderFundamentalType::Uint)
