Add groom deformer radius writting.

diff --git a/Engine/Plugins/Runtime/HairStrands/Shaders/Private/DeformerDataInterfaceGroom.ush b/Engine/Plugins/Runtime/HairStrands/Shaders/Private/DeformerDataInterfaceGroom.ush
--- a/Engine/Plugins/Runtime/HairStrands/Shaders/Private/DeformerDataInterfaceGroom.ush
+++ b/Engine/Plugins/Runtime/HairStrands/Shaders/Private/DeformerDataInterfaceGroom.ush
@@ -88,7 +88,7 @@ float ReadSeed_{DataInterfaceName}(uint ControlPointIndex)
 
 uint ReadClumpId_{DataInterfaceName}(uint ControlPointIndex)
 {
-	return GetHairStrandsClumpID(ControlPointIndex);
+	return GetHairStrandsClumpID(ControlPointIndex).x;
 }
 
 float3 ReadColor_{DataInterfaceName}(uint ControlPointIndex)
diff --git a/Engine/Plugins/Runtime/HairStrands/Shaders/Private/DeformerDataInterfaceGroomWrite.ush b/Engine/Plugins/Runtime/HairStrands/Shaders/Private/DeformerDataInterfaceGroomWrite.ush
--- a/Engine/Plugins/Runtime/HairStrands/Shaders/Private/DeformerDataInterfaceGroomWrite.ush
+++ b/Engine/Plugins/Runtime/HairStrands/Shaders/Private/DeformerDataInterfaceGroomWrite.ush
@@ -3,8 +3,8 @@
 #include "/Plugin/ComputeFramework/Private/ComputeKernelCommon.ush"
 #include "/Engine/Private/HairStrands/HairStrandsVertexFactoryCommon.ush"
 
-uint {DataInterfaceName}_NumControlPoints;
-uint {DataInterfaceName}_NumCurves;
+HAIR_STRANDS_INSTANCE_PARAMETERS({DataInterfaceName})
+
 uint {DataInterfaceName}_OutputStreamStart;
 Buffer<float4>	{DataInterfaceName}_PositionOffsetBufferSRV;
 Buffer<uint4>   {DataInterfaceName}_PositionBufferSRV;
@@ -12,21 +12,45 @@ RWBuffer<uint4> {DataInterfaceName}_PositionBufferUAV;
 
 uint ReadNumControlPoints_{DataInterfaceName}()
 {
-	return {DataInterfaceName}_NumControlPoints;
+	return {DataInterfaceName}_PointCount;
 }
 
 uint ReadNumCurves_{DataInterfaceName}()
 {
-	return {DataInterfaceName}_NumCurves;
+	return {DataInterfaceName}_CurveCount;
 }
 
 void WritePosition_{DataInterfaceName}(uint ControlPointIndex, float3 Position)
 {
 	const uint WriteIndex = {DataInterfaceName}_OutputStreamStart + ControlPointIndex;
-	if (WriteIndex < {DataInterfaceName}_NumControlPoints)
+	if (WriteIndex < {DataInterfaceName}_PointCount)
+	{
+		const uint4 PackedCP = {DataInterfaceName}_PositionBufferSRV[WriteIndex];
+		const uint4 NewPackedCP = PackHairControlPointPosition(PackedCP, Position, {DataInterfaceName}_PositionOffsetBufferSRV[0].xyz);
+		{DataInterfaceName}_PositionBufferUAV[WriteIndex] = NewPackedCP;
+	}
+}
+
+void WriteRadius_{DataInterfaceName}(uint ControlPointIndex, float Radius)
+{
+	const uint WriteIndex = {DataInterfaceName}_OutputStreamStart + ControlPointIndex;
+	if (WriteIndex < {DataInterfaceName}_PointCount)
 	{
 		const uint4 PackedCP = {DataInterfaceName}_PositionBufferSRV[WriteIndex];
-		const uint4 NewPackedCP = PackHairControlPointPosition(PackedCP, Position, {DataInterfaceName}_PositionOffsetBufferSRV[0]);
+		const uint4 NewPackedCP = PackHairControlPointRadius(PackedCP, Radius, {DataInterfaceName}_Radius);
 		{DataInterfaceName}_PositionBufferUAV[WriteIndex] = NewPackedCP;
 	}
-}
\ No newline at end of file
+}
+
+void WritePositionAndRadius_{DataInterfaceName}(uint ControlPointIndex, float4 PositionAndRadius)
+{
+	const uint WriteIndex = {DataInterfaceName}_OutputStreamStart + ControlPointIndex;
+	if (WriteIndex < {DataInterfaceName}_PointCount)
+	{
+		const uint4 PackedCP = {DataInterfaceName}_PositionBufferSRV[WriteIndex];
+		uint4 NewPackedCP = 0;
+		NewPackedCP = PackHairControlPointPosition(PackedCP, PositionAndRadius.xyz, {DataInterfaceName}_PositionOffsetBufferSRV[0].xyz);
+		NewPackedCP = PackHairControlPointRadius(NewPackedCP, PositionAndRadius.w, {DataInterfaceName}_Radius);
+		{DataInterfaceName}_PositionBufferUAV[WriteIndex] = NewPackedCP;
+	}
+}
diff --git a/Engine/Plugins/Runtime/HairStrands/Source/HairStrandsCore/Private/Deformer/DeformerDataInterfaceGroom.cpp b/Engine/Plugins/Runtime/HairStrands/Source/HairStrandsCore/Private/Deformer/DeformerDataInterfaceGroom.cpp
--- a/Engine/Plugins/Runtime/HairStrands/Source/HairStrandsCore/Private/Deformer/DeformerDataInterfaceGroom.cpp
+++ b/Engine/Plugins/Runtime/HairStrands/Source/HairStrandsCore/Private/Deformer/DeformerDataInterfaceGroom.cpp
@@ -246,14 +246,10 @@ void FOptimusGroomDataProviderProxy::GatherDispatchData(FDispatchData const& InD
 
 		if (FHairGroupInstance* Instance = GroomComponent->GetGroupInstance(InvocationIndex))
 		{
-			const bool bIsSRVValid = Resources[InvocationIndex].PositionBuffer != nullptr;
-			const int32 NumControlPoints = bIsSRVValid ? Instance->Strands.Data->GetNumPoints() : 0;
-			const int32 NumCurves = bIsSRVValid ? Instance->Strands.Data->GetNumCurves() : 0;
-
-			const FHairGroupPublicData::FVertexFactoryInput VFInput = ComputeHairStrandsVertexInputData(Instance, EGroomViewMode::None);
-			
+			const FHairGroupPublicData::FVertexFactoryInput VFInput = ComputeHairStrandsVertexInputData(Instance, EGroomViewMode::None);		
 			Parameters.Common = VFInput.Strands.Common;
 
+			const bool bIsSRVValid = Resources[InvocationIndex].PositionBuffer != nullptr;
 			if (bIsSRVValid)
 			{
 				Parameters.Resources = Resources[InvocationIndex];
diff --git a/Engine/Plugins/Runtime/HairStrands/Source/HairStrandsCore/Private/Deformer/DeformerDataInterfaceGroomWrite.cpp b/Engine/Plugins/Runtime/HairStrands/Source/HairStrandsCore/Private/Deformer/DeformerDataInterfaceGroomWrite.cpp
--- a/Engine/Plugins/Runtime/HairStrands/Source/HairStrandsCore/Private/Deformer/DeformerDataInterfaceGroomWrite.cpp
+++ b/Engine/Plugins/Runtime/HairStrands/Source/HairStrandsCore/Private/Deformer/DeformerDataInterfaceGroomWrite.cpp
@@ -30,7 +30,9 @@ TArray<FOptimusCDIPinDefinition> UOptimusGroomWriteDataInterface::GetPinDefiniti
 	FName Curve(UOptimusGroomComponentSource::Domains::Curve);
 
 	TArray<FOptimusCDIPinDefinition> Defs;
-	Defs.Add({ "Position", "WritePosition", ControlPoint, "ReadNumControlPoints" });
+	Defs.Add({ "Position",         "WritePosition",            ControlPoint, "ReadNumControlPoints" });
+	Defs.Add({ "Radius",           "WriteRadius",              ControlPoint, "ReadNumControlPoints" });
+	Defs.Add({ "PositionAndRadius","WritePositionAndRadius",   ControlPoint, "ReadNumControlPoints" });
 	return Defs;
 }
 
@@ -54,11 +56,20 @@ void UOptimusGroomWriteDataInterface::GetSupportedOutputs(TArray<FShaderFunction
 		.SetName(TEXT("WritePosition"))
 		.AddParam(EShaderFundamentalType::Uint)
 		.AddParam(EShaderFundamentalType::Float, 3);
+
+	OutFunctions.AddDefaulted_GetRef()
+		.SetName(TEXT("WriteRadius"))
+		.AddParam(EShaderFundamentalType::Uint)
+		.AddParam(EShaderFundamentalType::Float);
+
+	OutFunctions.AddDefaulted_GetRef()
+		.SetName(TEXT("WritePositionAndRadius"))
+		.AddParam(EShaderFundamentalType::Uint)
+		.AddParam(EShaderFundamentalType::Float, 4);
 }
 
 BEGIN_SHADER_PARAMETER_STRUCT(FGroomWriteDataInterfaceParameters, )
-	SHADER_PARAMETER(uint32, NumControlPoints)
-	SHADER_PARAMETER(uint32, NumCurves)
+	SHADER_PARAMETER_STRUCT_INCLUDE(FHairStrandsInstanceCommonParameters, Common)
 	SHADER_PARAMETER(uint32, OutputStreamStart)
 	SHADER_PARAMETER_RDG_BUFFER_SRV(Buffer, PositionOffsetBufferSRV)
 	SHADER_PARAMETER_RDG_BUFFER_SRV(Buffer<uint4>, PositionBufferSRV)
@@ -163,6 +174,7 @@ void FOptimusGroomWriteDataProviderProxy::AllocateResources(FRDGBuilder& GraphBu
 	}
 }
 
+FHairGroupPublicData::FVertexFactoryInput ComputeHairStrandsVertexInputData(const FHairGroupInstance* Instance, EGroomViewMode ViewMode);
 void FOptimusGroomWriteDataProviderProxy::GatherDispatchData(FDispatchData const& InDispatchData)
 {
 	const TStridedView<FParameters> ParameterArray = MakeStridedParameterView<FParameters>(InDispatchData);
@@ -173,12 +185,11 @@ void FOptimusGroomWriteDataProviderProxy::GatherDispatchData(FDispatchData const
 			FResources& Resource = Resources[InvocationIndex];
 
 			const bool bValid = Resource.PositionBufferUAV != nullptr && Resource.PositionBufferSRV != nullptr && Resource.PositionOffsetBufferSRV != nullptr;
-			const int32 NumCurves = bValid ? GroomInstance->Strands.Data->GetNumCurves() : 0;
-			const int32 NumControlPoints = bValid ? GroomInstance->Strands.Data->GetNumPoints() : 0;
+
+			const FHairGroupPublicData::FVertexFactoryInput VFInput = ComputeHairStrandsVertexInputData(GroomInstance, EGroomViewMode::None);		
 
 			FParameters& Parameters = ParameterArray[InvocationIndex];
-			Parameters.NumControlPoints = NumControlPoints;
-			Parameters.NumCurves = NumCurves;
+			Parameters.Common = VFInput.Strands.Common;
 			Parameters.OutputStreamStart = 0;
 			Parameters.PositionOffsetBufferSRV = bValid ? Resource.PositionOffsetBufferSRV : Resource.PositionBufferSRV_fallback;
 			Parameters.PositionBufferSRV = bValid ? Resource.PositionBufferSRV : Resource.PositionBufferSRV_fallback;
diff --git a/Engine/Shaders/Private/HairStrands/HairStrandsInterpolation.usf b/Engine/Shaders/Private/HairStrands/HairStrandsInterpolation.usf
--- a/Engine/Shaders/Private/HairStrands/HairStrandsInterpolation.usf
+++ b/Engine/Shaders/Private/HairStrands/HairStrandsInterpolation.usf
@@ -443,7 +443,7 @@ void MainCS(uint2 DispatchThreadId : SV_DispatchThreadID)
 		const FRootTriangleData RenTri = GetRootTriangleData(VertexIndex, false);
 
 		// Manual decoding of the rest position
-		const uint4 PackedRestPosition = RenRestPosePositionBuffer[VertexIndex];
+		uint4 PackedRestPosition = RenRestPosePositionBuffer[VertexIndex];
 		const float3 RestPosition = UnpackHairControlPointPosition(PackedRestPosition, InRenHairPositionOffset);
 
 		float3 OutPosition = 0.0f;
@@ -480,11 +480,11 @@ void MainCS(uint2 DispatchThreadId : SV_DispatchThreadID)
 
 		// 3. Apply deformer
 		#if PERMUTATION_DEFORMER
-		//if (VertexIndex > 10000000)
 		{
 			const uint4 PackedDeformerPosition = RenDeformerPositionBuffer[VertexIndex];
 			const float3 DeformerPosition = UnpackHairControlPointPosition(PackedDeformerPosition, InRenHairPositionOffset);
 			OutPosition = ApplyDeformer(RenTri, RestPosition, DeformerPosition, OutPosition);
+			PackedRestPosition.w = PackedDeformerPosition.w;
 		}
 		#endif
 
diff --git a/Engine/Shaders/Private/HairStrands/HairStrandsPack.ush b/Engine/Shaders/Private/HairStrands/HairStrandsPack.ush
--- a/Engine/Shaders/Private/HairStrands/HairStrandsPack.ush
+++ b/Engine/Shaders/Private/HairStrands/HairStrandsPack.ush
@@ -64,6 +64,12 @@ uint4 PackHairControlPointPosition(uint4 InPackedData, float3 InNewPosition, flo
 	return uint4(f32tof16(InNewPosition.xyz - InPositionOffset), InPackedData.w);
 }
 
+uint4 PackHairControlPointRadius(uint4 InPackedData, float InNewWorldRadius, float InVFRadius)
+{
+	const uint PackedNewRadius = uint(clamp(saturate(InNewWorldRadius / InVFRadius) * 63u, 0u, 63u)) << 2u;	
+	return uint4(InPackedData.xyz, (InPackedData.w & (~0xFC)) | PackedNewRadius);
+}
+
 bool IsValidHairStrandsSegment(const FHairControlPoint CP0, const FHairControlPoint CP1)
 {
 	return !(CP0.Type == HAIR_CONTROLPOINT_END && CP1.Type == HAIR_CONTROLPOINT_START);
