Add hair attribute writing to groom deformer.

For a given groom, among optional attributes (e.g., color/roughness/AO/clumpID/...) only existing attributes can be written out.

diff --git a/Engine/Plugins/Runtime/HairStrands/Shaders/Private/DeformerDataInterfaceGroomWrite.ush b/Engine/Plugins/Runtime/HairStrands/Shaders/Private/DeformerDataInterfaceGroomWrite.ush
--- a/Engine/Plugins/Runtime/HairStrands/Shaders/Private/DeformerDataInterfaceGroomWrite.ush
+++ b/Engine/Plugins/Runtime/HairStrands/Shaders/Private/DeformerDataInterfaceGroomWrite.ush
@@ -5,10 +5,12 @@
 
 HAIR_STRANDS_INSTANCE_PARAMETERS({DataInterfaceName})
 
-uint {DataInterfaceName}_OutputStreamStart;
-Buffer<float4>	{DataInterfaceName}_PositionOffsetBufferSRV;
-Buffer<uint4>   {DataInterfaceName}_PositionBufferSRV;
-RWBuffer<uint4> {DataInterfaceName}_PositionBufferUAV;
+uint 				{DataInterfaceName}_OutputStreamStart;
+Buffer<float4>		{DataInterfaceName}_PositionOffsetBufferSRV;
+Buffer<uint4>   	{DataInterfaceName}_PositionBufferSRV;
+RWBuffer<uint4> 	{DataInterfaceName}_PositionBufferUAV;
+RWByteAddressBuffer {DataInterfaceName}_PointAttributeBufferUAV;
+RWByteAddressBuffer {DataInterfaceName}_CurveAttributeBufferUAV;
 
 uint ReadNumControlPoints_{DataInterfaceName}()
 {
@@ -54,3 +56,84 @@ void WritePositionAndRadius_{DataInterfaceName}(uint ControlPointIndex, float4 P
 		{DataInterfaceName}_PositionBufferUAV[WriteIndex] = NewPackedCP;
 	}
 }
+
+void InternalSetHairStrandsAttributes(uint Index, uint AttributeOffsetInBytes, uint AttributeStrideInBytes, RWByteAddressBuffer InAttributeBuffer, uint In)
+{
+	// Ensure the reading address is 4-byte aligned
+	const uint Address 					= AttributeOffsetInBytes + Index * AttributeStrideInBytes;
+	const uint Address_4BytesAligned 	= Address & (~0x3);
+	const uint ElementIndex 			= Address - Address_4BytesAligned;
+
+	// Create a mask where the 
+	const uint Mask 		  = (1u<<(AttributeStrideInBytes*8u))-1u;
+	const uint AttributeMask  = Mask << (ElementIndex*8);
+	const uint AttributeValue = In   << (ElementIndex*8);
+
+	// Clear previous value, and then write the value
+	InAttributeBuffer.InterlockedAnd(Address_4BytesAligned, ~AttributeMask);
+	InAttributeBuffer.InterlockedOr(Address_4BytesAligned, AttributeValue);
+}
+
+void InternalSetHairStrandsAttributes32(uint Index, uint AttributeOffsetInBytes, uint AttributeStrideInBytes, RWByteAddressBuffer InAttributeBuffer, uint In)
+{
+	// Ensure the reading address is 4-byte aligned
+	const uint Address = AttributeOffsetInBytes + Index * AttributeStrideInBytes;
+	const uint Address_4BytesAligned = Address & (~0x3);
+
+	// Clear previous value, and then write the value
+	InAttributeBuffer.Store(Address_4BytesAligned, In);
+}
+
+void WriteRootUV_{DataInterfaceName}(uint CurveIndex, float2 RootUV)
+{
+	const uint WriteIndex = {DataInterfaceName}_OutputStreamStart + CurveIndex;
+	if (WriteIndex < {DataInterfaceName}_CurveCount && HasHairStrandsRootUV())
+	{
+		InternalSetHairStrandsAttributes32(CurveIndex, HAIR_CURVE_ATTRIBUTE_OFFSET_ROOTUV({DataInterfaceName}_CurveAttributeOffsets), HAIR_CURVE_ATTRIBUTE_STRIDE_ROOTUV, {DataInterfaceName}_CurveAttributeBufferUAV, PackHairRootUV(RootUV));
+	}
+}
+
+void WriteSeed_{DataInterfaceName}(uint CurveIndex, float Seed)
+{
+	const uint WriteIndex = {DataInterfaceName}_OutputStreamStart + CurveIndex;
+	if (WriteIndex < {DataInterfaceName}_CurveCount)
+	{
+		InternalSetHairStrandsAttributes(CurveIndex, HAIR_CURVE_ATTRIBUTE_OFFSET_SEED({DataInterfaceName}_CurveAttributeOffsets), HAIR_CURVE_ATTRIBUTE_STRIDE_SEED, {DataInterfaceName}_CurveAttributeBufferUAV, PackHairSeed(Seed));
+	}
+}
+
+void WriteClumpId_{DataInterfaceName}(uint CurveIndex, uint ClumpId)
+{
+	const uint WriteIndex = {DataInterfaceName}_OutputStreamStart + CurveIndex;
+	if (WriteIndex < {DataInterfaceName}_CurveCount && HasHairStrandsClumpId())
+	{
+		InternalSetHairStrandsAttributes(CurveIndex, HAIR_CURVE_ATTRIBUTE_OFFSET_CLUMPID({DataInterfaceName}_CurveAttributeOffsets), HAIR_CURVE_ATTRIBUTE_STRIDE_CLUMPID, {DataInterfaceName}_CurveAttributeBufferUAV, PackHairClumpID(ClumpId));
+	}
+}
+
+void WriteColor_{DataInterfaceName}(uint PointIndex, float3 Color)
+{
+	const uint WriteIndex = {DataInterfaceName}_OutputStreamStart + PointIndex;
+	if (WriteIndex < {DataInterfaceName}_PointCount && HasHairStrandsColor())
+	{
+		InternalSetHairStrandsAttributes32(PointIndex, HAIR_POINT_ATTRIBUTE_OFFSET_COLOR({DataInterfaceName}_PointAttributeOffsets), HAIR_POINT_ATTRIBUTE_STRIDE_COLOR, {DataInterfaceName}_PointAttributeBufferUAV, PackHairColor(Color));
+	}
+}
+
+void WriteRoughness_{DataInterfaceName}(uint PointIndex, float Roughness)
+{
+	const uint WriteIndex = {DataInterfaceName}_OutputStreamStart + PointIndex;
+	if (WriteIndex < {DataInterfaceName}_PointCount && HasHairStrandsRoughness())
+	{
+		InternalSetHairStrandsAttributes(PointIndex, HAIR_POINT_ATTRIBUTE_OFFSET_ROUGHNESS({DataInterfaceName}_PointAttributeOffsets), HAIR_POINT_ATTRIBUTE_STRIDE_ROUGHNESS, {DataInterfaceName}_PointAttributeBufferUAV, PackHairRoughness(Roughness));
+	}
+}
+
+void WriteAO_{DataInterfaceName}(uint PointIndex, float Ao)
+{
+	const uint WriteIndex = {DataInterfaceName}_OutputStreamStart + PointIndex;
+	if (WriteIndex < {DataInterfaceName}_PointCount && HasHairStrandsAO())
+	{
+		InternalSetHairStrandsAttributes(PointIndex, HAIR_POINT_ATTRIBUTE_OFFSET_AO({DataInterfaceName}_PointAttributeOffsets), HAIR_POINT_ATTRIBUTE_STRIDE_AO, {DataInterfaceName}_PointAttributeBufferUAV, PackHairAO(Ao));
+	}
+}
\ No newline at end of file
diff --git a/Engine/Plugins/Runtime/HairStrands/Source/HairStrandsCore/Private/Deformer/DeformerDataInterfaceGroomWrite.cpp b/Engine/Plugins/Runtime/HairStrands/Source/HairStrandsCore/Private/Deformer/DeformerDataInterfaceGroomWrite.cpp
--- a/Engine/Plugins/Runtime/HairStrands/Source/HairStrandsCore/Private/Deformer/DeformerDataInterfaceGroomWrite.cpp
+++ b/Engine/Plugins/Runtime/HairStrands/Source/HairStrandsCore/Private/Deformer/DeformerDataInterfaceGroomWrite.cpp
@@ -13,6 +13,7 @@
 #include "GroomComponent.h"
 #include "GroomInstance.h"
 #include "DeformerGroomComponentSource.h"
+#include "RenderGraphUtils.h"
 
 FString UOptimusGroomWriteDataInterface::GetDisplayName() const
 {
@@ -30,24 +31,34 @@ TArray<FOptimusCDIPinDefinition> UOptimusGroomWriteDataInterface::GetPinDefiniti
 	FName Curve(UOptimusGroomComponentSource::Domains::Curve);
 
 	TArray<FOptimusCDIPinDefinition> Defs;
-	Defs.Add({ "Position",         "WritePosition",            ControlPoint, "ReadNumControlPoints" });
-	Defs.Add({ "Radius",           "WriteRadius",              ControlPoint, "ReadNumControlPoints" });
-	Defs.Add({ "PositionAndRadius","WritePositionAndRadius",   ControlPoint, "ReadNumControlPoints" });
+	Defs.Add({ "Position",         "WritePosition",         	ControlPoint, "ReadNumControlPoints" });
+	Defs.Add({ "Radius",           "WriteRadius",           	ControlPoint, "ReadNumControlPoints" });
+	Defs.Add({ "PositionAndRadius","WritePositionAndRadius",	ControlPoint, "ReadNumControlPoints" });
+
+	Defs.Add({ "RootUV",			"WriteRootUV",   			Curve, 		  "ReadNumCurves" });
+	Defs.Add({ "Seed",				"WriteSeed",   				Curve, 		  "ReadNumCurves" });
+	Defs.Add({ "ClumpId",			"WriteClumpId",   			Curve, 		  "ReadNumCurves" });
+	Defs.Add({ "Color",				"WriteColor",   			ControlPoint, "ReadNumControlPoints" });
+	Defs.Add({ "Roughness",			"WriteRoughness",   		ControlPoint, "ReadNumControlPoints" });
+	Defs.Add({ "AO",				"WriteAO",   				ControlPoint, "ReadNumControlPoints" });
+
 	return Defs;
 }
 
-
 TSubclassOf<UActorComponent> UOptimusGroomWriteDataInterface::GetRequiredComponentClass() const
 {
 	return UGroomComponent::StaticClass();
 }
 
-
 void UOptimusGroomWriteDataInterface::GetSupportedInputs(TArray<FShaderFunctionDefinition>& OutFunctions) const
 {
 	OutFunctions.AddDefaulted_GetRef()
 		.SetName(TEXT("ReadNumControlPoints"))
 		.AddReturnType(EShaderFundamentalType::Uint);
+
+	OutFunctions.AddDefaulted_GetRef()
+		.SetName(TEXT("ReadNumCurves"))
+		.AddReturnType(EShaderFundamentalType::Uint);
 }
 
 void UOptimusGroomWriteDataInterface::GetSupportedOutputs(TArray<FShaderFunctionDefinition>& OutFunctions) const
@@ -66,6 +77,36 @@ void UOptimusGroomWriteDataInterface::GetSupportedOutputs(TArray<FShaderFunction
 		.SetName(TEXT("WritePositionAndRadius"))
 		.AddParam(EShaderFundamentalType::Uint)
 		.AddParam(EShaderFundamentalType::Float, 4);
+
+	OutFunctions.AddDefaulted_GetRef()
+		.SetName(TEXT("WriteRootUV"))
+		.AddParam(EShaderFundamentalType::Uint)
+		.AddParam(EShaderFundamentalType::Float, 2);
+
+	OutFunctions.AddDefaulted_GetRef()
+		.SetName(TEXT("WriteSeed"))
+		.AddParam(EShaderFundamentalType::Uint)
+		.AddParam(EShaderFundamentalType::Float);
+
+	OutFunctions.AddDefaulted_GetRef()
+		.SetName(TEXT("WriteClumpId"))
+		.AddParam(EShaderFundamentalType::Uint)
+		.AddParam(EShaderFundamentalType::Float);
+
+	OutFunctions.AddDefaulted_GetRef()
+		.SetName(TEXT("WriteColor"))
+		.AddParam(EShaderFundamentalType::Uint)
+		.AddParam(EShaderFundamentalType::Float, 3);
+
+	OutFunctions.AddDefaulted_GetRef()
+		.SetName(TEXT("WriteRoughness"))
+		.AddParam(EShaderFundamentalType::Uint)
+		.AddParam(EShaderFundamentalType::Float);
+
+	OutFunctions.AddDefaulted_GetRef()
+		.SetName(TEXT("WriteAO"))
+		.AddParam(EShaderFundamentalType::Uint)
+		.AddParam(EShaderFundamentalType::Float);
 }
 
 BEGIN_SHADER_PARAMETER_STRUCT(FGroomWriteDataInterfaceParameters, )
@@ -74,6 +115,8 @@ BEGIN_SHADER_PARAMETER_STRUCT(FGroomWriteDataInterfaceParameters, )
 	SHADER_PARAMETER_RDG_BUFFER_SRV(Buffer, PositionOffsetBufferSRV)
 	SHADER_PARAMETER_RDG_BUFFER_SRV(Buffer<uint4>, PositionBufferSRV)
 	SHADER_PARAMETER_RDG_BUFFER_UAV(RWBuffer<uint4>, PositionBufferUAV)
+	SHADER_PARAMETER_RDG_BUFFER_UAV(RWByteAddressBuffer, CurveAttributeBufferUAV)
+	SHADER_PARAMETER_RDG_BUFFER_UAV(RWByteAddressBuffer, PointAttributeBufferUAV)
 END_SHADER_PARAMETER_STRUCT()
 
 void UOptimusGroomWriteDataInterface::GetShaderParameters(TCHAR const* UID, FShaderParametersMetadataBuilder& InOutBuilder, FShaderParametersMetadataAllocations& InOutAllocations) const
@@ -154,21 +197,60 @@ void FOptimusGroomWriteDataProviderProxy::AllocateResources(FRDGBuilder& GraphBu
 			const int32 NumCurves = GroomInstance->Strands.Data->GetNumCurves();
 
 			FResources& R = Resources.AddDefaulted_GetRef();
-			if (OutputMask & 1)
+
+			FRDGBufferRef PositionBuffer_fallback = GraphBuilder.CreateBuffer(FRDGBufferDesc::CreateBufferDesc(8, 1), TEXT("Groom.DeformedPositionBuffer"), ERDGBufferFlags::None);
+			R.PositionBufferSRV_fallback = GraphBuilder.CreateSRV(GraphBuilder.RegisterExternalBuffer(GWhiteVertexBufferWithRDG->Buffer), PF_R16G16B16A16_UINT);
+			R.PositionBufferUAV_fallback = GraphBuilder.CreateUAV(PositionBuffer_fallback, PF_R16G16B16A16_UINT, ERDGUnorderedAccessViewFlags::SkipBarrier);
+
+			FRDGBufferRef AttributeBuffer_fallback = GraphBuilder.CreateBuffer(FRDGBufferDesc::CreateByteAddressDesc(16), TEXT("Groom.DeformedAttributeBuffer"), ERDGBufferFlags::None);
+			R.AttributeBufferUAV_fallback = GraphBuilder.CreateUAV(AttributeBuffer_fallback, PF_Unknown, ERDGUnorderedAccessViewFlags::SkipBarrier);
+
+			// Positions / Radius
+			if (OutputMask & 0x7)
 			{
-				R.PositionOffsetBufferSRV = Register(GraphBuilder, GroomInstance->Strands.RestResource->PositionOffsetBuffer, ERDGImportedBufferFlags::CreateSRV).SRV;
-				R.PositionBufferSRV = Register(GraphBuilder, GroomInstance->Strands.RestResource->PositionBuffer, ERDGImportedBufferFlags::CreateSRV).SRV;
-				R.PositionBufferUAV = Register(GraphBuilder, GroomInstance->Strands.DeformedResource->GetDeformerBuffer(GraphBuilder), ERDGImportedBufferFlags::CreateUAV).UAV;
-				
-				FRDGBufferRef PositionBuffer_fallback = GraphBuilder.CreateBuffer(FRDGBufferDesc::CreateBufferDesc(8, 1), TEXT("Groom.DeformedPositionBuffer"), ERDGBufferFlags::None);
-				R.PositionBufferSRV_fallback = GraphBuilder.CreateSRV(GraphBuilder.RegisterExternalBuffer(GWhiteVertexBufferWithRDG->Buffer), PF_R16G16B16A16_UINT);
-				R.PositionBufferUAV_fallback = GraphBuilder.CreateUAV(PositionBuffer_fallback, PF_R16G16B16A16_UINT, ERDGUnorderedAccessViewFlags::SkipBarrier);
+				R.PositionOffsetBufferSRV 	= Register(GraphBuilder, GroomInstance->Strands.RestResource->PositionOffsetBuffer, ERDGImportedBufferFlags::CreateSRV).SRV;
+				R.PositionBufferSRV 		= Register(GraphBuilder, GroomInstance->Strands.RestResource->PositionBuffer, ERDGImportedBufferFlags::CreateSRV).SRV;
+				R.PositionBufferUAV 		= Register(GraphBuilder, GroomInstance->Strands.DeformedResource->GetDeformerBuffer(GraphBuilder), ERDGImportedBufferFlags::CreateUAV).UAV;
 			}
 			else
 			{
-				R.PositionOffsetBufferSRV = GraphBuilder.CreateSRV(GraphBuilder.RegisterExternalBuffer(GWhiteVertexBufferWithRDG->Buffer), PF_A32B32G32R32F);
-				R.PositionBufferSRV = GraphBuilder.CreateSRV(GraphBuilder.RegisterExternalBuffer(GWhiteVertexBufferWithRDG->Buffer), PF_R16G16B16A16_UINT);
-				R.PositionBufferUAV = GraphBuilder.CreateUAV(GraphBuilder.RegisterExternalBuffer(GWhiteVertexBufferWithRDG->Buffer), PF_R16G16B16A16_UINT);
+				R.PositionOffsetBufferSRV 	= GraphBuilder.CreateSRV(GraphBuilder.RegisterExternalBuffer(GWhiteVertexBufferWithRDG->Buffer), PF_A32B32G32R32F);
+				R.PositionBufferSRV 		= GraphBuilder.CreateSRV(GraphBuilder.RegisterExternalBuffer(GWhiteVertexBufferWithRDG->Buffer), PF_R16G16B16A16_UINT);
+				R.PositionBufferUAV 		= GraphBuilder.CreateUAV(GraphBuilder.RegisterExternalBuffer(GWhiteVertexBufferWithRDG->Buffer), PF_R16G16B16A16_UINT);
+			}
+			
+			// Curve Attributes
+			if (OutputMask & 0x38)
+			{
+				FRDGExternalBuffer CurveAttributeBufferExt = GroomInstance->Strands.DeformedResource->GetDeformerCurveAttributeBuffer(GraphBuilder);
+				if (CurveAttributeBufferExt.Buffer)
+				{
+					FRDGImportedBuffer CurveAttributeBuffer = Register(GraphBuilder, CurveAttributeBufferExt, ERDGImportedBufferFlags::CreateUAV);
+					AddCopyBufferPass(GraphBuilder, CurveAttributeBuffer.Buffer, Register(GraphBuilder, GroomInstance->Strands.RestResource->CurveAttributeBuffer, ERDGImportedBufferFlags::None).Buffer);
+					R.CurveAttributeBufferUAV = CurveAttributeBuffer.UAV;
+				}
+			}
+			if (R.CurveAttributeBufferUAV == nullptr)
+			{
+				R.CurveAttributeBufferUAV = R.AttributeBufferUAV_fallback;
+				//R.CurveAttributeBufferUAV = GraphBuilder.CreateUAV(GraphBuilder.RegisterExternalBuffer(GWhiteVertexBufferWithRDG->Buffer), PF_R32_UINT);
+			}
+
+			// Point Attributes
+			if (OutputMask & 0x1c0)
+			{
+				FRDGExternalBuffer PointAttributeBufferExt = GroomInstance->Strands.DeformedResource->GetDeformerPointAttributeBuffer(GraphBuilder);
+				if (PointAttributeBufferExt.Buffer)
+				{
+					FRDGImportedBuffer PointAttributeBuffer = Register(GraphBuilder, PointAttributeBufferExt, ERDGImportedBufferFlags::CreateUAV);
+					AddCopyBufferPass(GraphBuilder, PointAttributeBuffer.Buffer, Register(GraphBuilder, GroomInstance->Strands.RestResource->PointAttributeBuffer, ERDGImportedBufferFlags::None).Buffer);
+					R.PointAttributeBufferUAV = PointAttributeBuffer.UAV;
+				}
+			}
+			if (R.PointAttributeBufferUAV == nullptr)
+			{
+				R.PointAttributeBufferUAV = R.AttributeBufferUAV_fallback;
+				//R.PointAttributeBufferUAV = GraphBuilder.CreateUAV(GraphBuilder.RegisterExternalBuffer(GWhiteVertexBufferWithRDG->Buffer), PF_R32_UINT);
 			}
 		}	
 	}
@@ -189,11 +271,13 @@ void FOptimusGroomWriteDataProviderProxy::GatherDispatchData(FDispatchData const
 			const FHairGroupPublicData::FVertexFactoryInput VFInput = ComputeHairStrandsVertexInputData(GroomInstance, EGroomViewMode::None);		
 
 			FParameters& Parameters = ParameterArray[InvocationIndex];
-			Parameters.Common = VFInput.Strands.Common;
-			Parameters.OutputStreamStart = 0;
-			Parameters.PositionOffsetBufferSRV = bValid ? Resource.PositionOffsetBufferSRV : Resource.PositionBufferSRV_fallback;
-			Parameters.PositionBufferSRV = bValid ? Resource.PositionBufferSRV : Resource.PositionBufferSRV_fallback;
-			Parameters.PositionBufferUAV = bValid ? Resource.PositionBufferUAV : Resource.PositionBufferUAV_fallback;
+			Parameters.Common 					= VFInput.Strands.Common;
+			Parameters.OutputStreamStart 		= 0;
+			Parameters.PositionOffsetBufferSRV 	= bValid ? Resource.PositionOffsetBufferSRV : Resource.PositionBufferSRV_fallback;
+			Parameters.PositionBufferSRV 		= bValid ? Resource.PositionBufferSRV : Resource.PositionBufferSRV_fallback;
+			Parameters.PositionBufferUAV 		= bValid ? Resource.PositionBufferUAV : Resource.PositionBufferUAV_fallback;
+			Parameters.PointAttributeBufferUAV 	= bValid ? Resource.PointAttributeBufferUAV : Resource.AttributeBufferUAV_fallback;
+			Parameters.CurveAttributeBufferUAV 	= bValid ? Resource.CurveAttributeBufferUAV : Resource.AttributeBufferUAV_fallback;
 		}
 	}
 }
\ No newline at end of file
diff --git a/Engine/Plugins/Runtime/HairStrands/Source/HairStrandsCore/Private/Deformer/DeformerDataInterfaceGroomWrite.h b/Engine/Plugins/Runtime/HairStrands/Source/HairStrandsCore/Private/Deformer/DeformerDataInterfaceGroomWrite.h
--- a/Engine/Plugins/Runtime/HairStrands/Source/HairStrandsCore/Private/Deformer/DeformerDataInterfaceGroomWrite.h
+++ b/Engine/Plugins/Runtime/HairStrands/Source/HairStrandsCore/Private/Deformer/DeformerDataInterfaceGroomWrite.h
@@ -83,6 +83,10 @@ private:
 		FRDGBufferUAV* PositionBufferUAV = nullptr;
 		FRDGBufferSRV* PositionBufferSRV_fallback = nullptr;
 		FRDGBufferUAV* PositionBufferUAV_fallback = nullptr;
+
+		FRDGBufferUAV* PointAttributeBufferUAV = nullptr;
+		FRDGBufferUAV* CurveAttributeBufferUAV = nullptr;
+		FRDGBufferUAV* AttributeBufferUAV_fallback = nullptr;
 	};
 	TArray<FResources> Resources;
 };
\ No newline at end of file
diff --git a/Engine/Plugins/Runtime/HairStrands/Source/HairStrandsCore/Private/GroomResources.cpp b/Engine/Plugins/Runtime/HairStrands/Source/HairStrandsCore/Private/GroomResources.cpp
--- a/Engine/Plugins/Runtime/HairStrands/Source/HairStrandsCore/Private/GroomResources.cpp
+++ b/Engine/Plugins/Runtime/HairStrands/Source/HairStrandsCore/Private/GroomResources.cpp
@@ -520,6 +520,17 @@ void InternalCreateByteAddressBufferRDG(FRDGBuilder& GraphBuilder, const TArray<
 	}
 }
 
+void InternalCreateByteAddressBufferRDG(FRDGBuilder& GraphBuilder, uint64 DataSizeInBytes, FRDGExternalBuffer& Out, const TCHAR* DebugName, const FName& OwnerName, EHairResourceUsageType UsageType)
+{
+	Out.Buffer = nullptr;
+	if (DataSizeInBytes != 0)
+	{
+		const FRDGBufferDesc Desc = ApplyUsage(FRDGBufferDesc::CreateByteAddressDesc(DataSizeInBytes), UsageType);
+		FRDGBufferRef Buffer = GraphBuilder.CreateBuffer(Desc, DebugName, ERDGBufferFlags::MultiFrame);
+		Buffer->SetOwnerName(OwnerName);
+		ConvertToExternalBufferWithViews(GraphBuilder, Buffer, Out, PF_Unknown);
+	}
+}
 /////////////////////////////////////////////////////////////////////////////////////////
 
 static UTexture2D* CreateCardTexture(FIntPoint Resolution)
@@ -1037,12 +1048,36 @@ FRDGExternalBuffer& FHairStrandsDeformedResource::GetDeformerBuffer(FRDGBuilder&
 	return DeformerBuffer;
 }
 
+FRDGExternalBuffer& FHairStrandsDeformedResource::GetDeformerCurveAttributeBuffer(FRDGBuilder& GraphBuilder)
+{
+	// Deformer curve attributes
+	if (DeformerCurveAttributeBuffer.Buffer == nullptr)
+	{		
+		check(BulkData.Data.CurveAttributes.Data.GetBulkDataSize() > 0);
+		InternalCreateByteAddressBufferRDG(GraphBuilder, BulkData.Data.CurveAttributes.Data.GetBulkDataSize(), DeformerCurveAttributeBuffer, ToHairResourceDebugName(HAIRSTRANDS_RESOUCE_NAME(CurveType, Hair.StrandsDeformed_DeformerCurveAttributeBuffer), ResourceName), OwnerName, EHairResourceUsageType::Dynamic);
+	}
+	return DeformerCurveAttributeBuffer;
+}
+
+FRDGExternalBuffer& FHairStrandsDeformedResource::GetDeformerPointAttributeBuffer(FRDGBuilder& GraphBuilder)
+{
+	// Deformer point attributes
+	if (DeformerPointAttributeBuffer.Buffer == nullptr && (BulkData.Header.Flags & FHairStrandsBulkData::DataFlags_HasPointAttribute))
+	{
+		check(BulkData.Data.PointAttributes.Data.GetBulkDataSize() > 0);
+		InternalCreateByteAddressBufferRDG(GraphBuilder, BulkData.Data.PointAttributes.Data.GetBulkDataSize(), DeformerPointAttributeBuffer, ToHairResourceDebugName(HAIRSTRANDS_RESOUCE_NAME(CurveType, Hair.StrandsDeformedt_DeformerPointAttributeBuffer), ResourceName), OwnerName, EHairResourceUsageType::Dynamic);
+	}
+	return DeformerPointAttributeBuffer;
+}
+
 void FHairStrandsDeformedResource::InternalRelease()
 {
 	DeformedPositionBuffer[0].Release();
 	DeformedPositionBuffer[1].Release();
 	TangentBuffer.Release();
 	DeformerBuffer.Release();
+	DeformerPointAttributeBuffer.Release();
+	DeformerCurveAttributeBuffer.Release();
 
 	DeformedOffsetBuffer[0].Release();
 	DeformedOffsetBuffer[1].Release();
diff --git a/Engine/Plugins/Runtime/HairStrands/Source/HairStrandsCore/Private/HairStrandsInterpolation.cpp b/Engine/Plugins/Runtime/HairStrands/Source/HairStrandsCore/Private/HairStrandsInterpolation.cpp
--- a/Engine/Plugins/Runtime/HairStrands/Source/HairStrandsCore/Private/HairStrandsInterpolation.cpp
+++ b/Engine/Plugins/Runtime/HairStrands/Source/HairStrandsCore/Private/HairStrandsInterpolation.cpp
@@ -1629,12 +1629,20 @@ static FHairGroupPublicData::FVertexFactoryInput InternalComputeHairStrandsVerte
 			OutVFInput.Strands.Common.PositionOffset	 = FVector3f(Instance->Strands.DeformedResource->GetPositionOffset(FHairStrandsDeformedResource::EFrameType::Current)); // LWC precision loss
 			OutVFInput.Strands.Common.PrevPositionOffset = FVector3f(Instance->Strands.DeformedResource->GetPositionOffset(FHairStrandsDeformedResource::EFrameType::Current)); // LWC precision loss
 		}
+		const bool bUseDebugCurveBuffer = (ViewMode == EGroomViewMode::RenderHairStrands || ViewMode == EGroomViewMode::Cluster || ViewMode == EGroomViewMode::ClusterAABB);
 		CONVERT_HAIRSSTRANDS_VF_PARAMETERS(OutVFInput.Strands.TangentBuffer,			Instance->Strands.DeformedResource->TangentBuffer);
-		CONVERT_HAIRSSTRANDS_VF_PARAMETERS(OutVFInput.Strands.CurveAttributeBuffer,		(ViewMode == EGroomViewMode::RenderHairStrands || ViewMode == EGroomViewMode::Cluster || ViewMode == EGroomViewMode::ClusterAABB) ? Instance->Strands.DebugCurveAttributeBuffer : Instance->Strands.RestResource->CurveAttributeBuffer);
+		CONVERT_HAIRSSTRANDS_VF_PARAMETERS(OutVFInput.Strands.CurveAttributeBuffer,		bUseDebugCurveBuffer ? Instance->Strands.DebugCurveAttributeBuffer : Instance->Strands.RestResource->CurveAttributeBuffer);
 		CONVERT_HAIRSSTRANDS_VF_PARAMETERS(OutVFInput.Strands.PointAttributeBuffer,		Instance->Strands.RestResource->PointAttributeBuffer);
 		CONVERT_HAIRSSTRANDS_VF_PARAMETERS(OutVFInput.Strands.PointToCurveBuffer,		Instance->Strands.RestResource->PointToCurveBuffer);
 		CONVERT_HAIRSSTRANDS_VF_PARAMETERS(OutVFInput.Strands.CurveBuffer,				Instance->Strands.RestResource->CurveBuffer);
 
+		// If using groom deformer, replace point/curve attribute with their deformer version
+		if (!bUseDebugCurveBuffer)
+		{
+			if (Instance->Strands.DeformedResource->DeformerCurveAttributeBuffer.Buffer) { CONVERT_HAIRSSTRANDS_VF_PARAMETERS(OutVFInput.Strands.CurveAttributeBuffer,	Instance->Strands.DeformedResource->DeformerCurveAttributeBuffer); }
+			if (Instance->Strands.DeformedResource->DeformerPointAttributeBuffer.Buffer) { CONVERT_HAIRSSTRANDS_VF_PARAMETERS(OutVFInput.Strands.PointAttributeBuffer,	Instance->Strands.DeformedResource->DeformerPointAttributeBuffer); }
+		}
+
 		OutVFInput.Strands.Common.PointCount = Instance->HairGroupPublicData->GetActiveStrandsPointCount();
 		OutVFInput.Strands.Common.CurveCount = Instance->HairGroupPublicData->GetActiveStrandsCurveCount();
 	}
diff --git a/Engine/Plugins/Runtime/HairStrands/Source/HairStrandsCore/Public/GroomResources.h b/Engine/Plugins/Runtime/HairStrands/Source/HairStrandsCore/Public/GroomResources.h
--- a/Engine/Plugins/Runtime/HairStrands/Source/HairStrandsCore/Public/GroomResources.h
+++ b/Engine/Plugins/Runtime/HairStrands/Source/HairStrandsCore/Public/GroomResources.h
@@ -411,6 +411,8 @@ struct FHairStrandsDeformedResource : public FHairCommonResource
 
 	/* Strand hair deformer buffer. This buffer is optionally created & filled in by a mesh deformer */
 	FRDGExternalBuffer DeformerBuffer;
+	FRDGExternalBuffer DeformerPointAttributeBuffer;
+	FRDGExternalBuffer DeformerCurveAttributeBuffer;
 
 	/* Position offset as the deformed positions are expressed in relative coordinate (16bits) */
 	FVector PositionOffset[2] = {FVector::ZeroVector, FVector::ZeroVector};
@@ -448,8 +450,10 @@ struct FHairStrandsDeformedResource : public FHairCommonResource
 	inline uint32 GetUniqueViewID(EFrameType T) const			{ return UniqueViewIDs[GetIndex(T)]; }
 	//bool NeedsToUpdateTangent();
 
-	// Return deformer buffer
+	// Return deformer buffers
 	FRDGExternalBuffer& GetDeformerBuffer(FRDGBuilder& GraphBuilder);
+	FRDGExternalBuffer& GetDeformerCurveAttributeBuffer(FRDGBuilder& GraphBuilder);
+	FRDGExternalBuffer& GetDeformerPointAttributeBuffer(FRDGBuilder& GraphBuilder);
 };
 
 struct FHairStrandsClusterCullingResource : public FHairCommonResource
diff --git a/Engine/Shaders/Private/HairStrands/HairStrandsAttributeTemplate.ush b/Engine/Shaders/Private/HairStrands/HairStrandsAttributeTemplate.ush
--- a/Engine/Shaders/Private/HairStrands/HairStrandsAttributeTemplate.ush
+++ b/Engine/Shaders/Private/HairStrands/HairStrandsAttributeTemplate.ush
@@ -100,8 +100,8 @@ float GetHairStrandsAO(uint InPrimitiveId, float InSegmentU)
 	const uint AoOffset = HAIR_POINT_ATTRIBUTE_OFFSET_AO(PREFIX(PointAttributeOffsets));
 	if (IsAttributeValid(AoOffset))
 	{
-		const float AO0 = UnpackHairRoughness(InternalGetHairStrandsAttributesPerVertex(InPrimitiveId, AoOffset, HAIR_POINT_ATTRIBUTE_STRIDE_AO, PREFIX(PointAttributeBuffer)));
-		const float AO1 = UnpackHairRoughness(InternalGetHairStrandsAttributesPerVertex(InPrimitiveId+1, AoOffset, HAIR_POINT_ATTRIBUTE_STRIDE_AO, PREFIX(PointAttributeBuffer)));
+		const float AO0 = UnpackHairAO(InternalGetHairStrandsAttributesPerVertex(InPrimitiveId, AoOffset, HAIR_POINT_ATTRIBUTE_STRIDE_AO, PREFIX(PointAttributeBuffer)));
+		const float AO1 = UnpackHairAO(InternalGetHairStrandsAttributesPerVertex(InPrimitiveId+1, AoOffset, HAIR_POINT_ATTRIBUTE_STRIDE_AO, PREFIX(PointAttributeBuffer)));
 		return lerp(AO0, AO1, InSegmentU);
 	}
 	else
diff --git a/Engine/Shaders/Private/HairStrands/HairStrandsPack.ush b/Engine/Shaders/Private/HairStrands/HairStrandsPack.ush
--- a/Engine/Shaders/Private/HairStrands/HairStrandsPack.ush
+++ b/Engine/Shaders/Private/HairStrands/HairStrandsPack.ush
@@ -78,6 +78,14 @@ bool IsValidHairStrandsSegment(const FHairControlPoint CP0, const FHairControlPo
 ////////////////////////////////////////////////////////////////////////////////////////////////////////////////////////////////
 // Hair attributes
 
+uint PackHairRootUV(float2 In)
+{
+	In.y = 1.0f - In.y;
+
+	return	(min(uint(saturate(In.x) * 0x7FF), 0x7FF)      )|
+			(min(uint(saturate(In.y) * 0x7FF), 0x7FF) << 11);
+}
+
 float2 UnpackHairRootUV(uint In)
 {
 	float2 RootUV, RootIndex;
@@ -89,7 +97,7 @@ float2 UnpackHairRootUV(uint In)
 	// Invert V to compensate image origin flip. Similar to GetHairStrandsRootUV in HairCardsAttributeCommon.ush
 	RootUV.y = 1.0f - RootUV.y;
 
-	return RootUV = RootUV + RootIndex;
+	return RootUV + RootIndex;
 }
 
 float UnpackHairLength(uint In)
@@ -97,9 +105,19 @@ float UnpackHairLength(uint In)
 	return f16tof32(In & 0xFFFF);
 }
 
+uint PackHairSeed(float In)
+{
+	return PackR8(In);
+}
+
 float UnpackHairSeed(uint In)
 {
-	return (In & 0xFF) / 255.f;
+	return UnpackR8(In);
+}
+
+uint PackHairClumpID(uint In)
+{
+	return min(In, 0xFFFF);
 }
 
 uint3 UnpackHairClumpID(uint In)
@@ -112,6 +130,13 @@ uint3 UnpackHairClumpID(uint2 In)
 	return uint3(In.x & 0xFFFF, (In.x>>16) & 0xFFFF, (In.x>>32) & 0xFFFF);
 }
 
+uint PackHairColor(float3 In)
+{
+	return	(min(uint(In.x * 0x7FF), 0x7FF)    )|
+			(min(uint(In.y * 0x7FF), 0x7FF)<<11)|
+			(min(uint(In.z * 0x3FF), 0x3FF)<<22);
+}
+
 float3 UnpackHairColor(uint In)
 {
 	float3 Out;
@@ -121,9 +146,24 @@ float3 UnpackHairColor(uint In)
 	return Out;
 }
 
+uint PackHairRoughness(float In)
+{
+	return PackR8(In);
+}
+
 float UnpackHairRoughness(uint In)
 {
-	return (In & 0xFF) / 255.f;
+	return UnpackR8(In);
+}
+
+uint PackHairAO(float In)
+{
+	return PackR8(In);
+}
+
+float UnpackHairAO(uint In)
+{
+	return UnpackR8(In);
 }
 
 ////////////////////////////////////////////////////////////////////////////////////////////////////////////////////////////////
