Add groom guide reading from mesh deformer when available.

diff --git a/Engine/Plugins/Runtime/HairStrands/Shaders/Private/DeformerDataInterfaceGroom.ush b/Engine/Plugins/Runtime/HairStrands/Shaders/Private/DeformerDataInterfaceGroom.ush
--- a/Engine/Plugins/Runtime/HairStrands/Shaders/Private/DeformerDataInterfaceGroom.ush
+++ b/Engine/Plugins/Runtime/HairStrands/Shaders/Private/DeformerDataInterfaceGroom.ush
@@ -12,6 +12,8 @@ HAIR_STRANDS_INSTANCE_PARAMETERS({DataInterfaceName})
 #define HAIR_STRANDS_ATTRIBUTE_ACCESSORS(Name) {DataInterfaceName}_##Name
 #include "/Engine/Private/HairStrands/HairStrandsAttributeTemplate.ush"
 
+ByteAddressBuffer {DataInterfaceName}_InterpolationBuffer;
+
 ///////////////////////////////////////////////////////////////////////////////////////////////////
 // Helpers
 
@@ -114,4 +116,19 @@ float ReadCurveOffsetPoint_{DataInterfaceName}(uint CurveIndex)
 float ReadCurveNumPoint_{DataInterfaceName}(uint CurveIndex)
 {
 	return ReadCurveData(CurveIndex).PointCount;
-}
\ No newline at end of file
+}
+
+uint ReadGuideIndex_{DataInterfaceName}(uint ControlPointIndex)
+{
+	FGuideData GuideData;
+	if ({DataInterfaceName}_bSingleGuide)
+	{
+		GuideData = UnpackGuideData({DataInterfaceName}_InterpolationBuffer.Load(ControlPointIndex * HAIR_INTERPOLATION_1GUIDE_STRIDE));
+	}
+	else
+	{
+		GuideData = UnpackGuideData({DataInterfaceName}_InterpolationBuffer.Load4(ControlPointIndex * HAIR_INTERPOLATION_3GUIDE_STRIDE));
+	}
+
+	return GuideData.VertexIndices.x;
+}
diff --git a/Engine/Plugins/Runtime/HairStrands/Shaders/Private/DeformerDataInterfaceGroomGuide.ush b/Engine/Plugins/Runtime/HairStrands/Shaders/Private/DeformerDataInterfaceGroomGuide.ush
new file mode 100644
--- /dev/null
+++ b/Engine/Plugins/Runtime/HairStrands/Shaders/Private/DeformerDataInterfaceGroomGuide.ush
@@ -0,0 +1,69 @@
+// Copyright Epic Games, Inc. All Rights Reserved.
+
+#include "/Plugin/ComputeFramework/Private/ComputeKernelCommon.ush"
+#include "/Engine/Private/HairStrands/HairStrandsVertexFactoryCommon.ush"
+#include "/Engine/Private/HairStrands/HairStrandsAttributeCommon.ush"
+
+///////////////////////////////////////////////////////////////////////////////////////////////////
+// Bindings
+
+HAIR_STRANDS_INSTANCE_PARAMETERS({DataInterfaceName})
+
+///////////////////////////////////////////////////////////////////////////////////////////////////
+// Helpers
+
+FHairControlPoint ReadControlPointData_{DataInterfaceName}(uint ControlPointIndex)
+{
+	const uint MaxControlPoints = {DataInterfaceName}_PointCount;
+	if (ControlPointIndex >= MaxControlPoints)
+	{
+		return (FHairControlPoint)0;
+	}
+
+	// Hair deformer work in local space so no need for position offset
+	return UnpackHairControlPoint(
+		{DataInterfaceName}_PositionBuffer[ControlPointIndex], 
+		{DataInterfaceName}_PositionOffsetBuffer[0].xyz,
+		{DataInterfaceName}_Radius, 
+		{DataInterfaceName}_RootScale, 
+		{DataInterfaceName}_TipScale);
+}
+
+FHairCurve ReadCurveData_{DataInterfaceName}(uint CurveIndex)
+{
+	const uint MaxCurve = {DataInterfaceName}_CurveCount;
+	if (CurveIndex >= MaxCurve)
+	{
+		return (FHairCurve)0;
+	}
+
+	return UnpackHairCurve({DataInterfaceName}_CurveBuffer[CurveIndex]);
+}
+
+///////////////////////////////////////////////////////////////////////////////////////////////////
+// Interface
+
+uint ReadNumGuidePoints_{DataInterfaceName}()
+{
+	return {DataInterfaceName}_PointCount;
+}
+
+uint ReadNumGuideCurves_{DataInterfaceName}()
+{
+	return {DataInterfaceName}_CurveCount;
+}
+
+float3 ReadPosition_{DataInterfaceName}(uint ControlPointIndex)
+{
+	return ReadControlPointData_{DataInterfaceName}(ControlPointIndex).Position;
+}
+
+float ReadCurveOffsetPoint_{DataInterfaceName}(uint CurveIndex)
+{
+	return ReadCurveData_{DataInterfaceName}(CurveIndex).PointIndex;
+}
+
+float ReadCurveNumPoint_{DataInterfaceName}(uint CurveIndex)
+{
+	return ReadCurveData_{DataInterfaceName}(CurveIndex).PointCount;
+}
\ No newline at end of file
diff --git a/Engine/Plugins/Runtime/HairStrands/Source/HairStrandsCore/Private/Deformer/DeformerDataInterfaceGroom.cpp b/Engine/Plugins/Runtime/HairStrands/Source/HairStrandsCore/Private/Deformer/DeformerDataInterfaceGroom.cpp
--- a/Engine/Plugins/Runtime/HairStrands/Source/HairStrandsCore/Private/Deformer/DeformerDataInterfaceGroom.cpp
+++ b/Engine/Plugins/Runtime/HairStrands/Source/HairStrandsCore/Private/Deformer/DeformerDataInterfaceGroom.cpp
@@ -39,6 +39,7 @@ TArray<FOptimusCDIPinDefinition> UOptimusGroomDataInterface::GetPinDefinitions()
 	Defs.Add({ "AO",               "ReadAO",                ControlPoint,   "ReadAO" });
 	Defs.Add({ "CurveOffsetPoint", "ReadCurveOffsetPoint",  Curve,          "ReadCurveOffsetPoint" });
 	Defs.Add({ "CurveNumPoint",    "ReadCurveNumPoint",     Curve,          "ReadCurveNumPoint" });
+	Defs.Add({ "GuideIndex",       "ReadGuideIndex",        ControlPoint,   "ReadGuideIndex" });
 	return Defs;
 }
  
@@ -116,11 +117,17 @@ void UOptimusGroomDataInterface::GetSupportedInputs(TArray<FShaderFunctionDefini
 		.SetName(TEXT("ReadCurveNumPoint"))
 		.AddReturnType(EShaderFundamentalType::Uint)
 		.AddParam(EShaderFundamentalType::Uint);
+
+	OutFunctions.AddDefaulted_GetRef()
+		.SetName(TEXT("ReadGuideIndex"))
+		.AddReturnType(EShaderFundamentalType::Uint)
+		.AddParam(EShaderFundamentalType::Uint);
 }
 
 BEGIN_SHADER_PARAMETER_STRUCT(FGroomDataInterfaceParameters, )
 	SHADER_PARAMETER_STRUCT_INCLUDE(FHairStrandsInstanceCommonParameters, Common)
 	SHADER_PARAMETER_STRUCT_INCLUDE(FHairStrandsInstanceResourceParameters, Resources)
+	SHADER_PARAMETER_STRUCT_INCLUDE(FHairStrandsInstanceInterpolationParameters, Interpolation)
 END_SHADER_PARAMETER_STRUCT()
 
 void UOptimusGroomDataInterface::GetShaderParameters(TCHAR const* UID, FShaderParametersMetadataBuilder& InOutBuilder, FShaderParametersMetadataAllocations& InOutAllocations) const
@@ -207,13 +214,20 @@ void FOptimusGroomDataProviderProxy::AllocateResources(FRDGBuilder& GraphBuilder
 	{
 		if (FHairGroupInstance* Instance = GroomComponent->GetGroupInstance(Index))
 		{
-			FHairStrandsInstanceResourceParameters& R = Resources.AddDefaulted_GetRef();
-			R.PositionBuffer		= RegisterAsSRV(GraphBuilder, Instance->Strands.RestResource->PositionBuffer);
-			R.PositionOffsetBuffer 	= RegisterAsSRV(GraphBuilder, Instance->Strands.RestResource->PositionOffsetBuffer);
-			R.CurveBuffer			= RegisterAsSRV(GraphBuilder, Instance->Strands.RestResource->CurveBuffer);
-			R.PointToCurveBuffer	= RegisterAsSRV(GraphBuilder, Instance->Strands.RestResource->PointToCurveBuffer);
-			R.CurveAttributeBuffer	= RegisterAsSRV(GraphBuilder, Instance->Strands.RestResource->CurveAttributeBuffer);
-			R.PointAttributeBuffer	= RegisterAsSRV(GraphBuilder, Instance->Strands.RestResource->PointAttributeBuffer);
+			{
+				FHairStrandsInstanceResourceParameters& R = Resources.AddDefaulted_GetRef();
+				R.PositionBuffer		= RegisterAsSRV(GraphBuilder, Instance->Strands.RestResource->PositionBuffer);
+				R.PositionOffsetBuffer 	= RegisterAsSRV(GraphBuilder, Instance->Strands.RestResource->PositionOffsetBuffer);
+				R.CurveBuffer			= RegisterAsSRV(GraphBuilder, Instance->Strands.RestResource->CurveBuffer);
+				R.PointToCurveBuffer	= RegisterAsSRV(GraphBuilder, Instance->Strands.RestResource->PointToCurveBuffer);
+				R.CurveAttributeBuffer	= RegisterAsSRV(GraphBuilder, Instance->Strands.RestResource->CurveAttributeBuffer);
+				R.PointAttributeBuffer	= RegisterAsSRV(GraphBuilder, Instance->Strands.RestResource->PointAttributeBuffer);
+			}
+
+			{
+				FHairStrandsInstanceInterpolationParameters& R = Interpolations.AddDefaulted_GetRef();
+				R.InterpolationBuffer = Instance->Strands.InterpolationResource ? RegisterAsSRV(GraphBuilder, Instance->Strands.InterpolationResource->InterpolationBuffer) : nullptr;
+			}
 
 			if (!FallbackSRV)
 			{
@@ -243,17 +257,21 @@ void FOptimusGroomDataProviderProxy::GatherDispatchData(FDispatchData const& InD
 		Parameters.Resources.PointAttributeBuffer	= FallbackSRV;
 		Parameters.Resources.PointToCurveBuffer		= FallbackSRV;
 		Parameters.Resources.CurveBuffer			= FallbackSRV;
+		Parameters.Interpolation.InterpolationBuffer= FallbackSRV;
 
 		if (FHairGroupInstance* Instance = GroomComponent->GetGroupInstance(InvocationIndex))
 		{
-			const FHairGroupPublicData::FVertexFactoryInput VFInput = ComputeHairStrandsVertexInputData(Instance, EGroomViewMode::None);		
+			const FHairGroupPublicData::FVertexFactoryInput VFInput = ComputeHairStrandsVertexInputData(Instance, EGroomViewMode::None);
 			Parameters.Common = VFInput.Strands.Common;
 
-			const bool bIsSRVValid = Resources[InvocationIndex].PositionBuffer != nullptr;
-			if (bIsSRVValid)
+			if (Resources[InvocationIndex].PositionBuffer != nullptr)
 			{
 				Parameters.Resources = Resources[InvocationIndex];
 			}
+			if (Interpolations[InvocationIndex].InterpolationBuffer != nullptr)
+			{
+				Parameters.Interpolation = Interpolations[InvocationIndex];
+			}
 		}
 	}
 }
\ No newline at end of file
diff --git a/Engine/Plugins/Runtime/HairStrands/Source/HairStrandsCore/Private/Deformer/DeformerDataInterfaceGroom.h b/Engine/Plugins/Runtime/HairStrands/Source/HairStrandsCore/Private/Deformer/DeformerDataInterfaceGroom.h
--- a/Engine/Plugins/Runtime/HairStrands/Source/HairStrandsCore/Private/Deformer/DeformerDataInterfaceGroom.h
+++ b/Engine/Plugins/Runtime/HairStrands/Source/HairStrandsCore/Private/Deformer/DeformerDataInterfaceGroom.h
@@ -69,6 +69,7 @@ private:
 
 	UGroomComponent* GroomComponent = nullptr; // Should it be HairInstance instead?
 	TArray<FHairStrandsInstanceResourceParameters> Resources;
+	TArray<FHairStrandsInstanceInterpolationParameters> Interpolations;
 	FRDGBufferSRVRef FallbackSRV;
 
 };
diff --git a/Engine/Plugins/Runtime/HairStrands/Source/HairStrandsCore/Private/Deformer/DeformerDataInterfaceGroomGuide.cpp b/Engine/Plugins/Runtime/HairStrands/Source/HairStrandsCore/Private/Deformer/DeformerDataInterfaceGroomGuide.cpp
new file mode 100644
--- /dev/null
+++ b/Engine/Plugins/Runtime/HairStrands/Source/HairStrandsCore/Private/Deformer/DeformerDataInterfaceGroomGuide.cpp
@@ -0,0 +1,206 @@
+// Copyright Epic Games, Inc. All Rights Reserved.
+
+#include "DeformerDataInterfaceGroomGuide.h"
+
+#include "ComputeFramework/ShaderParamTypeDefinition.h"
+#include "ComputeFramework/ComputeKernelPermutationVector.h"
+#include "OptimusDataDomain.h"
+#include "ShaderCompilerCore.h"
+#include "ShaderParameterMetadataBuilder.h"
+#include "GroomComponent.h"
+#include "GroomInstance.h"
+#include "RenderGraphBuilder.h"
+#include "GlobalRenderResources.h"
+#include "HairStrandsDefinitions.h"
+#include "DeformerGroomComponentSource.h"
+
+FString UOptimusGroomGuideDataInterface::GetDisplayName() const
+{
+	return TEXT("Guides");
+}
+
+TArray<FOptimusCDIPinDefinition> UOptimusGroomGuideDataInterface::GetPinDefinitions() const
+{
+	FName ControlPoint(UOptimusGroomComponentSource::Domains::ControlPoint);
+	FName Curve(UOptimusGroomComponentSource::Domains::Curve);
+
+	TArray<FOptimusCDIPinDefinition> Defs;
+	Defs.Add({ "NumGuidePoints",   "ReadNumGuidePoints" });
+	Defs.Add({ "NumGuideCurves",   "ReadNumGuideCurves" });
+	Defs.Add({ "Position",         "ReadPosition",          ControlPoint,   "ReadPosition" });
+	Defs.Add({ "CurveOffsetPoint", "ReadCurveOffsetPoint",  Curve,          "ReadCurveOffsetPoint" });
+	Defs.Add({ "CurveNumPoint",    "ReadCurveNumPoint",     Curve,          "ReadCurveNumPoint" });
+	return Defs;
+}
+ 
+TSubclassOf<UActorComponent> UOptimusGroomGuideDataInterface::GetRequiredComponentClass() const
+{
+	return UGroomComponent::StaticClass();
+}
+
+void UOptimusGroomGuideDataInterface::GetSupportedInputs(TArray<FShaderFunctionDefinition>& OutFunctions) const
+{
+	OutFunctions.AddDefaulted_GetRef()
+		.SetName(TEXT("ReadNumGuidePoints"))
+		.AddReturnType(EShaderFundamentalType::Uint);
+
+	OutFunctions.AddDefaulted_GetRef()
+		.SetName(TEXT("ReadNumGuideCurves"))
+		.AddReturnType(EShaderFundamentalType::Uint);
+
+	OutFunctions.AddDefaulted_GetRef()
+		.SetName(TEXT("ReadPosition"))
+		.AddReturnType(EShaderFundamentalType::Float, 3)
+		.AddParam(EShaderFundamentalType::Uint);
+
+	OutFunctions.AddDefaulted_GetRef()
+		.SetName(TEXT("ReadCurveOffsetPoint"))
+		.AddReturnType(EShaderFundamentalType::Uint)
+		.AddParam(EShaderFundamentalType::Uint);
+
+	OutFunctions.AddDefaulted_GetRef()
+		.SetName(TEXT("ReadCurveNumPoint"))
+		.AddReturnType(EShaderFundamentalType::Uint)
+		.AddParam(EShaderFundamentalType::Uint);
+}
+
+BEGIN_SHADER_PARAMETER_STRUCT(FGroomGuideDataInterfaceParameters, )
+	SHADER_PARAMETER_STRUCT_INCLUDE(FHairStrandsInstanceCommonParameters, Common)
+	SHADER_PARAMETER_STRUCT_INCLUDE(FHairStrandsInstanceResourceParameters, Resources)
+END_SHADER_PARAMETER_STRUCT()
+
+void UOptimusGroomGuideDataInterface::GetShaderParameters(TCHAR const* UID, FShaderParametersMetadataBuilder& InOutBuilder, FShaderParametersMetadataAllocations& InOutAllocations) const
+{
+	InOutBuilder.AddNestedStruct<FGroomGuideDataInterfaceParameters>(UID);
+}
+
+TCHAR const* UOptimusGroomGuideDataInterface::TemplateFilePath = TEXT("/Plugin/Runtime/HairStrands/Private/DeformerDataInterfaceGroomGuide.ush");
+
+TCHAR const* UOptimusGroomGuideDataInterface::GetShaderVirtualPath() const
+{
+	return TemplateFilePath;
+}
+
+void UOptimusGroomGuideDataInterface::GetShaderHash(FString& InOutKey) const
+{
+	GetShaderFileHash(TemplateFilePath, EShaderPlatform::SP_PCD3D_SM5).AppendString(InOutKey);
+}
+
+void UOptimusGroomGuideDataInterface::GetHLSL(FString& OutHLSL, FString const& InDataInterfaceName) const
+{
+	TMap<FString, FStringFormatArg> TemplateArgs =
+	{
+		{ TEXT("DataInterfaceName"), InDataInterfaceName },
+	};
+
+	FString TemplateFile;
+	LoadShaderSourceFile(TemplateFilePath, EShaderPlatform::SP_PCD3D_SM5, &TemplateFile, nullptr);
+	OutHLSL += FString::Format(*TemplateFile, TemplateArgs);
+}
+
+UComputeDataProvider* UOptimusGroomGuideDataInterface::CreateDataProvider(TObjectPtr<UObject> InBinding, uint64 InInputMask, uint64 InOutputMask) const
+{
+	UOptimusGroomGuideDataProvider* Provider = NewObject<UOptimusGroomGuideDataProvider>();
+	Provider->Groom = Cast<UGroomComponent>(InBinding);
+	return Provider;
+}
+
+
+FComputeDataProviderRenderProxy* UOptimusGroomGuideDataProvider::GetRenderProxy()
+{
+	return new FOptimusGroomGuideDataProviderProxy(Groom);
+}
+
+
+FOptimusGroomGuideDataProviderProxy::FOptimusGroomGuideDataProviderProxy(UGroomComponent* InGroomComponent)
+{
+	GroomComponent = InGroomComponent;
+}
+
+bool FOptimusGroomGuideDataProviderProxy::IsValid(FValidationData const& InValidationData) const
+{
+	if (InValidationData.ParameterStructSize != sizeof(FParameters))
+	{
+		return false;
+	}
+	if (GroomComponent == nullptr)
+	{
+		return false;
+	}
+	const uint32 GroupCount = GroomComponent->GetGroupCount();
+	if (InValidationData.NumInvocations != GroupCount)
+	{
+		return false;
+	}
+	// Earlier invalidation in case one of the instance is not ready/valid
+	for (uint32 GroupIndex = 0; GroupIndex < GroupCount; ++GroupIndex)
+	{
+		if (GroomComponent->GetGroupInstance(GroupIndex) == nullptr)
+		{
+			return false;
+		}
+	}
+	
+	return true;
+}
+
+void FOptimusGroomGuideDataProviderProxy::AllocateResources(FRDGBuilder& GraphBuilder)
+{
+	Resources.Empty();
+	FallbackSRV = nullptr;
+	const uint32 InstanceCount = GroomComponent ? GroomComponent->GetGroupCount() : 0;
+	for (uint32 Index =0; Index <InstanceCount;++Index)
+	{
+		if (FHairGroupInstance* Instance = GroomComponent->GetGroupInstance(Index))
+		{
+			FHairStrandsInstanceResourceParameters& R = Resources.AddDefaulted_GetRef();
+			if (Instance->Guides.RestResource)
+			{
+				R.PositionBuffer		= RegisterAsSRV(GraphBuilder, Instance->Guides.RestResource->PositionBuffer);
+				R.PositionOffsetBuffer 	= RegisterAsSRV(GraphBuilder, Instance->Guides.RestResource->PositionOffsetBuffer);
+				R.CurveBuffer			= RegisterAsSRV(GraphBuilder, Instance->Guides.RestResource->CurveBuffer);
+				R.PointToCurveBuffer	= RegisterAsSRV(GraphBuilder, Instance->Guides.RestResource->PointToCurveBuffer);
+			}
+
+			if (!FallbackSRV)
+			{
+				FallbackSRV = GraphBuilder.CreateSRV(GraphBuilder.RegisterExternalBuffer(GWhiteVertexBufferWithRDG->Buffer), PF_R16G16B16A16_UINT);
+			}
+		}
+	}
+}
+
+FHairGroupPublicData::FVertexFactoryInput ComputeHairStrandsVertexInputData(const FHairGroupInstance* Instance, EGroomViewMode ViewMode);
+void FOptimusGroomGuideDataProviderProxy::GatherDispatchData(FDispatchData const& InDispatchData)
+{
+	const uint32 InstanceCount = GroomComponent ? GroomComponent->GetGroupCount() : 0;
+	check(InDispatchData.NumInvocations == InstanceCount);
+	check(uint32(Resources.Num()) == InstanceCount);
+
+	const TStridedView<FParameters> ParameterArray = MakeStridedParameterView<FParameters>(InDispatchData);
+	for (int32 InvocationIndex = 0; InvocationIndex < ParameterArray.Num(); ++InvocationIndex)
+	{
+		// Default valid in case the instance is not valid
+		FParameters& Parameters = ParameterArray[InvocationIndex];
+		Parameters.Common.PointCount 				= 0;
+		Parameters.Common.CurveCount 				= 0;
+		Parameters.Resources.PositionBuffer			= FallbackSRV;
+		Parameters.Resources.PositionOffsetBuffer 	= FallbackSRV;
+		Parameters.Resources.PointToCurveBuffer		= FallbackSRV;
+		Parameters.Resources.CurveBuffer			= FallbackSRV;
+
+		if (FHairGroupInstance* Instance = GroomComponent->GetGroupInstance(InvocationIndex))
+		{
+			const bool bIsSRVValid = Resources[InvocationIndex].PositionBuffer != nullptr;
+			if (bIsSRVValid)
+			{
+				// Reuse the strands properties for guides, and apply guides values
+				const FHairGroupPublicData::FVertexFactoryInput VFInput = ComputeHairStrandsVertexInputData(Instance, EGroomViewMode::None);
+				Parameters.Common 				= VFInput.Strands.Common;
+				Parameters.Common.PointCount 	= Instance->Guides.Data ? Instance->Guides.Data->Header.PointCount : 0;
+				Parameters.Common.CurveCount 	= Instance->Guides.Data ? Instance->Guides.Data->Header.CurveCount : 0;
+				Parameters.Resources 			= Resources[InvocationIndex];
+			}
+		}
+	}
+}
\ No newline at end of file
diff --git a/Engine/Plugins/Runtime/HairStrands/Source/HairStrandsCore/Private/Deformer/DeformerDataInterfaceGroomGuide.h b/Engine/Plugins/Runtime/HairStrands/Source/HairStrandsCore/Private/Deformer/DeformerDataInterfaceGroomGuide.h
new file mode 100644
--- /dev/null
+++ b/Engine/Plugins/Runtime/HairStrands/Source/HairStrandsCore/Private/Deformer/DeformerDataInterfaceGroomGuide.h
@@ -0,0 +1,74 @@
+// Copyright Epic Games, Inc. All Rights Reserved.
+
+#pragma once
+
+#include "RenderGraphDefinitions.h"
+#include "OptimusComputeDataInterface.h"
+#include "HairStrandsInterface.h"
+#include "ComputeFramework/ComputeDataProvider.h"
+#include "DeformerDataInterfaceGroomGuide.generated.h"
+
+class FGroomGuideDataInterfaceParameters;
+class UGroomComponent;
+
+/** Compute Framework Data Interface for reading groom guides. */
+UCLASS(Category = ComputeFramework)
+class HAIRSTRANDSCORE_API UOptimusGroomGuideDataInterface : public UOptimusComputeDataInterface
+{
+	GENERATED_BODY()
+
+public:
+	//~ Begin UOptimusComputeDataInterface Interface
+	FString GetDisplayName() const override;
+	TArray<FOptimusCDIPinDefinition> GetPinDefinitions() const override;
+	TSubclassOf<UActorComponent> GetRequiredComponentClass() const override;
+	//~ End UOptimusComputeDataInterface Interface
+	
+	//~ Begin UComputeDataInterface Interface
+	TCHAR const* GetClassName() const override { return TEXT("Groom"); }
+	void GetSupportedInputs(TArray<FShaderFunctionDefinition>& OutFunctions) const override;
+	void GetShaderParameters(TCHAR const* UID, FShaderParametersMetadataBuilder& InOutBuilder, FShaderParametersMetadataAllocations& InOutAllocations) const override;
+	TCHAR const* GetShaderVirtualPath() const override;
+	void GetShaderHash(FString& InOutKey) const override;
+	void GetHLSL(FString& OutHLSL, FString const& InDataInterfaceName) const override;
+	UComputeDataProvider* CreateDataProvider(TObjectPtr<UObject> InBinding, uint64 InInputMask, uint64 InOutputMask) const override;
+	//~ End UComputeDataInterface Interface
+
+private:
+	static TCHAR const* TemplateFilePath;
+};
+
+/** Compute Framework Data Provider for reading groom. */
+UCLASS(BlueprintType, editinlinenew, Category = ComputeFramework)
+class UOptimusGroomGuideDataProvider : public UComputeDataProvider
+{
+	GENERATED_BODY()
+
+public:
+	UPROPERTY(EditAnywhere, BlueprintReadWrite, Category = Binding)
+	TObjectPtr<UGroomComponent> Groom = nullptr;
+
+	//~ Begin UComputeDataProvider Interface
+	FComputeDataProviderRenderProxy* GetRenderProxy() override;
+	//~ End UComputeDataProvider Interface
+};
+
+class FOptimusGroomGuideDataProviderProxy : public FComputeDataProviderRenderProxy
+{
+public:
+	FOptimusGroomGuideDataProviderProxy(UGroomComponent* InGroomComponent);
+
+	//~ Begin FComputeDataProviderRenderProxy Interface
+	bool IsValid(FValidationData const& InValidationData) const override;
+	void AllocateResources(FRDGBuilder& GraphBuilder) override;
+	void GatherDispatchData(FDispatchData const& InDispatchData) override;
+	//~ End FComputeDataProviderRenderProxy Interface
+
+private:
+	using FParameters = FGroomGuideDataInterfaceParameters;
+
+	UGroomComponent* GroomComponent = nullptr;
+	TArray<FHairStrandsInstanceResourceParameters> Resources;
+	FRDGBufferSRVRef FallbackSRV;
+
+};
diff --git a/Engine/Plugins/Runtime/HairStrands/Source/HairStrandsCore/Private/HairStrandsInterpolation.cpp b/Engine/Plugins/Runtime/HairStrands/Source/HairStrandsCore/Private/HairStrandsInterpolation.cpp
--- a/Engine/Plugins/Runtime/HairStrands/Source/HairStrandsCore/Private/HairStrandsInterpolation.cpp
+++ b/Engine/Plugins/Runtime/HairStrands/Source/HairStrandsCore/Private/HairStrandsInterpolation.cpp
@@ -1677,6 +1677,8 @@ static FHairGroupPublicData::FVertexFactoryInput InternalComputeHairStrandsVerte
 	OutVFInput.Strands.Common.RaytracingProceduralSplits = GetHairRaytracingProceduralSplits();
 	OutVFInput.Strands.Common.GroupIndex = Instance->Debug.GroupIndex;
 	OutVFInput.Strands.Common.GroupCount = Instance->Debug.GroupCount;
+	OutVFInput.Strands.Common.bSimulation = (Instance->Guides.bIsSimulationEnable || Instance->Guides.bIsSimulationCacheEnable) ? 1 : 0;
+	OutVFInput.Strands.Common.bSingleGuide = (Instance->Strands.InterpolationResource && (Instance->Strands.InterpolationResource->BulkData.Header.Flags & FHairStrandsInterpolationBulkData::DataFlags_HasSingleGuideData)) ? 1 : 0;
 
 	// Add local transform assignement here ...
 
diff --git a/Engine/Shaders/Private/HairStrands/HairStrandsVertexFactoryCommon.ush b/Engine/Shaders/Private/HairStrands/HairStrandsVertexFactoryCommon.ush
--- a/Engine/Shaders/Private/HairStrands/HairStrandsVertexFactoryCommon.ush
+++ b/Engine/Shaders/Private/HairStrands/HairStrandsVertexFactoryCommon.ush
@@ -27,6 +27,8 @@
 	float3				Prefix##_PositionOffset; \
 	uint				Prefix##_PointCount; \
 	uint				Prefix##_CurveCount; \
+	uint				Prefix##_bSimulation; \
+	uint				Prefix##_bSingleGuide; \
 	float4x4			Prefix##_LocalToWorldPrimitiveTransform; \
 	float4x4			Prefix##_LocalToTranslatedWorldPrimitiveTransform; \
 	uint4				Prefix##_CurveAttributeOffsets[HAIR_CURVE_ATTRIBUTE_OFFSET_COUNT]; \
diff --git a/Engine/Source/Runtime/Renderer/Public/HairStrandsInterface.h b/Engine/Source/Runtime/Renderer/Public/HairStrandsInterface.h
--- a/Engine/Source/Runtime/Renderer/Public/HairStrandsInterface.h
+++ b/Engine/Source/Runtime/Renderer/Public/HairStrandsInterface.h
@@ -40,6 +40,8 @@ BEGIN_SHADER_PARAMETER_STRUCT(FHairStrandsInstanceCommonParameters, RENDERER_API
 	SHADER_PARAMETER(uint32, bRaytracingGeometry)
 	SHADER_PARAMETER(uint32, bStableRasterization)
 	SHADER_PARAMETER(uint32, bScatterSceneLighting)
+	SHADER_PARAMETER(uint32, bSimulation)
+	SHADER_PARAMETER(uint32, bSingleGuide)
 	SHADER_PARAMETER(FVector3f, PositionOffset)
 	SHADER_PARAMETER(FVector3f, PrevPositionOffset)
 	SHADER_PARAMETER(FMatrix44f, LocalToWorldPrimitiveTransform)
@@ -73,6 +75,11 @@ BEGIN_SHADER_PARAMETER_STRUCT(FHairStrandsInstanceCullingParameters, RENDERER_AP
 	RDG_BUFFER_ACCESS(CullingIndirectBufferArgs, ERHIAccess::IndirectArgs)
 END_SHADER_PARAMETER_STRUCT()
 
+// Instance interpolation resources (RDG)
+BEGIN_SHADER_PARAMETER_STRUCT(FHairStrandsInstanceInterpolationParameters, RENDERER_API)
+	SHADER_PARAMETER_RDG_BUFFER_SRV(ByteAddressBuffer, InterpolationBuffer)
+END_SHADER_PARAMETER_STRUCT()
+
 // Instance resources (Raw)
 BEGIN_SHADER_PARAMETER_STRUCT(FHairStrandsInstanceResourceRawParameters, RENDERER_API)
 	SHADER_PARAMETER_SRV(Buffer<uint4>, PositionBuffer)
