Fix missing groom shadow when hidden in game and hidden shadow are enabled at the same time with MRQ path tracer.

diff --git a/Engine/Plugins/Runtime/HairStrands/Source/HairStrandsCore/Private/GroomComponent.cpp b/Engine/Plugins/Runtime/HairStrands/Source/HairStrandsCore/Private/GroomComponent.cpp
--- a/Engine/Plugins/Runtime/HairStrands/Source/HairStrandsCore/Private/GroomComponent.cpp
+++ b/Engine/Plugins/Runtime/HairStrands/Source/HairStrandsCore/Private/GroomComponent.cpp
@@ -639,7 +639,7 @@ public:
 				continue;
 			}
 
-			if (IsShown(View) && (VisibilityMap & (1 << ViewIndex)))
+			if ((IsShadowCast(View) || IsShown(View)) && (VisibilityMap & (1 << ViewIndex)))
 			{
 				for (uint32 GroupIt = 0; GroupIt < GroupCount; ++GroupIt)
 				{
@@ -921,6 +921,11 @@ public:
 
 	virtual FPrimitiveViewRelevance GetViewRelevance(const FSceneView* View) const override
 	{
+		// When path tracing is enabled force DrawRelevance if not visible in main view ('hidden in game'), 
+		// but visible in shadow 'hidden shadow') so that raytracing geometry is created/updated correctly
+		const bool bPathtracing = View->Family->EngineShowFlags.PathTracing;
+		const bool bForceDrawRelevance = bPathtracing && (!IsShown(View) && IsShadowCast(View));
+
 		bool bUseCardsOrMesh = false;
 		for (FHairGroupInstance* Instance : HairGroupInstances)
 		{
@@ -930,10 +935,10 @@ public:
 		}
 
 		FPrimitiveViewRelevance Result;
-		Result.bHairStrands = IsShown(View);
+		Result.bHairStrands = IsShown(View) || bForceDrawRelevance;
 
 		// Special pass for hair strands geometry (not part of the base pass, and shadowing is handlded in a custom fashion). When cards rendering is enabled we reusethe base pass
-		Result.bDrawRelevance		= IsShown(View);
+		Result.bDrawRelevance		= IsShown(View) || bForceDrawRelevance;
 		Result.bRenderInMainPass	= bUseCardsOrMesh && ShouldRenderInMainPass();
 		Result.bShadowRelevance		= IsShadowCast(View);
 		Result.bDynamicRelevance	= bUseCardsOrMesh;
