Fix hair bound with scaled transform and take into account hair width, so that large width increase correctly the bound size.

diff --git a/Engine/Plugins/Runtime/HairStrands/Source/HairStrandsCore/Private/GroomComponent.cpp b/Engine/Plugins/Runtime/HairStrands/Source/HairStrandsCore/Private/GroomComponent.cpp
--- a/Engine/Plugins/Runtime/HairStrands/Source/HairStrandsCore/Private/GroomComponent.cpp
+++ b/Engine/Plugins/Runtime/HairStrands/Source/HairStrandsCore/Private/GroomComponent.cpp
@@ -126,7 +126,7 @@ static EHairInterpolationType ToHairInterpolationType(EGroomInterpolationType In
 }
 ////////////////////////////////////////////////////////////////////////////////////////////////////////////////////////
 
-static FHairGroupDesc GetGroomGroupsDesc(const UGroomAsset* Asset, UGroomComponent* Component, uint32 GroupIndex)
+static FHairGroupDesc GetGroomGroupsDesc(const UGroomAsset* Asset, const UGroomComponent* Component, uint32 GroupIndex)
 {
 	if (!Asset || GroupIndex >= uint32(Component->GroomGroupsDesc.Num()))
 	{
@@ -1759,11 +1759,23 @@ FBoxSphereBounds UGroomComponent::CalcBounds(const FTransform& InLocalToWorld) c
 		FBox LocalHairBound(EForceInit::ForceInitToZero);
 		if (!GroomCacheBuffers.IsValid())
 		{
-			for (const FHairGroupPlatformData& GroupData : GroomAsset->HairGroupsPlatformData)
+			for (uint32 GroupIndex = 0, GroupCount = GroomAsset->HairGroupsPlatformData.Num(); GroupIndex<GroupCount; ++GroupIndex)
 			{
+				const FHairGroupPlatformData& GroupData = GroomAsset->HairGroupsPlatformData[GroupIndex];
 				if (IsHairStrandsEnabled(EHairStrandsShaderType::Strands) && GroupData.Strands.HasValidData())
 				{
+					// The hair width is not scaled by the transform scale. Since we increase the local bound by the hair width, 
+					// and that amount is then scaled by the local to world matrix which conatins a scale factor, we apply a 
+					// inverse scale to the radius computation
+					const float InvScale = 1.f / InLocalToWorld.GetMaximumAxisScale();
+
+					// Take into account the strands size and add it as a bound 'border'
+					const FHairGroupDesc Desc = GetGroomGroupsDesc(GroomAsset, this, GroupIndex);
+					const FVector HairBorder(0.5f * Desc.HairWidth * FMath::Max3(1.0f, Desc.HairRootScale, Desc.HairTipScale) * InvScale);
+
 					LocalHairBound += GroupData.Strands.GetBounds();
+					LocalHairBound.Min -= HairBorder;
+					LocalHairBound.Max += HairBorder;
 				}
 				else if (IsHairStrandsEnabled(EHairStrandsShaderType::Cards) && GroupData.Cards.HasValidData())
 				{
@@ -1794,13 +1806,24 @@ FBoxSphereBounds UGroomComponent::CalcBounds(const FTransform& InLocalToWorld) c
 		// * GHairStrands_BoundsMode=2 : use the skel mesh bounds + groom bounds. This is more conservative compares to GHairStrands_BoundsMode=0.
 		if (RegisteredMeshComponent && GHairStrands_BoundsMode == 1)
 		{
-			return RegisteredMeshComponent->CalcBounds(InLocalToWorld).GetBox();
+			const FVector Scale = GetRelativeScale3D();
+			const FVector InvScale(1.f / Scale.X, 1.f / Scale.Y, 1.f / Scale.Z);
+
+			FBox EffectiveBound = RegisteredMeshComponent->CalcBounds(InLocalToWorld).GetBox();
+			EffectiveBound.Min *= InvScale;
+			EffectiveBound.Max *= InvScale;
+			return EffectiveBound;
 		}
 		else if (RegisteredMeshComponent && GHairStrands_BoundsMode == 2)
 		{
+			const FVector Scale = GetRelativeScale3D();
+			const FVector InvScale(1.f / Scale.X, 1.f / Scale.Y, 1.f / Scale.Z);
+
 			const FVector3d GroomExtends = LocalHairBound.GetExtent();
-			const float BoundExtraRadius = 0.5f * static_cast<float>(FMath::Max(0.0, FMath::Max3(GroomExtends.X, GroomExtends.Y, GroomExtends.Z)));
+			const float BoundExtraRadius = static_cast<float>(FMath::Max(0.0, FMath::Max3(GroomExtends.X, GroomExtends.Y, GroomExtends.Z)));
 			FBox EffectiveBound = RegisteredMeshComponent->CalcBounds(FTransform::Identity).GetBox();
+			EffectiveBound.Min *= InvScale;
+			EffectiveBound.Max *= InvScale;			
 			EffectiveBound.Min -= FVector3d(BoundExtraRadius);
 			EffectiveBound.Max += FVector3d(BoundExtraRadius);
 			return FBoxSphereBounds(EffectiveBound.TransformBy(InLocalToWorld));
@@ -1820,10 +1843,13 @@ FBoxSphereBounds UGroomComponent::CalcBounds(const FTransform& InLocalToWorld) c
 				BoundExtraRadius = static_cast<float>(FMath::Max3(0.0, FMath::Max3(MinDiff.X, MinDiff.Y, MinDiff.Z), FMath::Max3(MaxDiff.X, MaxDiff.Y, MaxDiff.Z)));
 			}
 			
+			const FVector Scale = GetRelativeScale3D();
+			const FVector InvScale(1.f / Scale.X, 1.f / Scale.Y, 1.f / Scale.Z);
+
 			const FBox LocalSkeletalBound = RegisteredMeshComponent->CalcBounds(FTransform::Identity).GetBox();
 			FBox LocalBound(EForceInit::ForceInitToZero);
-			LocalBound += LocalSkeletalBound.Min - BoundExtraRadius;
-			LocalBound += LocalSkeletalBound.Max + BoundExtraRadius;
+			LocalBound += InvScale * LocalSkeletalBound.Min - BoundExtraRadius;
+			LocalBound += InvScale * LocalSkeletalBound.Max + BoundExtraRadius;
 			return FBoxSphereBounds(LocalBound.TransformBy(InLocalToWorld));
 		}
 		else
