Fix crash when adding new hair strands LOD entry.

Derived data were not rebuild when adding a new LOD due to missing invalidation logic. In addition the DDC key was not containing all the LOD data, causing mismatching between loaded data, an requested derived data.

diff --git a/Engine/Plugins/Runtime/HairStrands/Source/HairStrandsCore/Private/GroomAsset.cpp b/Engine/Plugins/Runtime/HairStrands/Source/HairStrandsCore/Private/GroomAsset.cpp
--- a/Engine/Plugins/Runtime/HairStrands/Source/HairStrandsCore/Private/GroomAsset.cpp
+++ b/Engine/Plugins/Runtime/HairStrands/Source/HairStrandsCore/Private/GroomAsset.cpp
@@ -1539,6 +1539,30 @@ static bool IsStrandsInterpolationAttributes(const FName PropertyName)
 		|| PropertyName == GET_MEMBER_NAME_CHECKED(FHairInterpolationSettings, bUseUniqueGuide);
 }
 
+static bool IsStrandsLODAttributes(const FName PropertyName)
+{
+	return
+		// LOD count needs for cluster rebuilding
+		PropertyName == GET_MEMBER_NAME_CHECKED(UGroomAsset, HairGroupsLOD)
+
+		// LOD groups needs for cluster rebuilding
+		|| PropertyName == GET_MEMBER_NAME_CHECKED(FHairGroupsLOD, ClusterWorldSize)
+		|| PropertyName == GET_MEMBER_NAME_CHECKED(FHairGroupsLOD, ClusterScreenSizeScale)
+		|| PropertyName == GET_MEMBER_NAME_CHECKED(FHairGroupsLOD, LODs)
+
+		// LOD settings needs for cluster rebuilding
+		|| PropertyName == GET_MEMBER_NAME_CHECKED(FHairLODSettings, CurveDecimation)
+		|| PropertyName == GET_MEMBER_NAME_CHECKED(FHairLODSettings, VertexDecimation)
+		|| PropertyName == GET_MEMBER_NAME_CHECKED(FHairLODSettings, AngularThreshold)
+		|| PropertyName == GET_MEMBER_NAME_CHECKED(FHairLODSettings, ScreenSize)
+		|| PropertyName == GET_MEMBER_NAME_CHECKED(FHairLODSettings, ThicknessScale)
+		|| PropertyName == GET_MEMBER_NAME_CHECKED(FHairLODSettings, bVisible)
+		|| PropertyName == GET_MEMBER_NAME_CHECKED(FHairLODSettings, GeometryType)
+		|| PropertyName == GET_MEMBER_NAME_CHECKED(FHairLODSettings, BindingType)
+		|| PropertyName == GET_MEMBER_NAME_CHECKED(FHairLODSettings, Simulation)
+		|| PropertyName == GET_MEMBER_NAME_CHECKED(FHairLODSettings, GlobalInterpolation);
+}
+
 void UGroomAsset::PostEditChangeProperty(FPropertyChangedEvent& PropertyChangedEvent)
 {
 	Super::PostEditChangeProperty(PropertyChangedEvent);
@@ -1568,9 +1592,9 @@ void UGroomAsset::PostEditChangeProperty(FPropertyChangedEvent& PropertyChangedE
 		}
 	}
 
-	// Rebuild the groom asset if some decimation attribute has changed
-	const bool bStrandsInterpolationChanged = IsStrandsInterpolationAttributes(PropertyName);
-	if (bStrandsInterpolationChanged)
+	// Rebuild the groom cached data if interpolation or LODs have changed
+	const bool bNeedRebuildDerivedData = IsStrandsInterpolationAttributes(PropertyName) || IsStrandsLODAttributes(PropertyName);
+	if (bNeedRebuildDerivedData)
 	{
 		CacheDerivedDatas();
 	}
@@ -1589,7 +1613,7 @@ void UGroomAsset::PostEditChangeProperty(FPropertyChangedEvent& PropertyChangedE
 	// Cards should be refresh only under user action
 	// By pass update if bStrandsInterpolationChanged has the resources have already been recreated
 	const bool bCardsToolUpdate = IsCardsProceduralAttributes(PropertyName);
-	if (!bCardsToolUpdate && !bStrandsInterpolationChanged)
+	if (!bCardsToolUpdate && !bNeedRebuildDerivedData)
 	{
 		FGroomComponentRecreateRenderStateContext Context(this);
 		UpdateResource();
@@ -1607,7 +1631,7 @@ void UGroomAsset::PostEditChangeProperty(FPropertyChangedEvent& PropertyChangedE
 	const bool bCardMaterialChanged = PropertyName == GET_MEMBER_NAME_CHECKED(FHairGroupsCardsSourceDescription, Material);
 	const bool bMeshMaterialChanged = PropertyName == GET_MEMBER_NAME_CHECKED(FHairGroupsMeshesSourceDescription, Material);
 
-	if (bStrandsInterpolationChanged || bGeometryTypeChanged || bCardMaterialChanged || bMeshMaterialChanged)
+	if (bNeedRebuildDerivedData || bGeometryTypeChanged || bCardMaterialChanged || bMeshMaterialChanged)
 	{
 		// Delegate used for notifying groom data & groom resoures invalidation
 		OnGroomAssetResourcesChanged.Broadcast();
@@ -3304,23 +3328,6 @@ void UGroomAsset::InitStrandsResources()
 
 			if (GroupData.Strands.ClusterCullingBulkData.IsValid())
 			{
-				const uint32 HairGroupLODCount = HairGroupsLOD[GroupIndex].LODs.Num();
-				bool bLastVisibility = true;
-				// LOD visibility is not serialized into FHairStrandsClusterCullingBulkData as it does not affect the actual generated data. For consistency we 
-				// patch the LOD visibility with the groom asset value, which might be different from what has been serialized
-				for (uint32 LODIt = 0, LODCount = GroupData.Strands.ClusterCullingBulkData.Header.LODVisibility.Num(); LODIt < LODCount; ++LODIt)
-				{
-					if (LODIt < HairGroupLODCount)
-					{
-						bLastVisibility = HairGroupsLOD[GroupIndex].LODs[LODIt].bVisible;
-						GroupData.Strands.ClusterCullingBulkData.Header.LODVisibility[LODIt] = bLastVisibility;
-					}
-					else
-					{
-						GroupData.Strands.ClusterCullingBulkData.Header.LODVisibility[LODIt] = bLastVisibility;
-					}
-				}
-
 				GroupData.Strands.ClusterCullingResource = new FHairStrandsClusterCullingResource(GroupData.Strands.ClusterCullingBulkData, ResourceName, OwnerName);
 			}
 
diff --git a/Engine/Plugins/Runtime/HairStrands/Source/HairStrandsCore/Private/GroomAssetInterpolation.cpp b/Engine/Plugins/Runtime/HairStrands/Source/HairStrandsCore/Private/GroomAssetInterpolation.cpp
--- a/Engine/Plugins/Runtime/HairStrands/Source/HairStrandsCore/Private/GroomAssetInterpolation.cpp
+++ b/Engine/Plugins/Runtime/HairStrands/Source/HairStrandsCore/Private/GroomAssetInterpolation.cpp
@@ -129,13 +129,14 @@ void FHairGroupsLOD::BuildDDCKey(FArchive& Ar)
 
 	for (FHairLODSettings& LOD : LODs)
 	{
-		if (LOD.GeometryType == EGroomGeometryType::Strands)
-		{
-			Ar << LOD.CurveDecimation;
-			Ar << LOD.VertexDecimation;
-			Ar << LOD.AngularThreshold;
-			Ar << LOD.ScreenSize;
-			Ar << LOD.ThicknessScale;
-		}
+		Ar << LOD.CurveDecimation;
+		Ar << LOD.VertexDecimation;
+		Ar << LOD.AngularThreshold;
+		Ar << LOD.ScreenSize;
+		Ar << LOD.ThicknessScale;
+		Ar << LOD.bVisible;
+		Ar << LOD.BindingType;
+		Ar << LOD.Simulation;
+		Ar << LOD.GlobalInterpolation;
 	}
 }
diff --git a/Engine/Source/Runtime/Core/Private/UObject/UE5MainStreamObjectVersion.cpp b/Engine/Source/Runtime/Core/Private/UObject/UE5MainStreamObjectVersion.cpp
--- a/Engine/Source/Runtime/Core/Private/UObject/UE5MainStreamObjectVersion.cpp
+++ b/Engine/Source/Runtime/Core/Private/UObject/UE5MainStreamObjectVersion.cpp
@@ -19,6 +19,7 @@ TMap<FGuid, FGuid> FUE5MainStreamObjectVersion::GetSystemGuids()
 	SystemGuids.Add(DevGuids.STATICMESH_DERIVEDDATA_VER, FGuid("98F7D79A5811013825E75BBCC41ED3E9"));
 	SystemGuids.Add(DevGuids.POSESEARCHDB_DERIVEDDATA_VER, FGuid("4E595C2AC5E947D6BA9ABC874353E5BC"));
 	SystemGuids.Add(DevGuids.GROOM_BINDING_DERIVED_DATA_VERSION, FGuid("156678D4F9084D7CAE98FADC6AB93573"));
+	SystemGuids.Add(DevGuids.GROOM_DERIVED_DATA_VERSION, FGuid("00B8207248CC44B88F7DEC8F328A5D8B"));
 
 	return SystemGuids;
 }
