Fix crash when requesting invalid groom data from DDC.

diff --git a/Engine/Plugins/Runtime/HairStrands/Source/HairStrandsCore/Private/GroomAsset.cpp b/Engine/Plugins/Runtime/HairStrands/Source/HairStrandsCore/Private/GroomAsset.cpp
--- a/Engine/Plugins/Runtime/HairStrands/Source/HairStrandsCore/Private/GroomAsset.cpp
+++ b/Engine/Plugins/Runtime/HairStrands/Source/HairStrandsCore/Private/GroomAsset.cpp
@@ -1773,9 +1773,9 @@ static void InternalSerializeStrand(FArchive& Ar, UObject* Owner, FHairGroupPlat
 	// When cooking data, force loading of bulk data prior to saving them
 	if (Ar.IsCooking() && Ar.IsSaving())
 	{
-		{ FHairResourceRequest R; R.Request(StrandData.BulkData, true /*bWait*/, true /*bFillBulkdata*/); }
-		{ FHairResourceRequest R; R.Request(StrandData.InterpolationBulkData, true /*bWait*/, true /*bFillBulkdata*/); }
-		{ FHairResourceRequest R; R.Request(StrandData.ClusterCullingBulkData, true /*bWait*/, true /*bFillBulkdata*/); }
+		{ FHairResourceRequest R; R.Request(StrandData.BulkData, true /*bWait*/, true /*bFillBulkdata*/, Owner->GetFName()); }
+		{ FHairResourceRequest R; R.Request(StrandData.InterpolationBulkData, true /*bWait*/, true /*bFillBulkdata*/, Owner->GetFName()); }
+		{ FHairResourceRequest R; R.Request(StrandData.ClusterCullingBulkData, true /*bWait*/, true /*bFillBulkdata*/, Owner->GetFName()); }
 	}
 
 	if (!Ar.IsCooking() || !StrandData.bIsCookedOut)
diff --git a/Engine/Plugins/Runtime/HairStrands/Source/HairStrandsCore/Private/GroomResources.cpp b/Engine/Plugins/Runtime/HairStrands/Source/HairStrandsCore/Private/GroomResources.cpp
--- a/Engine/Plugins/Runtime/HairStrands/Source/HairStrandsCore/Private/GroomResources.cpp
+++ b/Engine/Plugins/Runtime/HairStrands/Source/HairStrandsCore/Private/GroomResources.cpp
@@ -106,7 +106,7 @@ const TCHAR* ToHairResourceDebugName(const TCHAR* In, FHairResourceName& InDebug
 #if !WITH_EDITORONLY_DATA
 bool FHairResourceRequest::IsNone() const 		{ return IORequest.IsNone(); }
 bool FHairResourceRequest::IsCompleted() const	{ return IORequest.IsCompleted(); }
-void FHairResourceRequest::Request(FHairStrandsBulkCommon& In, bool bWait, bool bFillBulkdata)
+void FHairResourceRequest::Request(FHairStrandsBulkCommon& In, bool bWait, bool bFillBulkdata, const FName& InOwnerName)
 { 
 	if (In.GetResourceCount() != 0)
 	{
@@ -134,7 +134,7 @@ bool FHairResourceRequest::IsCompleted() const
 // Request fullfil 2 use cases:
 // * Load DDC data and upload them to GPU
 // * Load DDC data and store them into bulkdata for serialization
-void FHairResourceRequest::Request(FHairStrandsBulkCommon& In, bool bWait, bool bFillBulkdata)
+void FHairResourceRequest::Request(FHairStrandsBulkCommon& In, bool bWait, bool bFillBulkdata, const FName& InOwnerName)
 {
 	if (In.GetResourceCount() == 0)
 	{
@@ -156,6 +156,7 @@ void FHairResourceRequest::Request(FHairStrandsBulkCommon& In, bool bWait, bool
 		Chunk.Status = FHairChunkRequest::EStatus::Pending;
 		Chunk.Container = (FHairBulkContainer*)R.UserData;
 		R.UserData = (uint64)&Chunk;
+		R.Name = InOwnerName.ToString();
 	}
 
 	DDCRequestOwner = MakeUnique<FRequestOwner>(bWait ? UE::DerivedData::EPriority::Blocking : UE::DerivedData::EPriority::Normal); // <= Move this onto the resource (Buffer, in order to ensure no race condition)
@@ -174,17 +175,22 @@ void FHairResourceRequest::Request(FHairStrandsBulkCommon& In, bool bWait, bool
 		}
 		else
 		{
-			checkNoEntry();
+			FHairChunkRequest& Chunk = *(FHairChunkRequest*)Response.UserData;
+			Chunk.Status = FHairChunkRequest::Failed;
+			UE_LOG(LogHairStrands, Error, TEXT("[Groom] DDC request failed for '%s' (Key:%s) "), *Response.Name, *In.DerivedDataKey);
 		}
 	});
 
-	// TODO remove this
-	DDCRequestOwner->Wait();
+	// TODO remove for resources supporting streaming request
+	bWait = true;
 
 	if (bWait)
 	{
 		DDCRequestOwner->Wait();
-		check(IsCompleted());
+		if (!IsCompleted())
+		{
+			return;
+		}
 	}
 
 	// Fill in bytebulkdata
@@ -916,7 +922,7 @@ bool FHairStrandsRestResource::InternalIsDataLoaded()
 {
 	if (BulkDataRequest.IsNone())
 	{
-		BulkDataRequest.Request(BulkData, true, true);
+		BulkDataRequest.Request(BulkData, true, true, OwnerName);
 	}
 	return BulkDataRequest.IsCompleted();
 }
@@ -1096,7 +1102,7 @@ bool FHairStrandsClusterCullingResource::InternalIsDataLoaded()
 {
 	if (BulkDataRequest.IsNone())
 	{
-		BulkDataRequest.Request(BulkData, true, true);
+		BulkDataRequest.Request(BulkData, true, true, OwnerName);
 	}
 	return BulkDataRequest.IsCompleted();
 }
@@ -1524,7 +1530,7 @@ bool FHairStrandsInterpolationResource::InternalIsDataLoaded()
 {
 	if (BulkDataRequest.IsNone())
 	{
-		BulkDataRequest.Request(BulkData, true, true);
+		BulkDataRequest.Request(BulkData, true, true, OwnerName);
 	}
 	return BulkDataRequest.IsCompleted();
 }
diff --git a/Engine/Plugins/Runtime/HairStrands/Source/HairStrandsCore/Private/HairStrandsDatas.cpp b/Engine/Plugins/Runtime/HairStrands/Source/HairStrandsCore/Private/HairStrandsDatas.cpp
--- a/Engine/Plugins/Runtime/HairStrands/Source/HairStrandsCore/Private/HairStrandsDatas.cpp
+++ b/Engine/Plugins/Runtime/HairStrands/Source/HairStrandsCore/Private/HairStrandsDatas.cpp
@@ -188,7 +188,7 @@ void FHairStrandsBulkCommon::FQuery::Add(FHairBulkContainer& In, const TCHAR* In
 
 		using namespace UE::DerivedData;
 		FCachePutValueRequest& Out = OutWriteDDC->AddDefaulted_GetRef();
-		Out.Name 	= Owner->GetPathName();
+		if (Owner) { Out.Name = Owner->GetPathName(); }
 		Out.Key 	= ConvertLegacyCacheKey(*DerivedDataKey + InSuffix);
 		Out.Value 	= FValue::Compress(MakeSharedBufferFromArray(MoveTemp(WriteData)));
 		Out.Policy 	= ECachePolicy::Default;
diff --git a/Engine/Plugins/Runtime/HairStrands/Source/HairStrandsCore/Public/GroomResources.h b/Engine/Plugins/Runtime/HairStrands/Source/HairStrandsCore/Public/GroomResources.h
--- a/Engine/Plugins/Runtime/HairStrands/Source/HairStrandsCore/Public/GroomResources.h
+++ b/Engine/Plugins/Runtime/HairStrands/Source/HairStrandsCore/Public/GroomResources.h
@@ -90,7 +90,7 @@ struct FHairChunkRequest
 
 struct FHairResourceRequest
 {
-	void Request(FHairStrandsBulkCommon& In, bool bWait=false, bool bFillBulkData=false);
+	void Request(FHairStrandsBulkCommon& In, bool bWait=false, bool bFillBulkData=false, const FName& InOwnerName = NAME_None);
 	bool IsNone() const ;
 	bool IsCompleted() const;
 
