Fix crash when opening groom with mesh deformer applied in non-PIE mode.

diff --git a/Engine/Plugins/Runtime/HairStrands/Source/HairStrandsCore/Private/GroomResources.cpp b/Engine/Plugins/Runtime/HairStrands/Source/HairStrandsCore/Private/GroomResources.cpp
--- a/Engine/Plugins/Runtime/HairStrands/Source/HairStrandsCore/Private/GroomResources.cpp
+++ b/Engine/Plugins/Runtime/HairStrands/Source/HairStrandsCore/Private/GroomResources.cpp
@@ -1057,9 +1057,9 @@ FRDGExternalBuffer& FHairStrandsDeformedResource::GetDeformerBuffer(FRDGBuilder&
 FRDGExternalBuffer& FHairStrandsDeformedResource::GetDeformerCurveAttributeBuffer(FRDGBuilder& GraphBuilder)
 {
 	// Deformer curve attributes
-	if (DeformerCurveAttributeBuffer.Buffer == nullptr)
-	{		
-		check(BulkData.Data.CurveAttributes.Data.GetBulkDataSize() > 0);
+	// Bulkdata might be not filled in directly when first called. Hence the extra size check
+	if (DeformerCurveAttributeBuffer.Buffer == nullptr && BulkData.Data.CurveAttributes.Data.GetBulkDataSize() > 0)
+	{
 		InternalCreateByteAddressBufferRDG(GraphBuilder, BulkData.Data.CurveAttributes.Data.GetBulkDataSize(), DeformerCurveAttributeBuffer, ToHairResourceDebugName(HAIRSTRANDS_RESOUCE_NAME(CurveType, Hair.StrandsDeformed_DeformerCurveAttributeBuffer), ResourceName), OwnerName, EHairResourceUsageType::Dynamic);
 	}
 	return DeformerCurveAttributeBuffer;
@@ -1068,9 +1068,9 @@ FRDGExternalBuffer& FHairStrandsDeformedResource::GetDeformerCurveAttributeBuffe
 FRDGExternalBuffer& FHairStrandsDeformedResource::GetDeformerPointAttributeBuffer(FRDGBuilder& GraphBuilder)
 {
 	// Deformer point attributes
-	if (DeformerPointAttributeBuffer.Buffer == nullptr && (BulkData.Header.Flags & FHairStrandsBulkData::DataFlags_HasPointAttribute))
+	// Bulkdata might be not filled in directly when first called. Hence the extra size check
+	if (DeformerPointAttributeBuffer.Buffer == nullptr && (BulkData.Header.Flags & FHairStrandsBulkData::DataFlags_HasPointAttribute) && BulkData.Data.PointAttributes.Data.GetBulkDataSize() > 0)
 	{
-		check(BulkData.Data.PointAttributes.Data.GetBulkDataSize() > 0);
 		InternalCreateByteAddressBufferRDG(GraphBuilder, BulkData.Data.PointAttributes.Data.GetBulkDataSize(), DeformerPointAttributeBuffer, ToHairResourceDebugName(HAIRSTRANDS_RESOUCE_NAME(CurveType, Hair.StrandsDeformedt_DeformerPointAttributeBuffer), ResourceName), OwnerName, EHairResourceUsageType::Dynamic);
 	}
 	return DeformerPointAttributeBuffer;
