Change groom attribute to be grouped in chunk so that they could be streamed at a finer grain.

diff --git a/Engine/Plugins/Runtime/HairStrands/Shaders/Private/DeformerDataInterfaceGroomWrite.ush b/Engine/Plugins/Runtime/HairStrands/Shaders/Private/DeformerDataInterfaceGroomWrite.ush
--- a/Engine/Plugins/Runtime/HairStrands/Shaders/Private/DeformerDataInterfaceGroomWrite.ush
+++ b/Engine/Plugins/Runtime/HairStrands/Shaders/Private/DeformerDataInterfaceGroomWrite.ush
@@ -57,8 +57,13 @@ void WritePositionAndRadius_{DataInterfaceName}(uint ControlPointIndex, float4 P
 	}
 }
 
-void InternalSetHairStrandsAttributes(uint Index, uint AttributeOffsetInBytes, uint AttributeStrideInBytes, RWByteAddressBuffer InAttributeBuffer, uint In)
+void InternalSetHairStrandsAttributes(uint Index, uint AttributeOffsetInBytes, uint AttributeStrideInBytes, RWByteAddressBuffer InAttributeBuffer, uint In, uint InIndexToChunkDivAsShift, uint InChunkElementCount, uint InChunkStrideInBytes)
 {
+	// Adapt index and offset
+	const uint ChunkIndex   = Index >> InIndexToChunkDivAsShift;
+	Index                  -= ChunkIndex * InChunkElementCount;
+	AttributeOffsetInBytes += ChunkIndex * InChunkStrideInBytes;
+
 	// Ensure the reading address is 4-byte aligned
 	const uint Address 					= AttributeOffsetInBytes + Index * AttributeStrideInBytes;
 	const uint Address_4BytesAligned 	= Address & (~0x3);
@@ -74,8 +79,13 @@ void InternalSetHairStrandsAttributes(uint Index, uint AttributeOffsetInBytes, u
 	InAttributeBuffer.InterlockedOr(Address_4BytesAligned, AttributeValue);
 }
 
-void InternalSetHairStrandsAttributes32(uint Index, uint AttributeOffsetInBytes, uint AttributeStrideInBytes, RWByteAddressBuffer InAttributeBuffer, uint In)
+void InternalSetHairStrandsAttributes32(uint Index, uint AttributeOffsetInBytes, uint AttributeStrideInBytes, RWByteAddressBuffer InAttributeBuffer, uint In, uint InIndexToChunkDivAsShift, uint InChunkElementCount, uint InChunkStrideInBytes)
 {
+	// Adapt index and offset
+	const uint ChunkIndex   = Index >> InIndexToChunkDivAsShift;
+	Index                  -= ChunkIndex * InChunkElementCount;
+	AttributeOffsetInBytes += ChunkIndex * InChunkStrideInBytes;
+
 	// Ensure the reading address is 4-byte aligned
 	const uint Address = AttributeOffsetInBytes + Index * AttributeStrideInBytes;
 	const uint Address_4BytesAligned = Address & (~0x3);
@@ -84,12 +94,15 @@ void InternalSetHairStrandsAttributes32(uint Index, uint AttributeOffsetInBytes,
 	InAttributeBuffer.Store(Address_4BytesAligned, In);
 }
 
+, {DataInterfaceName}_CurveAttributeIndexToChunkDivAsShift, {DataInterfaceName}_CurveAttributeChunkElementCount), {DataInterfaceName}_CurveAttributeChunkStrideInBytes
+, {DataInterfaceName}_PointAttributeIndexToChunkDivAsShift, {DataInterfaceName}_PointAttributeChunkElementCount), {DataInterfaceName}_PointAttributeChunkStrideInBytes
+
 void WriteRootUV_{DataInterfaceName}(uint CurveIndex, float2 RootUV)
 {
 	const uint WriteIndex = {DataInterfaceName}_OutputStreamStart + CurveIndex;
 	if (WriteIndex < {DataInterfaceName}_CurveCount && HasHairStrandsRootUV())
 	{
-		InternalSetHairStrandsAttributes32(CurveIndex, HAIR_CURVE_ATTRIBUTE_OFFSET_ROOTUV({DataInterfaceName}_CurveAttributeOffsets), HAIR_CURVE_ATTRIBUTE_STRIDE_ROOTUV, {DataInterfaceName}_CurveAttributeBufferUAV, PackHairRootUV(RootUV));
+		InternalSetHairStrandsAttributes32(CurveIndex, HAIR_CURVE_ATTRIBUTE_OFFSET_ROOTUV({DataInterfaceName}_CurveAttributeOffsets), HAIR_CURVE_ATTRIBUTE_STRIDE_ROOTUV, {DataInterfaceName}_CurveAttributeBufferUAV, PackHairRootUV(RootUV), {DataInterfaceName}_CurveAttributeIndexToChunkDivAsShift, {DataInterfaceName}_CurveAttributeChunkElementCount), {DataInterfaceName}_CurveAttributeChunkStrideInBytes);
 	}
 }
 
@@ -98,7 +111,7 @@ void WriteSeed_{DataInterfaceName}(uint CurveIndex, float Seed)
 	const uint WriteIndex = {DataInterfaceName}_OutputStreamStart + CurveIndex;
 	if (WriteIndex < {DataInterfaceName}_CurveCount)
 	{
-		InternalSetHairStrandsAttributes(CurveIndex, HAIR_CURVE_ATTRIBUTE_OFFSET_SEED({DataInterfaceName}_CurveAttributeOffsets), HAIR_CURVE_ATTRIBUTE_STRIDE_SEED, {DataInterfaceName}_CurveAttributeBufferUAV, PackHairSeed(Seed));
+		InternalSetHairStrandsAttributes(CurveIndex, HAIR_CURVE_ATTRIBUTE_OFFSET_SEED({DataInterfaceName}_CurveAttributeOffsets), HAIR_CURVE_ATTRIBUTE_STRIDE_SEED, {DataInterfaceName}_CurveAttributeBufferUAV, PackHairSeed(Seed), {DataInterfaceName}_CurveAttributeIndexToChunkDivAsShift, {DataInterfaceName}_CurveAttributeChunkElementCount), {DataInterfaceName}_CurveAttributeChunkStrideInBytes);
 	}
 }
 
@@ -107,7 +120,7 @@ void WriteClumpId_{DataInterfaceName}(uint CurveIndex, uint ClumpId)
 	const uint WriteIndex = {DataInterfaceName}_OutputStreamStart + CurveIndex;
 	if (WriteIndex < {DataInterfaceName}_CurveCount && HasHairStrandsClumpId())
 	{
-		InternalSetHairStrandsAttributes(CurveIndex, HAIR_CURVE_ATTRIBUTE_OFFSET_CLUMPID({DataInterfaceName}_CurveAttributeOffsets), HAIR_CURVE_ATTRIBUTE_STRIDE_CLUMPID, {DataInterfaceName}_CurveAttributeBufferUAV, PackHairClumpID(ClumpId));
+		InternalSetHairStrandsAttributes(CurveIndex, HAIR_CURVE_ATTRIBUTE_OFFSET_CLUMPID({DataInterfaceName}_CurveAttributeOffsets), HAIR_CURVE_ATTRIBUTE_STRIDE_CLUMPID, {DataInterfaceName}_CurveAttributeBufferUAV, PackHairClumpID(ClumpId), {DataInterfaceName}_CurveAttributeIndexToChunkDivAsShift, {DataInterfaceName}_CurveAttributeChunkElementCount), {DataInterfaceName}_CurveAttributeChunkStrideInBytes);
 	}
 }
 
@@ -116,7 +129,7 @@ void WriteColor_{DataInterfaceName}(uint PointIndex, float3 Color)
 	const uint WriteIndex = {DataInterfaceName}_OutputStreamStart + PointIndex;
 	if (WriteIndex < {DataInterfaceName}_PointCount && HasHairStrandsColor())
 	{
-		InternalSetHairStrandsAttributes32(PointIndex, HAIR_POINT_ATTRIBUTE_OFFSET_COLOR({DataInterfaceName}_PointAttributeOffsets), HAIR_POINT_ATTRIBUTE_STRIDE_COLOR, {DataInterfaceName}_PointAttributeBufferUAV, PackHairColor(Color));
+		InternalSetHairStrandsAttributes32(PointIndex, HAIR_POINT_ATTRIBUTE_OFFSET_COLOR({DataInterfaceName}_PointAttributeOffsets), HAIR_POINT_ATTRIBUTE_STRIDE_COLOR, {DataInterfaceName}_PointAttributeBufferUAV, PackHairColor(Color), {DataInterfaceName}_PointAttributeIndexToChunkDivAsShift, {DataInterfaceName}_PointAttributeChunkElementCount), {DataInterfaceName}_PointAttributeChunkStrideInBytes);
 	}
 }
 
@@ -125,7 +138,7 @@ void WriteRoughness_{DataInterfaceName}(uint PointIndex, float Roughness)
 	const uint WriteIndex = {DataInterfaceName}_OutputStreamStart + PointIndex;
 	if (WriteIndex < {DataInterfaceName}_PointCount && HasHairStrandsRoughness())
 	{
-		InternalSetHairStrandsAttributes(PointIndex, HAIR_POINT_ATTRIBUTE_OFFSET_ROUGHNESS({DataInterfaceName}_PointAttributeOffsets), HAIR_POINT_ATTRIBUTE_STRIDE_ROUGHNESS, {DataInterfaceName}_PointAttributeBufferUAV, PackHairRoughness(Roughness));
+		InternalSetHairStrandsAttributes(PointIndex, HAIR_POINT_ATTRIBUTE_OFFSET_ROUGHNESS({DataInterfaceName}_PointAttributeOffsets), HAIR_POINT_ATTRIBUTE_STRIDE_ROUGHNESS, {DataInterfaceName}_PointAttributeBufferUAV, PackHairRoughness(Roughness), {DataInterfaceName}_PointAttributeIndexToChunkDivAsShift, {DataInterfaceName}_PointAttributeChunkElementCount), {DataInterfaceName}_PointAttributeChunkStrideInBytes);
 	}
 }
 
@@ -134,6 +147,6 @@ void WriteAO_{DataInterfaceName}(uint PointIndex, float Ao)
 	const uint WriteIndex = {DataInterfaceName}_OutputStreamStart + PointIndex;
 	if (WriteIndex < {DataInterfaceName}_PointCount && HasHairStrandsAO())
 	{
-		InternalSetHairStrandsAttributes(PointIndex, HAIR_POINT_ATTRIBUTE_OFFSET_AO({DataInterfaceName}_PointAttributeOffsets), HAIR_POINT_ATTRIBUTE_STRIDE_AO, {DataInterfaceName}_PointAttributeBufferUAV, PackHairAO(Ao));
+		InternalSetHairStrandsAttributes(PointIndex, HAIR_POINT_ATTRIBUTE_OFFSET_AO({DataInterfaceName}_PointAttributeOffsets), HAIR_POINT_ATTRIBUTE_STRIDE_AO, {DataInterfaceName}_PointAttributeBufferUAV, PackHairAO(Ao), {DataInterfaceName}_PointAttributeIndexToChunkDivAsShift, {DataInterfaceName}_PointAttributeChunkElementCount), {DataInterfaceName}_PointAttributeChunkStrideInBytes);
 	}
 }
\ No newline at end of file
diff --git a/Engine/Plugins/Runtime/HairStrands/Source/HairStrandsCore/Private/GroomBuilder.cpp b/Engine/Plugins/Runtime/HairStrands/Source/HairStrandsCore/Private/GroomBuilder.cpp
--- a/Engine/Plugins/Runtime/HairStrands/Source/HairStrandsCore/Private/GroomBuilder.cpp
+++ b/Engine/Plugins/Runtime/HairStrands/Source/HairStrandsCore/Private/GroomBuilder.cpp
@@ -40,7 +40,7 @@ static FAutoConsoleVariableRef CVarHairGroupIndexBuilder_MaxVoxelResolution(TEXT
 
 FString FGroomBuilder::GetVersion()
 {
-	return TEXT("v8r9");
+	return TEXT("v8r19");
 }
 
 namespace FHairStrandsDecimation
@@ -187,21 +187,33 @@ namespace HairStrandsBuilder
 		const uint32 Attributes = HairStrands.GetAttributes();
 		const uint32 AttributeFlags = HairStrands.GetAttributeFlags();
 		const bool bHasMultipleClumpIDs = HasHairAttributeFlags(AttributeFlags, EHairAttributeFlags::HasMultipleClumpIDs);
-		TArray<uint32> AttributeRootUV;
-		TArray<uint32> AttributeSeed;
-		TArray<uint32> AttributeLength;
-		TArray<uint32> AttributeClumpIDs;
-		TArray<uint32> AttributeColor;
-		TArray<uint32> AttributeRoughness;
-		TArray<uint32> AttributeAO;
-
-		AttributeSeed.SetNum(FMath::DivideAndRoundUp(NumCurves, 4u));		// 1 bytes encoding - Per-Curve
-		AttributeLength.SetNum(FMath::DivideAndRoundUp(NumCurves, 2u));		// 2 bytes encoding - Per-Curve
-		if (HasHairAttribute(Attributes, EHairAttribute::RootUV)) 	{ AttributeRootUV.SetNum(NumCurves); }																		// 4 bytes encoding - Per-Curve
-		if (HasHairAttribute(Attributes, EHairAttribute::ClumpID))	{ AttributeClumpIDs.SetNum(bHasMultipleClumpIDs ? NumCurves * 2u : FMath::DivideAndRoundUp(NumCurves, 2u));}// 8 or 2 bytes encoding - Per-Curve
-		if (HasHairAttribute(Attributes, EHairAttribute::Color))	{ AttributeColor.SetNum(NumPoints); }																		// 4 bytes encoding - Per-Vertex
-		if (HasHairAttribute(Attributes, EHairAttribute::Roughness)){ AttributeRoughness.SetNum(FMath::DivideAndRoundUp(NumPoints, 2u)); }										// 1 bytes encoding - Per-Vertex
-		if (HasHairAttribute(Attributes, EHairAttribute::AO))		{ AttributeAO.SetNum(FMath::DivideAndRoundUp(NumPoints, 4u)); }												// 1 bytes encoding - Per-Vertex
+
+		// Byte arrays
+		TArray<uint8> AttributeRootUV;
+		TArray<uint8> AttributeSeed;
+		TArray<uint8> AttributeLength;
+		TArray<uint8> AttributeClumpIDs;
+		TArray<uint8> AttributeColor;
+		TArray<uint8> AttributeRoughness;
+		TArray<uint8> AttributeAO;
+
+		// Stride data in bytes
+		const uint32 Stride_Seed 		= 1;
+		const uint32 Stride_Length 		= 2;
+		const uint32 Stride_RootUV 		= 4;
+		const uint32 Stride_ClumpID 	= bHasMultipleClumpIDs ? 8 : 2;
+		const uint32 Stride_Color		= 4;
+		const uint32 Stride_Roughness	= 1;
+		const uint32 Stride_AO			= 1;
+
+		// Ensure all data array are 4bytes aligned so that data are properly padded
+																	{ AttributeSeed		.SetNum(FMath::DivideAndRoundUp(NumCurves * Stride_Seed, 		4u) * 4u); }
+																	{ AttributeLength	.SetNum(FMath::DivideAndRoundUp(NumCurves * Stride_Length, 		4u) * 4u); }
+		if (HasHairAttribute(Attributes, EHairAttribute::RootUV)) 	{ AttributeRootUV	.SetNum(FMath::DivideAndRoundUp(NumCurves * Stride_RootUV, 		4u) * 4u); }
+		if (HasHairAttribute(Attributes, EHairAttribute::ClumpID))	{ AttributeClumpIDs	.SetNum(FMath::DivideAndRoundUp(NumCurves * Stride_ClumpID,		4u) * 4u); }
+		if (HasHairAttribute(Attributes, EHairAttribute::Color))	{ AttributeColor	.SetNum(FMath::DivideAndRoundUp(NumPoints * Stride_Color, 		4u) * 4u); }
+		if (HasHairAttribute(Attributes, EHairAttribute::Roughness)){ AttributeRoughness.SetNum(FMath::DivideAndRoundUp(NumPoints * Stride_Roughness, 	4u) * 4u); }
+		if (HasHairAttribute(Attributes, EHairAttribute::AO))		{ AttributeAO		.SetNum(FMath::DivideAndRoundUp(NumPoints * Stride_AO, 			4u) * 4u); }
 
 		const FVector HairBoxCenter = HairStrands.BoundingBox.GetCenter();
 
@@ -262,7 +274,8 @@ namespace HairStrandsBuilder
 				// Per-Vertex Color
 				if (HasHairAttribute(Attributes, EHairAttribute::Color))
 				{
-					uint32& Packed = AttributeColor[PointIndex + IndexOffset];
+					uint32* Data = (uint32*)AttributeColor.GetData();
+					uint32& Packed = Data[PointIndex + IndexOffset];
 					Packed |= FMath::Clamp(uint32(FMath::Sqrt(Points.PointsBaseColor[PointIndex + IndexOffset].R) * 0x7FF), 0u, 0x7FFu);
 					Packed |= FMath::Clamp(uint32(FMath::Sqrt(Points.PointsBaseColor[PointIndex + IndexOffset].G) * 0x7FF), 0u, 0x7FFu)<<11;
 					Packed |= FMath::Clamp(uint32(FMath::Sqrt(Points.PointsBaseColor[PointIndex + IndexOffset].B) * 0x3FF), 0u, 0x3FFu)<<22;
@@ -271,15 +284,15 @@ namespace HairStrandsBuilder
 				// Per-Vertex Roughness
 				if (HasHairAttribute(Attributes, EHairAttribute::Roughness))
 				{
-					uint8* Packed = (uint8*)AttributeRoughness.GetData();
-					Packed[PointIndex + IndexOffset] = FMath::Clamp(uint32(Points.PointsRoughness[PointIndex + IndexOffset] * 0xFF), 0u, 0xFFu);
+					uint8* Data = (uint8*)AttributeRoughness.GetData();
+					Data[PointIndex + IndexOffset] = FMath::Clamp(uint32(Points.PointsRoughness[PointIndex + IndexOffset] * 0xFF), 0u, 0xFFu);
 				}
 
 				// Per-Vertex AO
 				if (HasHairAttribute(Attributes, EHairAttribute::AO))
 				{
-					uint8* Packed = (uint8*)AttributeAO.GetData();
-					Packed[PointIndex + IndexOffset] = FMath::Clamp(uint32(Points.PointsAO[PointIndex + IndexOffset] * 0xFF), 0u, 0xFFu);
+					uint8* Data = (uint8*)AttributeAO.GetData();
+					Data[PointIndex + IndexOffset] = FMath::Clamp(uint32(Points.PointsAO[PointIndex + IndexOffset] * 0xFF), 0u, 0xFFu);
 				}
 			}
 
@@ -309,7 +322,8 @@ namespace HairStrandsBuilder
 				const FVector2D TextureIndexUV = RootUV - TextureRootUV;
 
 				// UDIM
-				uint32& Packed = AttributeRootUV[CurveIndex];
+				uint32* Data = (uint32*)AttributeRootUV.GetData();
+				uint32& Packed = Data[CurveIndex];
 				Packed = 0;
 				Packed |= (uint32(FMath::Clamp(TextureRootUV.X * 2047.f, 0.f, 2047.f)) & 0x7FFu);
 				Packed |= (uint32(FMath::Clamp(TextureRootUV.Y * 2047.f, 0.f, 2047.f)) & 0x7FFu) << 11u;
@@ -363,47 +377,80 @@ namespace HairStrandsBuilder
 
 		const uint32 UintToByte = 4;
 
+		const uint32 CurveAttributeChunkElementCount = 1024; // Make it asset dependent?
+		const uint32 PointAttributeChunkElementCount = 8192; // Make it asset dependent?
+
 		// Concatenate all curve-attributes
+		uint32 CurveAttributeChunkStride = 0;
 		{
 			TArray<FHairStrandsAttributeFormat::Type> OutPackedAttributes;
-			OutPackedAttributes.Reserve(
+			OutPackedAttributes.Reserve(FMath::DivideAndRoundUp(
 				AttributeRootUV.Num() + 
 				AttributeSeed.Num() +
 				AttributeLength.Num() +
-				AttributeClumpIDs.Num());
+				AttributeClumpIDs.Num(), 4));
 
 			for (uint32 AttributeIt=0; AttributeIt< HAIR_CURVE_ATTRIBUTE_COUNT; ++AttributeIt)
 			{
 				OutBulkData.Header.CurveAttributeOffsets[AttributeIt] = HAIR_ATTRIBUTE_INVALID_OFFSET;
 			}
 
-			if (HasHairAttribute(Attributes, EHairAttribute::RootUV))
+			auto AppendAttribute = [&](uint32 InAttributeIndex, uint32 InStrideInBytes, const TArray<uint8>& InData, uint32 InChunkIndex, uint32 InChunkCount)
 			{
-				OutBulkData.Header.CurveAttributeOffsets[HAIR_CURVE_ATTRIBUTE_ROOTUV] = OutPackedAttributes.Num() * UintToByte;
-				OutPackedAttributes.Append(AttributeRootUV);
-			}
+				const uint32 ElementIndex 	= InChunkIndex * CurveAttributeChunkElementCount;
+				const uint32 ElementCount 	= FMath::Min(OutBulkData.Header.CurveCount - ElementIndex, CurveAttributeChunkElementCount);
 
-			{
-				OutBulkData.Header.CurveAttributeOffsets[HAIR_CURVE_ATTRIBUTE_SEED]   = OutPackedAttributes.Num() * UintToByte;
-				OutPackedAttributes.Append(AttributeSeed);
-			}
+				const uint32 OffsetInBytes	= ElementIndex * InStrideInBytes;
+				const uint32 SizeInBytes	= ElementCount * InStrideInBytes;
+				const uint32*DataInUints	= (uint32*)(InData.GetData() + OffsetInBytes);
+				const uint32 SizeInUints	= FMath::DivideAndRoundUp(SizeInBytes, 4u);
 
-			{
-				OutBulkData.Header.CurveAttributeOffsets[HAIR_CURVE_ATTRIBUTE_LENGTH] = OutPackedAttributes.Num() * UintToByte;
-				OutPackedAttributes.Append(AttributeLength);
-			}
+				OutPackedAttributes.Append(MakeArrayView(DataInUints, SizeInUints));
 
-			if (HasHairAttribute(Attributes, EHairAttribute::ClumpID))
+				// If the attribute has several chunk, pad the last chunk with 0u, so that 
+				// the attribute offset remain correct for this last chunk
+				if (InChunkCount > 1 && InChunkIndex == InChunkCount-1)
+				{
+					const uint32 PaddingElementCount = CurveAttributeChunkElementCount - ElementCount;
+					const uint32 PaddingSizeInBytes	 = PaddingElementCount * InStrideInBytes;
+					const uint32 PaddingSizeInUints	 = FMath::DivideAndRoundUp(PaddingSizeInBytes, 4u);
+					for (uint32 PadIt = 0; PadIt < PaddingSizeInUints; ++PadIt)
+					{
+						OutPackedAttributes.Add(0u);
+					}
+				}
+
+				if (InChunkIndex == 0)
+				{
+					OutBulkData.Header.CurveAttributeOffsets[InAttributeIndex] = CurveAttributeChunkStride;
+					CurveAttributeChunkStride = OutPackedAttributes.Num() * 4u;
+				}
+			};
+
+			// Concatenate all curve-attribute into chunks
+			// [            Chunk0            ] [            Chunk1            ] [ ... ]
+			// [   RootUV     ][     Seed     ] [   RootUV     ][     Seed     ] [ ... ]
+			// [C0][C1][C2][C3][C0][C1][C2][C3] [C4][C5][C6][C7][C4][C5][C6][C7] [ ... ]
+			const uint32 ChunkCount = FMath::DivideAndRoundUp(NumCurves, CurveAttributeChunkElementCount);
+			const uint32 ChunkElementCount = FMath::Min(CurveAttributeChunkElementCount, OutBulkData.Header.CurveCount);
+			for (uint32 ChunkIt=0; ChunkIt<ChunkCount; ++ChunkIt)
 			{
-				if (HasHairAttributeFlags(OutBulkData.Header.ImportedAttributeFlags, EHairAttributeFlags::HasMultipleClumpIDs))
+				if (HasHairAttribute(Attributes, EHairAttribute::RootUV))
 				{
-					OutBulkData.Header.CurveAttributeOffsets[HAIR_CURVE_ATTRIBUTE_CLUMPID3] = OutPackedAttributes.Num() * UintToByte;
-					OutPackedAttributes.Append(AttributeClumpIDs);
+					AppendAttribute(HAIR_CURVE_ATTRIBUTE_ROOTUV, Stride_RootUV, AttributeRootUV, ChunkIt, ChunkCount);
 				}
-				else
+	
 				{
-					OutBulkData.Header.CurveAttributeOffsets[HAIR_CURVE_ATTRIBUTE_CLUMPID] = OutPackedAttributes.Num() * UintToByte;
-					OutPackedAttributes.Append(AttributeClumpIDs);
+					AppendAttribute(HAIR_CURVE_ATTRIBUTE_SEED, Stride_Seed, AttributeSeed, ChunkIt, ChunkCount);
+				}
+	
+				{
+					AppendAttribute(HAIR_CURVE_ATTRIBUTE_LENGTH, Stride_Length, AttributeLength, ChunkIt, ChunkCount);
+				}
+	
+				if (HasHairAttribute(Attributes, EHairAttribute::ClumpID))
+				{
+					AppendAttribute(bHasMultipleClumpIDs ? HAIR_CURVE_ATTRIBUTE_CLUMPID3 : HAIR_CURVE_ATTRIBUTE_CLUMPID, Stride_ClumpID, AttributeClumpIDs, ChunkIt, ChunkCount);
 				}
 			}
 
@@ -411,34 +458,74 @@ namespace HairStrandsBuilder
 		}
 
 		// Concatenate all point-attributes
+		uint32 PointAttributeChunkStride = 0;
 		{
 			TArray<FHairStrandsAttributeFormat::Type> OutPackedAttributes;
-			OutPackedAttributes.Reserve(
+			OutPackedAttributes.Reserve(FMath::DivideAndRoundUp(
 				AttributeColor.Num() +
 				AttributeRoughness.Num() +
-				AttributeAO.Num());
+				AttributeAO.Num(), 4));
 
 			for (uint32 AttributeIt=0; AttributeIt< HAIR_POINT_ATTRIBUTE_COUNT; ++AttributeIt)
 			{
 				OutBulkData.Header.PointAttributeOffsets[AttributeIt] = HAIR_ATTRIBUTE_INVALID_OFFSET;
 			}
 			
-			if (HasHairAttribute(Attributes, EHairAttribute::Color))
+			auto AppendAttribute = [&](uint32 InAttributeIndex, uint32 InStrideInBytes, const TArray<uint8>& InData, uint32 InChunkIndex, uint32 InChunkCount)
 			{
-				OutBulkData.Header.PointAttributeOffsets[HAIR_POINT_ATTRIBUTE_COLOR] = OutPackedAttributes.Num() * UintToByte;
-				OutPackedAttributes.Append(AttributeColor);
-			}
+				const uint32 ElementIndex 	= InChunkIndex * PointAttributeChunkElementCount;
+				const uint32 ElementCount 	= FMath::Min(OutBulkData.Header.PointCount - ElementIndex, PointAttributeChunkElementCount);
 
-			if (HasHairAttribute(Attributes, EHairAttribute::Roughness))
-			{
-				OutBulkData.Header.PointAttributeOffsets[HAIR_POINT_ATTRIBUTE_ROUGHNESS] = OutPackedAttributes.Num() * UintToByte;
-				OutPackedAttributes.Append(AttributeRoughness);
-			}
+				const uint32 OffsetInBytes	= ElementIndex * InStrideInBytes;
+				const uint32 SizeInBytes	= ElementCount * InStrideInBytes;
+				const uint32*DataInUints	= (uint32*)(InData.GetData() + OffsetInBytes);
+				const uint32 SizeInUints	= FMath::DivideAndRoundUp(SizeInBytes, 4u);
+
+				OutPackedAttributes.Append(MakeArrayView(DataInUints, SizeInUints));
+
+				// If the attribute has several chunk, pad the last chunk with 0u, so that 
+				// the attribute offset remain correct for this last chunk
+				if (InChunkCount > 1 && InChunkIndex == InChunkCount-1)
+				{
+					const uint32 PaddingElementCount = PointAttributeChunkElementCount - ElementCount;
+					const uint32 PaddingSizeInBytes	 = PaddingElementCount * InStrideInBytes;
+					const uint32 PaddingSizeInUints	 = FMath::DivideAndRoundUp(PaddingSizeInBytes, 4u);
+					for (uint32 PadIt = 0; PadIt < PaddingSizeInUints; ++PadIt)
+					{
+						OutPackedAttributes.Add(0u);
+					}
+				}
 
-			if (HasHairAttribute(Attributes, EHairAttribute::AO))
+				if (InChunkIndex == 0)
+				{
+					OutBulkData.Header.PointAttributeOffsets[InAttributeIndex] = PointAttributeChunkStride;
+					PointAttributeChunkStride = OutPackedAttributes.Num() * 4u;
+				}
+			};
+
+			// Concatenate all point-attributes into chunks
+			// [            Chunk0            ] [            Chunk1            ] [ ... ]
+			// [    Color     ][      AO      ] [    Color     ][      AO      ] [ ... ]
+			// [P0][P1][P2][P3][P0][P1][P2][P3] [P4][P5][P6][P7][P4][P5][P6][P7] [ ... ]
+			const uint32 ChunkCount = FMath::DivideAndRoundUp(NumPoints, PointAttributeChunkElementCount);
+			const uint32 ChunkElementCount = FMath::Min(PointAttributeChunkElementCount, OutBulkData.Header.PointCount);
+			for (uint32 ChunkIt=0; ChunkIt<ChunkCount; ++ChunkIt)
 			{
-				OutBulkData.Header.PointAttributeOffsets[HAIR_POINT_ATTRIBUTE_AO] = OutPackedAttributes.Num() * UintToByte;
-				OutPackedAttributes.Append(AttributeAO);
+
+				if (HasHairAttribute(Attributes, EHairAttribute::Color))
+				{
+					AppendAttribute(HAIR_POINT_ATTRIBUTE_COLOR, Stride_Color, AttributeColor, ChunkIt, ChunkCount);
+				}
+	
+				if (HasHairAttribute(Attributes, EHairAttribute::Roughness))
+				{
+					AppendAttribute(HAIR_POINT_ATTRIBUTE_ROUGHNESS, Stride_Roughness, AttributeRoughness, ChunkIt, ChunkCount);
+				}
+	
+				if (HasHairAttribute(Attributes, EHairAttribute::AO))
+				{
+					AppendAttribute(HAIR_POINT_ATTRIBUTE_AO, Stride_AO, AttributeAO, ChunkIt, ChunkCount);
+				}
 			}
 
 			const bool bHasPointAttribute = OutPackedAttributes.Num() > 0;
@@ -460,6 +547,17 @@ namespace HairStrandsBuilder
 		{
 			CopyToBulkData<FHairStrandsPointToCurveFormat32>(OutBulkData.Data.PointToCurve, OutPointToCurve32);
 		}
+
+		// Stride datas
+		OutBulkData.Header.Strides.PositionStride = FHairStrandsPositionFormat::SizeInByte;
+		OutBulkData.Header.Strides.CurveStride = FHairStrandsCurveFormat::SizeInByte;
+		OutBulkData.Header.Strides.PointToCurveStride = (bUse16bitsCurveIndex ? FHairStrandsPointToCurveFormat16::SizeInByte : FHairStrandsPointToCurveFormat32::SizeInByte);
+
+		OutBulkData.Header.Strides.CurveAttributeChunkStride = CurveAttributeChunkStride;
+		OutBulkData.Header.Strides.PointAttributeChunkStride = PointAttributeChunkStride;
+
+		OutBulkData.Header.Strides.CurveAttributeChunkElementCount = CurveAttributeChunkElementCount;
+		OutBulkData.Header.Strides.PointAttributeChunkElementCount = PointAttributeChunkElementCount;
 	}
 
 	void BuildRenderData(const FHairStrandsDatas& HairStrands, FHairStrandsBulkData& OutBulkData)
diff --git a/Engine/Plugins/Runtime/HairStrands/Source/HairStrandsCore/Private/GroomTextureBuilder.cpp b/Engine/Plugins/Runtime/HairStrands/Source/HairStrandsCore/Private/GroomTextureBuilder.cpp
--- a/Engine/Plugins/Runtime/HairStrands/Source/HairStrandsCore/Private/GroomTextureBuilder.cpp
+++ b/Engine/Plugins/Runtime/HairStrands/Source/HairStrandsCore/Private/GroomTextureBuilder.cpp
@@ -954,6 +954,8 @@ static void ExportTextures(FRDGBuilder& GraphBuilder, const FHairStrandsRDGTextu
 	Out.TileCount = In.TileCount;
 }
 
+void GetHairStrandsAttributeParameter(const FHairStrandsBulkData& In, FHairStrandsInstanceAttributeParameters& Out);
+
 static bool TraceTextures(
 	FRDGBuilder& GraphBuilder,
 	FGlobalShaderMap* ShaderMap,
@@ -1060,8 +1062,7 @@ static bool TraceTextures(
 			Instance.Common.LengthScale = 1.0f;
 			Instance.Common.Length = GroupData.Strands.RestResource->BulkData.Header.MaxLength;
 			Instance.Common.PositionOffset = FVector3f(GroupData.Strands.RestResource->GetPositionOffset());
-			PACK_CURVE_HAIR_ATTRIBUTE_OFFSETS(Instance.Common.CurveAttributeOffsets, GroupData.Strands.RestResource->BulkData.Header.CurveAttributeOffsets);
-			PACK_POINT_HAIR_ATTRIBUTE_OFFSETS(Instance.Common.PointAttributeOffsets, GroupData.Strands.RestResource->BulkData.Header.PointAttributeOffsets);
+			GetHairStrandsAttributeParameter(GroupData.Strands.RestResource->BulkData, Instance.Common.Attributes);
 
 			Instance.Resources.PositionBuffer = GroupData.Strands.RestResource->PositionBuffer.SRV;
 			Instance.Resources.CurveAttributeBuffer = GroupData.Strands.RestResource->CurveAttributeBuffer.SRV;
diff --git a/Engine/Plugins/Runtime/HairStrands/Source/HairStrandsCore/Private/HairStrandsDatas.cpp b/Engine/Plugins/Runtime/HairStrands/Source/HairStrandsCore/Private/HairStrandsDatas.cpp
--- a/Engine/Plugins/Runtime/HairStrands/Source/HairStrandsCore/Private/HairStrandsDatas.cpp
+++ b/Engine/Plugins/Runtime/HairStrands/Source/HairStrandsCore/Private/HairStrandsDatas.cpp
@@ -250,6 +250,14 @@ void FHairStrandsBulkData::SerializeHeader(FArchive& Ar, UObject* Owner)
 	}
 	Ar << Header.ImportedAttributes;
 	Ar << Header.ImportedAttributeFlags;
+
+	Ar << Header.Strides.PositionStride;
+	Ar << Header.Strides.CurveStride;
+	Ar << Header.Strides.PointToCurveStride;
+	Ar << Header.Strides.CurveAttributeChunkStride;
+	Ar << Header.Strides.PointAttributeChunkStride;
+	Ar << Header.Strides.CurveAttributeChunkElementCount;
+	Ar << Header.Strides.PointAttributeChunkElementCount;
 }
 
 void FHairStrandsBulkData::GetResourceVersion(FArchive& Ar) const
@@ -294,11 +302,11 @@ void FHairStrandsBulkData::Reset()
 	Header.Flags = 0;
 	for (uint8 AttributeIt = 0; AttributeIt < HAIR_CURVE_ATTRIBUTE_COUNT; ++AttributeIt)
 	{
-		Header.CurveAttributeOffsets[AttributeIt] = 0xFFFFFFFF;
+		Header.CurveAttributeOffsets[AttributeIt] = HAIR_ATTRIBUTE_INVALID_OFFSET;
 	}
 	for (uint8 AttributeIt = 0; AttributeIt < HAIR_POINT_ATTRIBUTE_COUNT; ++AttributeIt)
 	{
-		Header.PointAttributeOffsets[AttributeIt] = 0xFFFFFFFF;
+		Header.PointAttributeOffsets[AttributeIt] = HAIR_ATTRIBUTE_INVALID_OFFSET;
 	}
 	// Deallocate memory if needed
 	Data.Positions.RemoveBulkData();
diff --git a/Engine/Plugins/Runtime/HairStrands/Source/HairStrandsCore/Private/HairStrandsInterpolation.cpp b/Engine/Plugins/Runtime/HairStrands/Source/HairStrandsCore/Private/HairStrandsInterpolation.cpp
--- a/Engine/Plugins/Runtime/HairStrands/Source/HairStrandsCore/Private/HairStrandsInterpolation.cpp
+++ b/Engine/Plugins/Runtime/HairStrands/Source/HairStrandsCore/Private/HairStrandsInterpolation.cpp
@@ -92,6 +92,21 @@ enum class EDeformationType : uint8
 	OffsetGuide		// Offset the guides
 };
 
+void GetHairStrandsAttributeParameter(const FHairStrandsBulkData& In, FHairStrandsInstanceAttributeParameters& Out)
+{
+	check(FMath::IsPowerOfTwo(In.Header.Strides.CurveAttributeChunkElementCount));
+	Out.CurveAttributeIndexToChunkDivAsShift	= FMath::FloorToInt(FMath::Log2(float(In.Header.Strides.CurveAttributeChunkElementCount)));
+	Out.CurveAttributeChunkElementCount			= In.Header.Strides.CurveAttributeChunkElementCount;
+	Out.CurveAttributeChunkStrideInBytes		= In.Header.Strides.CurveAttributeChunkStride;
+
+	check(FMath::IsPowerOfTwo(In.Header.Strides.PointAttributeChunkElementCount));
+	Out.PointAttributeIndexToChunkDivAsShift	= FMath::FloorToInt(FMath::Log2(float(In.Header.Strides.PointAttributeChunkElementCount)));
+	Out.PointAttributeChunkElementCount     	= In.Header.Strides.PointAttributeChunkElementCount;
+	Out.PointAttributeChunkStrideInBytes    	= In.Header.Strides.PointAttributeChunkStride;
+
+	PACK_CURVE_HAIR_ATTRIBUTE_OFFSETS(Out.CurveAttributeOffsets, In.Header.CurveAttributeOffsets);
+	PACK_POINT_HAIR_ATTRIBUTE_OFFSETS(Out.PointAttributeOffsets, In.Header.PointAttributeOffsets);
+}
 
 struct FGroomCacheResources
 {
@@ -724,7 +739,7 @@ class FHairPatchAttributeCS : public FGlobalShader
 
 	BEGIN_SHADER_PARAMETER_STRUCT(FParameters, )
 		SHADER_PARAMETER(uint32, CurveCount)
-		SHADER_PARAMETER_ARRAY(FUintVector4, RenCurveAttributeOffsets, [HAIR_CURVE_ATTRIBUTE_OFFSET_COUNT])
+		SHADER_PARAMETER_STRUCT_INCLUDE(FHairStrandsInstanceAttributeParameters, RenAttributes)
 		SHADER_PARAMETER_RDG_BUFFER_SRV(Buffer, RenCurveBuffer)
 		SHADER_PARAMETER_RDG_BUFFER_SRV(ByteAddressBuffer, InterpolationBuffer)
 		SHADER_PARAMETER_RDG_BUFFER_SRV(Buffer, RenVertexToClusterIdBuffer)
@@ -756,7 +771,7 @@ static void AddPatchAttributePass(
 	const EHairPatchAttribute Mode,
 	const bool bSimulation,
 	const bool bUseSingleGuide,
-	const uint32* RenCurveAttributeOffests,
+	const FHairStrandsBulkData& RenBulkData,
 	const FRDGBufferRef& RenAttributeBuffer,
 	const FRDGBufferSRVRef& RenCurveBuffer,
 	const FRDGBufferSRVRef& RenVertexToClusterIdBuffer,
@@ -776,7 +791,7 @@ static void AddPatchAttributePass(
 	Parameters->CurveCount = CurveCount;
 	Parameters->RenCurveBuffer = RenCurveBuffer;
 	Parameters->RenVertexToClusterIdBuffer = RenVertexToClusterIdBuffer;
-	PACK_CURVE_HAIR_ATTRIBUTE_OFFSETS(Parameters->RenCurveAttributeOffsets, RenCurveAttributeOffests)
+	GetHairStrandsAttributeParameter(RenBulkData, Parameters->RenAttributes);
 
 	if (bSimulation)
 	{
@@ -1571,11 +1586,11 @@ static FHairGroupPublicData::FVertexFactoryInput InternalComputeHairStrandsVerte
 		const FUintVector4 InvalidOffset = FUintVector4(HAIR_ATTRIBUTE_INVALID_OFFSET,HAIR_ATTRIBUTE_INVALID_OFFSET,HAIR_ATTRIBUTE_INVALID_OFFSET,HAIR_ATTRIBUTE_INVALID_OFFSET);
 		for (uint32 OffsetIt = 0; OffsetIt < HAIR_CURVE_ATTRIBUTE_OFFSET_COUNT; ++OffsetIt)
 		{
-			OutVFInput.Strands.Common.CurveAttributeOffsets[OffsetIt] = InvalidOffset;
+			OutVFInput.Strands.Common.Attributes.CurveAttributeOffsets[OffsetIt] = InvalidOffset;
 		}
 		for (uint32 OffsetIt = 0; OffsetIt < HAIR_POINT_ATTRIBUTE_OFFSET_COUNT; ++OffsetIt)
 		{
-			OutVFInput.Strands.Common.PointAttributeOffsets[OffsetIt] = InvalidOffset;
+			OutVFInput.Strands.Common.Attributes.PointAttributeOffsets[OffsetIt] = InvalidOffset;
 		}
 		return OutVFInput;
 	}
@@ -1684,13 +1699,11 @@ static FHairGroupPublicData::FVertexFactoryInput InternalComputeHairStrandsVerte
 
 	if (bUseGuideAttributeOffsets)
 	{
-		PACK_CURVE_HAIR_ATTRIBUTE_OFFSETS(OutVFInput.Strands.Common.CurveAttributeOffsets, Instance->Guides.Data->Header.CurveAttributeOffsets);
-		PACK_POINT_HAIR_ATTRIBUTE_OFFSETS(OutVFInput.Strands.Common.PointAttributeOffsets, Instance->Guides.Data->Header.PointAttributeOffsets);
+		GetHairStrandsAttributeParameter(*Instance->Guides.Data, OutVFInput.Strands.Common.Attributes);
 	}
 	else
 	{
-		PACK_CURVE_HAIR_ATTRIBUTE_OFFSETS(OutVFInput.Strands.Common.CurveAttributeOffsets, Instance->Strands.Data->Header.CurveAttributeOffsets);
-		PACK_POINT_HAIR_ATTRIBUTE_OFFSETS(OutVFInput.Strands.Common.PointAttributeOffsets, Instance->Strands.Data->Header.PointAttributeOffsets);
+		GetHairStrandsAttributeParameter(*Instance->Strands.Data, OutVFInput.Strands.Common.Attributes);
 	}
 
 #if RHI_RAYTRACING
@@ -2088,6 +2101,7 @@ void ComputeHairStrandsInterpolation(
 				}
 				FRDGImportedBuffer OutRenCurveAttributeBuffer = Register(GraphBuilder, Instance->Strands.DebugCurveAttributeBuffer, ERDGImportedBufferFlags::CreateUAV);
 
+				check(Instance->Strands.Data);
 				AddPatchAttributePass(
 					GraphBuilder,
 					ShaderMap,
@@ -2095,7 +2109,7 @@ void ComputeHairStrandsInterpolation(
 					PatchMode,
 					bValidGuide,
 					bUseSingleGuide,
-					Instance->Strands.Data->Header.CurveAttributeOffsets,
+					*Instance->Strands.Data,
 					Register(GraphBuilder, Instance->Strands.RestResource->CurveAttributeBuffer, ERDGImportedBufferFlags::CreateSRV).Buffer,
 					RegisterAsSRV(GraphBuilder, Instance->Strands.RestResource->CurveBuffer),
 					RegisterAsSRV(GraphBuilder, Instance->Strands.ClusterCullingResource->VertexToClusterIdBuffer),
diff --git a/Engine/Plugins/Runtime/HairStrands/Source/HairStrandsCore/Public/HairStrandsDatas.h b/Engine/Plugins/Runtime/HairStrands/Source/HairStrandsCore/Public/HairStrandsDatas.h
--- a/Engine/Plugins/Runtime/HairStrands/Source/HairStrandsCore/Public/HairStrandsDatas.h
+++ b/Engine/Plugins/Runtime/HairStrands/Source/HairStrandsCore/Public/HairStrandsDatas.h
@@ -499,6 +499,20 @@ struct HAIRSTRANDSCORE_API FHairStrandsBulkData : FHairStrandsBulkCommon
 		/** Imported attribute info */
 		uint32 ImportedAttributes = 0;
 		uint32 ImportedAttributeFlags = 0;
+
+		// Data strides
+		struct FStrides
+		{
+			uint32 PositionStride = 0;
+			uint32 CurveStride = 0;
+			uint32 PointToCurveStride = 0;
+			uint32 CurveAttributeChunkStride = 0;
+			uint32 PointAttributeChunkStride = 0;
+
+			// Number of element per chunk block
+			uint32 CurveAttributeChunkElementCount = 0;
+			uint32 PointAttributeChunkElementCount = 0;
+		} Strides;
 	} Header;
 
 	struct FData
diff --git a/Engine/Shaders/Private/HairStrands/HairStrandsAttributeCommon.ush b/Engine/Shaders/Private/HairStrands/HairStrandsAttributeCommon.ush
--- a/Engine/Shaders/Private/HairStrands/HairStrandsAttributeCommon.ush
+++ b/Engine/Shaders/Private/HairStrands/HairStrandsAttributeCommon.ush
@@ -30,8 +30,13 @@
 
 ////////////////////////////////////////////////////////////////////////////////
 // Attribute access functions
-uint InternalGetHairStrandsAttributes(uint Index, uint AttributeOffsetInBytes, uint AttributeStrideInBytes, ByteAddressBuffer InAttributeBuffer)
+uint InternalGetHairStrandsAttributes(uint Index, uint AttributeOffsetInBytes, uint AttributeStrideInBytes, ByteAddressBuffer InAttributeBuffer, uint InIndexToChunkDivAsShift, uint InChunkElementCount, uint InChunkStrideInBytes)
 {
+	// Adapt index and offset
+	const uint ChunkIndex   = Index >> InIndexToChunkDivAsShift;
+	Index                  -= ChunkIndex * InChunkElementCount;
+	AttributeOffsetInBytes += ChunkIndex * InChunkStrideInBytes;
+
 	// Ensure the reading address is 4-byte aligned
 	const uint Address = AttributeOffsetInBytes + Index * AttributeStrideInBytes;
 	const uint Address_4BytesAligned = Address & (~0x3);
@@ -42,8 +47,12 @@ uint InternalGetHairStrandsAttributes(uint Index, uint AttributeOffsetInBytes, u
 	return Packed >> (8u * ElementIndex);
 }
 
-uint2 InternalGetHairStrandsAttributes64(uint Index, uint AttributeOffsetInBytes, uint AttributeStrideInBytes, ByteAddressBuffer InAttributeBuffer)
+uint2 InternalGetHairStrandsAttributes64(uint Index, uint AttributeOffsetInBytes, uint AttributeStrideInBytes, ByteAddressBuffer InAttributeBuffer, uint InIndexToChunkDivAsShift, uint InChunkElementCount, uint InChunkStrideInBytes)
 {
+	const uint ChunkIndex   = Index >> InIndexToChunkDivAsShift;
+	Index                  -= ChunkIndex * InChunkElementCount;
+	AttributeOffsetInBytes += ChunkIndex * InChunkStrideInBytes;
+
 	// Ensure the reading address is 4-byte aligned
 	const uint Address = AttributeOffsetInBytes + Index * AttributeStrideInBytes;
 	const uint Address_4BytesAligned = Address & (~0x3);
@@ -55,21 +64,21 @@ uint2 InternalGetHairStrandsAttributes64(uint Index, uint AttributeOffsetInBytes
 	return OutPacked;
 }
 
-uint InternalGetHairStrandsAttributesPerCurve(uint HairPrimitiveId, uint AttributeOffsetInBytes, uint AttributeStrideInBytes, ByteAddressBuffer InAttributeBuffer, Buffer<uint> InPointToCurveBuffer)
+uint InternalGetHairStrandsAttributesPerCurve(uint HairPrimitiveId, uint AttributeOffsetInBytes, uint AttributeStrideInBytes, ByteAddressBuffer InAttributeBuffer, Buffer<uint> InPointToCurveBuffer, uint InIndexToChunkDivAsShift, uint InChunkElementCount, uint InChunkStrideInBytes)
 {
 	const uint CurveIndex = InPointToCurveBuffer[HairPrimitiveId];
-	return InternalGetHairStrandsAttributes(CurveIndex, AttributeOffsetInBytes, AttributeStrideInBytes, InAttributeBuffer);
+	return InternalGetHairStrandsAttributes(CurveIndex, AttributeOffsetInBytes, AttributeStrideInBytes, InAttributeBuffer, InIndexToChunkDivAsShift, InChunkElementCount, InChunkStrideInBytes);
 }
 
-uint2 InternalGetHairStrandsAttributesPerCurve64(uint HairPrimitiveId, uint AttributeOffsetInBytes, uint AttributeStrideInBytes, ByteAddressBuffer InAttributeBuffer, Buffer<uint> InPointToCurveBuffer)
+uint2 InternalGetHairStrandsAttributesPerCurve64(uint HairPrimitiveId, uint AttributeOffsetInBytes, uint AttributeStrideInBytes, ByteAddressBuffer InAttributeBuffer, Buffer<uint> InPointToCurveBuffer, uint InIndexToChunkDivAsShift, uint InChunkElementCount, uint InChunkStrideInBytes)
 {
 	const uint CurveIndex = InPointToCurveBuffer[HairPrimitiveId];
-	return InternalGetHairStrandsAttributes64(CurveIndex, AttributeOffsetInBytes, AttributeStrideInBytes, InAttributeBuffer);
+	return InternalGetHairStrandsAttributes64(CurveIndex, AttributeOffsetInBytes, AttributeStrideInBytes, InAttributeBuffer, InIndexToChunkDivAsShift, InChunkElementCount, InChunkStrideInBytes);
 }
 
-uint InternalGetHairStrandsAttributesPerVertex(uint HairPrimitiveId, uint AttributeOffsetInBytes, uint AttributeStrideInBytes, ByteAddressBuffer InAttributeBuffer)
+uint InternalGetHairStrandsAttributesPerVertex(uint HairPrimitiveId, uint AttributeOffsetInBytes, uint AttributeStrideInBytes, ByteAddressBuffer InAttributeBuffer, uint InIndexToChunkDivAsShift, uint InChunkElementCount, uint InChunkStrideInBytes)
 {
-	return InternalGetHairStrandsAttributes(HairPrimitiveId, AttributeOffsetInBytes, AttributeStrideInBytes, InAttributeBuffer);
+	return InternalGetHairStrandsAttributes(HairPrimitiveId, AttributeOffsetInBytes, AttributeStrideInBytes, InAttributeBuffer, InIndexToChunkDivAsShift, InChunkElementCount, InChunkStrideInBytes);
 }
 
 float InternalGetHairStrandsCoordU(Buffer<uint4> InPositionBuffer, uint VertexId, float InRadius, float InRootScale, float InTipScale)
diff --git a/Engine/Shaders/Private/HairStrands/HairStrandsAttributeTemplate.ush b/Engine/Shaders/Private/HairStrands/HairStrandsAttributeTemplate.ush
--- a/Engine/Shaders/Private/HairStrands/HairStrandsAttributeTemplate.ush
+++ b/Engine/Shaders/Private/HairStrands/HairStrandsAttributeTemplate.ush
@@ -22,7 +22,7 @@ float2 GetHairStrandsUV(uint InPrimitiveId, float2 InSegmentUV)
 }
 float2 GetHairStrandsDimensions(uint InPrimitiveId, float InSegmentU)
 {
-	const float CurveLength = UnpackHairLength(InternalGetHairStrandsAttributesPerCurve(InPrimitiveId, HAIR_CURVE_ATTRIBUTE_OFFSET_LENGTH(PREFIX(CurveAttributeOffsets)), HAIR_CURVE_ATTRIBUTE_STRIDE_LENGTH, PREFIX(CurveAttributeBuffer), PREFIX(PointToCurveBuffer)));
+	const float CurveLength = UnpackHairLength(InternalGetHairStrandsAttributesPerCurve(InPrimitiveId, HAIR_CURVE_ATTRIBUTE_OFFSET_LENGTH(PREFIX(CurveAttributeOffsets)), HAIR_CURVE_ATTRIBUTE_STRIDE_LENGTH, PREFIX(CurveAttributeBuffer), PREFIX(PointToCurveBuffer), PREFIX(CurveAttributeIndexToChunkDivAsShift), PREFIX(CurveAttributeChunkElementCount), PREFIX(CurveAttributeChunkStrideInBytes)));
 
 	const float CoordU0 = InternalGetHairStrandsCoordU(PREFIX(PositionBuffer), InPrimitiveId,   PREFIX(Radius), PREFIX(RootScale), PREFIX(TipScale));
 	const float CoordU1 = InternalGetHairStrandsCoordU(PREFIX(PositionBuffer), InPrimitiveId+1, PREFIX(Radius), PREFIX(RootScale), PREFIX(TipScale));
@@ -39,7 +39,7 @@ float2 GetHairStrandsRootUV(uint InPrimitiveId)
 	const uint RootUVOffset = HAIR_CURVE_ATTRIBUTE_OFFSET_ROOTUV(PREFIX(CurveAttributeOffsets));
 	if (IsAttributeValid(RootUVOffset))
 	{
-		return UnpackHairRootUV(InternalGetHairStrandsAttributesPerCurve(InPrimitiveId, RootUVOffset, HAIR_CURVE_ATTRIBUTE_STRIDE_ROOTUV, PREFIX(CurveAttributeBuffer), PREFIX(PointToCurveBuffer)));
+		return UnpackHairRootUV(InternalGetHairStrandsAttributesPerCurve(InPrimitiveId, RootUVOffset, HAIR_CURVE_ATTRIBUTE_STRIDE_ROOTUV, PREFIX(CurveAttributeBuffer), PREFIX(PointToCurveBuffer), PREFIX(CurveAttributeIndexToChunkDivAsShift), PREFIX(CurveAttributeChunkElementCount), PREFIX(CurveAttributeChunkStrideInBytes)));
 	}
 	else
 	{
@@ -48,7 +48,7 @@ float2 GetHairStrandsRootUV(uint InPrimitiveId)
 }
 float GetHairStrandsSeed(uint InPrimitiveId)
 {
-	return UnpackHairSeed(InternalGetHairStrandsAttributesPerCurve(InPrimitiveId, HAIR_CURVE_ATTRIBUTE_OFFSET_SEED(PREFIX(CurveAttributeOffsets)), HAIR_CURVE_ATTRIBUTE_STRIDE_SEED, PREFIX(CurveAttributeBuffer), PREFIX(PointToCurveBuffer)));
+	return UnpackHairSeed(InternalGetHairStrandsAttributesPerCurve(InPrimitiveId, HAIR_CURVE_ATTRIBUTE_OFFSET_SEED(PREFIX(CurveAttributeOffsets)), HAIR_CURVE_ATTRIBUTE_STRIDE_SEED, PREFIX(CurveAttributeBuffer), PREFIX(PointToCurveBuffer), PREFIX(CurveAttributeIndexToChunkDivAsShift), PREFIX(CurveAttributeChunkElementCount), PREFIX(CurveAttributeChunkStrideInBytes)));
 }
 uint3 GetHairStrandsClumpID(uint InPrimitiveId)
 {
@@ -56,11 +56,11 @@ uint3 GetHairStrandsClumpID(uint InPrimitiveId)
 	const uint ClumpId3Offset = HAIR_CURVE_ATTRIBUTE_OFFSET_CLUMPID3(PREFIX(CurveAttributeOffsets));
 	if (IsAttributeValid(ClumpIdOffset))
 	{
-		return UnpackHairClumpID(InternalGetHairStrandsAttributesPerCurve(InPrimitiveId, ClumpIdOffset, HAIR_CURVE_ATTRIBUTE_STRIDE_CLUMPID, PREFIX(CurveAttributeBuffer), PREFIX(PointToCurveBuffer)));
+		return UnpackHairClumpID(InternalGetHairStrandsAttributesPerCurve(InPrimitiveId, ClumpIdOffset, HAIR_CURVE_ATTRIBUTE_STRIDE_CLUMPID, PREFIX(CurveAttributeBuffer), PREFIX(PointToCurveBuffer), PREFIX(CurveAttributeIndexToChunkDivAsShift), PREFIX(CurveAttributeChunkElementCount), PREFIX(CurveAttributeChunkStrideInBytes)));
 	}
 	else if (IsAttributeValid(ClumpId3Offset))
 	{
-		return UnpackHairClumpID(InternalGetHairStrandsAttributesPerCurve64(InPrimitiveId, ClumpId3Offset, HAIR_CURVE_ATTRIBUTE_STRIDE_CLUMPID3, PREFIX(CurveAttributeBuffer), PREFIX(PointToCurveBuffer)));
+		return UnpackHairClumpID(InternalGetHairStrandsAttributesPerCurve64(InPrimitiveId, ClumpId3Offset, HAIR_CURVE_ATTRIBUTE_STRIDE_CLUMPID3, PREFIX(CurveAttributeBuffer), PREFIX(PointToCurveBuffer), PREFIX(CurveAttributeIndexToChunkDivAsShift), PREFIX(CurveAttributeChunkElementCount), PREFIX(CurveAttributeChunkStrideInBytes)));
 	}
 	else
 	{
@@ -72,8 +72,8 @@ float3 GetHairStrandsColor(uint InPrimitiveId, float InSegmentU)
 	const uint ColorOffset = HAIR_POINT_ATTRIBUTE_OFFSET_COLOR(PREFIX(PointAttributeOffsets));
 	if (IsAttributeValid(ColorOffset))
 	{
-		const float3 Color0 = UnpackHairColor(InternalGetHairStrandsAttributesPerVertex(InPrimitiveId,   ColorOffset, HAIR_POINT_ATTRIBUTE_STRIDE_COLOR, PREFIX(PointAttributeBuffer)));
-		const float3 Color1 = UnpackHairColor(InternalGetHairStrandsAttributesPerVertex(InPrimitiveId+1, ColorOffset, HAIR_POINT_ATTRIBUTE_STRIDE_COLOR, PREFIX(PointAttributeBuffer)));
+		const float3 Color0 = UnpackHairColor(InternalGetHairStrandsAttributesPerVertex(InPrimitiveId,   ColorOffset, HAIR_POINT_ATTRIBUTE_STRIDE_COLOR, PREFIX(PointAttributeBuffer), PREFIX(PointAttributeIndexToChunkDivAsShift), PREFIX(PointAttributeChunkElementCount), PREFIX(PointAttributeChunkStrideInBytes)));
+		const float3 Color1 = UnpackHairColor(InternalGetHairStrandsAttributesPerVertex(InPrimitiveId+1, ColorOffset, HAIR_POINT_ATTRIBUTE_STRIDE_COLOR, PREFIX(PointAttributeBuffer), PREFIX(PointAttributeIndexToChunkDivAsShift), PREFIX(PointAttributeChunkElementCount), PREFIX(PointAttributeChunkStrideInBytes)));
 		return lerp(Color0, Color1, InSegmentU);
 	}
 	else
@@ -86,8 +86,8 @@ float GetHairStrandsRoughness(uint InPrimitiveId, float InSegmentU)
 	const uint RoughnessOffset = HAIR_POINT_ATTRIBUTE_OFFSET_ROUGHNESS(PREFIX(PointAttributeOffsets));
 	if (IsAttributeValid(RoughnessOffset))
 	{
-		const float Roughness0 = UnpackHairRoughness(InternalGetHairStrandsAttributesPerVertex(InPrimitiveId,   RoughnessOffset, HAIR_POINT_ATTRIBUTE_STRIDE_ROUGHNESS, PREFIX(PointAttributeBuffer)));
-		const float Roughness1 = UnpackHairRoughness(InternalGetHairStrandsAttributesPerVertex(InPrimitiveId+1, RoughnessOffset, HAIR_POINT_ATTRIBUTE_STRIDE_ROUGHNESS, PREFIX(PointAttributeBuffer)));
+		const float Roughness0 = UnpackHairRoughness(InternalGetHairStrandsAttributesPerVertex(InPrimitiveId,   RoughnessOffset, HAIR_POINT_ATTRIBUTE_STRIDE_ROUGHNESS, PREFIX(PointAttributeBuffer), PREFIX(PointAttributeIndexToChunkDivAsShift), PREFIX(PointAttributeChunkElementCount), PREFIX(PointAttributeChunkStrideInBytes)));
+		const float Roughness1 = UnpackHairRoughness(InternalGetHairStrandsAttributesPerVertex(InPrimitiveId+1, RoughnessOffset, HAIR_POINT_ATTRIBUTE_STRIDE_ROUGHNESS, PREFIX(PointAttributeBuffer), PREFIX(PointAttributeIndexToChunkDivAsShift), PREFIX(PointAttributeChunkElementCount), PREFIX(PointAttributeChunkStrideInBytes)));
 		return lerp(Roughness0, Roughness1, InSegmentU);
 	}
 	else
@@ -100,8 +100,8 @@ float GetHairStrandsAO(uint InPrimitiveId, float InSegmentU)
 	const uint AoOffset = HAIR_POINT_ATTRIBUTE_OFFSET_AO(PREFIX(PointAttributeOffsets));
 	if (IsAttributeValid(AoOffset))
 	{
-		const float AO0 = UnpackHairAO(InternalGetHairStrandsAttributesPerVertex(InPrimitiveId, AoOffset, HAIR_POINT_ATTRIBUTE_STRIDE_AO, PREFIX(PointAttributeBuffer)));
-		const float AO1 = UnpackHairAO(InternalGetHairStrandsAttributesPerVertex(InPrimitiveId+1, AoOffset, HAIR_POINT_ATTRIBUTE_STRIDE_AO, PREFIX(PointAttributeBuffer)));
+		const float AO0 = UnpackHairAO(InternalGetHairStrandsAttributesPerVertex(InPrimitiveId, AoOffset, HAIR_POINT_ATTRIBUTE_STRIDE_AO, PREFIX(PointAttributeBuffer), PREFIX(PointAttributeIndexToChunkDivAsShift), PREFIX(PointAttributeChunkElementCount), PREFIX(PointAttributeChunkStrideInBytes)));
+		const float AO1 = UnpackHairAO(InternalGetHairStrandsAttributesPerVertex(InPrimitiveId+1, AoOffset, HAIR_POINT_ATTRIBUTE_STRIDE_AO, PREFIX(PointAttributeBuffer), PREFIX(PointAttributeIndexToChunkDivAsShift), PREFIX(PointAttributeChunkElementCount), PREFIX(PointAttributeChunkStrideInBytes)));
 		return lerp(AO0, AO1, InSegmentU);
 	}
 	else
@@ -126,4 +126,4 @@ bool HasHairStrandsClumpId() 	{ return IsAttributeValid(HAIR_CURVE_ATTRIBUTE_OFF
 bool HasHairStrandsClumpId3() 	{ return IsAttributeValid(HAIR_CURVE_ATTRIBUTE_OFFSET_CLUMPID3(PREFIX(CurveAttributeOffsets))); }
 bool HasHairStrandsColor() 		{ return IsAttributeValid(HAIR_POINT_ATTRIBUTE_OFFSET_COLOR(PREFIX(PointAttributeOffsets))); }
 bool HasHairStrandsRoughness() 	{ return IsAttributeValid(HAIR_POINT_ATTRIBUTE_OFFSET_ROUGHNESS(PREFIX(PointAttributeOffsets))); }
-bool HasHairStrandsAO() 		{ return IsAttributeValid(HAIR_POINT_ATTRIBUTE_OFFSET_AO(PREFIX(PointAttributeOffsets))); }
\ No newline at end of file
+bool HasHairStrandsAO() 		{ return IsAttributeValid(HAIR_POINT_ATTRIBUTE_OFFSET_AO(PREFIX(PointAttributeOffsets))); }
diff --git a/Engine/Shaders/Private/HairStrands/HairStrandsInterpolation.usf b/Engine/Shaders/Private/HairStrands/HairStrandsInterpolation.usf
--- a/Engine/Shaders/Private/HairStrands/HairStrandsInterpolation.usf
+++ b/Engine/Shaders/Private/HairStrands/HairStrandsInterpolation.usf
@@ -513,10 +513,17 @@ void MainCS(uint2 DispatchThreadId : SV_DispatchThreadID)
 uint				CurveCount;
 
 Buffer<uint>		RenCurveBuffer;
-uint4				RenCurveAttributeOffsets[HAIR_CURVE_ATTRIBUTE_OFFSET_COUNT];
 Buffer<int>			RenVertexToClusterIdBuffer;
 RWByteAddressBuffer	OutRenAttributeBuffer;
 
+uint 				CurveAttributeIndexToChunkDivAsShift;
+uint 				CurveAttributeChunkElementCount;
+uint 				CurveAttributeChunkStrideInBytes;
+uint 				PointAttributeIndexToChunkDivAsShift;
+uint 				PointAttributeChunkElementCount;
+uint 				PointAttributeChunkStrideInBytes;
+uint4				CurveAttributeOffsets[HAIR_CURVE_ATTRIBUTE_OFFSET_COUNT];
+
 groupshared uint SeedValues[1024];
 
 uint ToSeedHash(uint In)
@@ -570,8 +577,14 @@ void MainCS(uint2 DispatchThreadId : SV_DispatchThreadID, uint GroupThread1D : S
 		Out |= (SeedValues[GroupThread1D+2] & 0xFF) << 16;
 		Out |= (SeedValues[GroupThread1D+3] & 0xFF) << 24;
 
-		const uint OffsetInBytes = HAIR_CURVE_ATTRIBUTE_OFFSET_SEED(RenCurveAttributeOffsets);
-		OutRenAttributeBuffer.Store(OffsetInBytes + CurveIndex, Out);
+		uint WriteIndex 				= CurveIndex;
+		uint WriteAttributOffsetInBytes = HAIR_CURVE_ATTRIBUTE_OFFSET_SEED(CurveAttributeOffsets);
+
+		const uint ChunkIndex   	= CurveIndex >> CurveAttributeIndexToChunkDivAsShift;
+		WriteIndex                 -= ChunkIndex * CurveAttributeIndexToChunkDivAsShift;
+		WriteAttributOffsetInBytes += ChunkIndex * CurveAttributeChunkStrideInBytes;
+
+		OutRenAttributeBuffer.Store(WriteAttributOffsetInBytes + WriteIndex, Out);
 	}
 }
 #endif // SHADER_PATCHATTRIBUTE
diff --git a/Engine/Shaders/Private/HairStrands/HairStrandsVertexFactoryCommon.ush b/Engine/Shaders/Private/HairStrands/HairStrandsVertexFactoryCommon.ush
--- a/Engine/Shaders/Private/HairStrands/HairStrandsVertexFactoryCommon.ush
+++ b/Engine/Shaders/Private/HairStrands/HairStrandsVertexFactoryCommon.ush
@@ -33,6 +33,12 @@
 	float4x4			Prefix##_LocalToTranslatedWorldPrimitiveTransform; \
 	uint4				Prefix##_CurveAttributeOffsets[HAIR_CURVE_ATTRIBUTE_OFFSET_COUNT]; \
 	uint4				Prefix##_PointAttributeOffsets[HAIR_POINT_ATTRIBUTE_OFFSET_COUNT]; \
+	uint 				Prefix##_CurveAttributeIndexToChunkDivAsShift; \
+	uint				Prefix##_CurveAttributeChunkElementCount; \
+	uint				Prefix##_CurveAttributeChunkStrideInBytes; \
+	uint				Prefix##_PointAttributeIndexToChunkDivAsShift; \
+	uint				Prefix##_PointAttributeChunkElementCount; \
+	uint				Prefix##_PointAttributeChunkStrideInBytes; \
 	Buffer<uint4>		Prefix##_PositionBuffer; \
 	Buffer<float4>		Prefix##_PositionOffsetBuffer; \
 	Buffer<uint>		Prefix##_CurveBuffer; \
diff --git a/Engine/Source/Runtime/Renderer/Public/HairStrandsInterface.h b/Engine/Source/Runtime/Renderer/Public/HairStrandsInterface.h
--- a/Engine/Source/Runtime/Renderer/Public/HairStrandsInterface.h
+++ b/Engine/Source/Runtime/Renderer/Public/HairStrandsInterface.h
@@ -23,6 +23,18 @@ class FSceneInterface;
 ////////////////////////////////////////////////////////////////////////////////////////////////////
 // Shader parameters
 
+// Instance attribute parameters
+BEGIN_SHADER_PARAMETER_STRUCT(FHairStrandsInstanceAttributeParameters, RENDERER_API)
+	SHADER_PARAMETER(uint32, CurveAttributeIndexToChunkDivAsShift)
+	SHADER_PARAMETER(uint32, CurveAttributeChunkElementCount)
+	SHADER_PARAMETER(uint32, CurveAttributeChunkStrideInBytes)
+	SHADER_PARAMETER(uint32, PointAttributeIndexToChunkDivAsShift)
+	SHADER_PARAMETER(uint32, PointAttributeChunkElementCount)
+	SHADER_PARAMETER(uint32, PointAttributeChunkStrideInBytes)
+	SHADER_PARAMETER_ARRAY(FUintVector4, CurveAttributeOffsets, [HAIR_CURVE_ATTRIBUTE_OFFSET_COUNT])
+	SHADER_PARAMETER_ARRAY(FUintVector4, PointAttributeOffsets, [HAIR_POINT_ATTRIBUTE_OFFSET_COUNT])
+END_SHADER_PARAMETER_STRUCT()
+
 // Instance common parameters
 BEGIN_SHADER_PARAMETER_STRUCT(FHairStrandsInstanceCommonParameters, RENDERER_API)
 	SHADER_PARAMETER(float,  Density)
@@ -46,8 +58,7 @@ BEGIN_SHADER_PARAMETER_STRUCT(FHairStrandsInstanceCommonParameters, RENDERER_API
 	SHADER_PARAMETER(FVector3f, PrevPositionOffset)
 	SHADER_PARAMETER(FMatrix44f, LocalToWorldPrimitiveTransform)
 	SHADER_PARAMETER(FMatrix44f, LocalToTranslatedWorldPrimitiveTransform)
-	SHADER_PARAMETER_ARRAY(FUintVector4, CurveAttributeOffsets, [HAIR_CURVE_ATTRIBUTE_OFFSET_COUNT])
-	SHADER_PARAMETER_ARRAY(FUintVector4, PointAttributeOffsets, [HAIR_POINT_ATTRIBUTE_OFFSET_COUNT])
+	SHADER_PARAMETER_STRUCT_INCLUDE(FHairStrandsInstanceAttributeParameters, Attributes)
 END_SHADER_PARAMETER_STRUCT()
 
 // Instance resources (RDG)
