Fix typo in groom deformer shader.

diff --git a/Engine/Plugins/Runtime/HairStrands/Shaders/Private/DeformerDataInterfaceGroomWrite.ush b/Engine/Plugins/Runtime/HairStrands/Shaders/Private/DeformerDataInterfaceGroomWrite.ush
--- a/Engine/Plugins/Runtime/HairStrands/Shaders/Private/DeformerDataInterfaceGroomWrite.ush
+++ b/Engine/Plugins/Runtime/HairStrands/Shaders/Private/DeformerDataInterfaceGroomWrite.ush
@@ -94,9 +94,6 @@ void InternalSetHairStrandsAttributes32(uint Index, uint AttributeOffsetInBytes,
 	InAttributeBuffer.Store(Address_4BytesAligned, In);
 }
 
-, {DataInterfaceName}_CurveAttributeIndexToChunkDivAsShift, {DataInterfaceName}_CurveAttributeChunkElementCount), {DataInterfaceName}_CurveAttributeChunkStrideInBytes
-, {DataInterfaceName}_PointAttributeIndexToChunkDivAsShift, {DataInterfaceName}_PointAttributeChunkElementCount), {DataInterfaceName}_PointAttributeChunkStrideInBytes
-
 void WriteRootUV_{DataInterfaceName}(uint CurveIndex, float2 RootUV)
 {
 	const uint WriteIndex = {DataInterfaceName}_OutputStreamStart + CurveIndex;
