Add support for groom streaming based on requested curve count.

Default behavior request entire resource by default, and fine grain streaming is enabled with CLOD settings.

diff --git a/Engine/Plugins/Runtime/HairStrands/Source/HairStrandsCore/Private/GroomAsset.cpp b/Engine/Plugins/Runtime/HairStrands/Source/HairStrandsCore/Private/GroomAsset.cpp
--- a/Engine/Plugins/Runtime/HairStrands/Source/HairStrandsCore/Private/GroomAsset.cpp
+++ b/Engine/Plugins/Runtime/HairStrands/Source/HairStrandsCore/Private/GroomAsset.cpp
@@ -1770,12 +1770,12 @@ static void InternalSerializeStrand(FArchive& Ar, UObject* Owner, FHairGroupPlat
 {
 	Ar.UsingCustomVersion(FAnimObjectVersion::GUID);
 
-	// When cooking data, force loading of bulk data prior to saving them
+	// When cooking data, force loading of *all* bulk data prior to saving them
 	if (Ar.IsCooking() && Ar.IsSaving())
 	{
-		{ FHairResourceRequest R; R.Request(StrandData.BulkData, true /*bWait*/, true /*bFillBulkdata*/, Owner->GetFName()); }
-		{ FHairResourceRequest R; R.Request(StrandData.InterpolationBulkData, true /*bWait*/, true /*bFillBulkdata*/, Owner->GetFName()); }
-		{ FHairResourceRequest R; R.Request(StrandData.ClusterCullingBulkData, true /*bWait*/, true /*bFillBulkdata*/, Owner->GetFName()); }
+		{ FHairStreamingRequest R; R.Request(HAIR_MAX_NUM_CURVE_PER_GROUP, StrandData.BulkData,               true /*bWait*/, true /*bFillBulkdata*/, Owner->GetFName()); }
+		{ FHairStreamingRequest R; R.Request(HAIR_MAX_NUM_CURVE_PER_GROUP, StrandData.InterpolationBulkData,  true /*bWait*/, true /*bFillBulkdata*/, Owner->GetFName()); }
+		{ FHairStreamingRequest R; R.Request(HAIR_MAX_NUM_CURVE_PER_GROUP, StrandData.ClusterCullingBulkData, true /*bWait*/, true /*bFillBulkdata*/, Owner->GetFName()); }
 	}
 
 	if (!Ar.IsCooking() || !StrandData.bIsCookedOut)
diff --git a/Engine/Plugins/Runtime/HairStrands/Source/HairStrandsCore/Private/GroomBuilder.cpp b/Engine/Plugins/Runtime/HairStrands/Source/HairStrandsCore/Private/GroomBuilder.cpp
--- a/Engine/Plugins/Runtime/HairStrands/Source/HairStrandsCore/Private/GroomBuilder.cpp
+++ b/Engine/Plugins/Runtime/HairStrands/Source/HairStrandsCore/Private/GroomBuilder.cpp
@@ -40,7 +40,7 @@ static FAutoConsoleVariableRef CVarHairGroupIndexBuilder_MaxVoxelResolution(TEXT
 
 FString FGroomBuilder::GetVersion()
 {
-	return TEXT("v8r19");
+	return TEXT("v8r22");
 }
 
 namespace FHairStrandsDecimation
@@ -548,6 +548,13 @@ namespace HairStrandsBuilder
 			CopyToBulkData<FHairStrandsPointToCurveFormat32>(OutBulkData.Data.PointToCurve, OutPointToCurve32);
 		}
 
+		// Build curve to point count mapping for runtime CLOD
+		OutBulkData.Header.CurveToPointCount.Reserve(OutPackedCurves.Num());
+		for (const FHairStrandsCurveFormat::Type& Curve : OutPackedCurves)
+		{
+			OutBulkData.Header.CurveToPointCount.Add(Curve.PointOffset + Curve.PointCount);
+		}
+
 		// Stride datas
 		OutBulkData.Header.Strides.PositionStride = FHairStrandsPositionFormat::SizeInByte;
 		OutBulkData.Header.Strides.CurveStride = FHairStrandsCurveFormat::SizeInByte;
diff --git a/Engine/Plugins/Runtime/HairStrands/Source/HairStrandsCore/Private/GroomComponent.cpp b/Engine/Plugins/Runtime/HairStrands/Source/HairStrandsCore/Private/GroomComponent.cpp
--- a/Engine/Plugins/Runtime/HairStrands/Source/HairStrandsCore/Private/GroomComponent.cpp
+++ b/Engine/Plugins/Runtime/HairStrands/Source/HairStrandsCore/Private/GroomComponent.cpp
@@ -2698,6 +2698,10 @@ void UGroomComponent::InitResources(bool bIsBindingReloading)
 				if (MeshDeformer)
 				{
 					bNeedDynamicResources = true;
+
+					// When an instance has a groom deformer, we disable streaming in order to load all its data, as they 
+					// can be accessed/modified by a groom deformer
+					HairGroupInstance->bSupportStreaming = false;
 				}
 			}
 
diff --git a/Engine/Plugins/Runtime/HairStrands/Source/HairStrandsCore/Private/GroomManager.cpp b/Engine/Plugins/Runtime/HairStrands/Source/HairStrandsCore/Private/GroomManager.cpp
--- a/Engine/Plugins/Runtime/HairStrands/Source/HairStrandsCore/Private/GroomManager.cpp
+++ b/Engine/Plugins/Runtime/HairStrands/Source/HairStrandsCore/Private/GroomManager.cpp
@@ -935,6 +935,12 @@ static void RunHairLODSelection(
 				GeometryType = EHairGeometryType::NoneGeometry;
 			}
 
+			uint32 RequestedCurveCount  = Instance->HairGroupPublicData->RestCurveCount;
+			if (GeometryType == EHairGeometryType::Strands && IsHairStrandsContinousLODEnabled() && Instance->bSupportStreaming)
+			{
+				RequestedCurveCount = ComputeActiveCurveCount(Instance->HairGroupPublicData->ContinuousLODScreenSize, Instance->HairGroupPublicData->RestCurveCount);
+			}
+
 			const bool bSimulationEnable			= Instance->HairGroupPublicData->IsSimulationEnable(IntLODIndex);
 			const bool bDeformationEnable			= Instance->HairGroupPublicData->bIsDeformationEnable;
 			const bool bGlobalInterpolationEnable	= Instance->HairGroupPublicData->IsGlobalInterpolationEnable(IntLODIndex);
@@ -942,7 +948,9 @@ static void RunHairLODSelection(
 			const bool bLODNeedsGuides				= bSimulationEnable || bDeformationEnable || bGlobalInterpolationEnable || bSimulationCacheEnable;
 
 			const EHairResourceLoadingType LoadingType = GetHairResourceLoadingType(GeometryType, IntLODIndex);
-			EHairResourceStatus ResourceStatus = EHairResourceStatus::None;
+			EHairResourceStatus ResourceStatus;
+			ResourceStatus.Status 				= EHairResourceStatus::EStatus::None;
+			ResourceStatus.AvailableCurveCount 	= Instance->HairGroupPublicData->RestCurveCount;
 
 			// Lazy allocation of resources
 			// Note: Allocation will only be done if the resources is not initialized yet. Guides deformed position are also initialized from the Rest position at creation time.
@@ -999,9 +1007,9 @@ static void RunHairLODSelection(
 				Instance->HairGroupPublicData->Allocate(GraphBuilder);
 
 				if (Instance->Strands.RestRootResource)			{ Instance->Strands.RestRootResource->Allocate(GraphBuilder, LoadingType, ResourceStatus); Instance->Strands.RestRootResource->AllocateLOD(GraphBuilder, MeshLODIndex, LoadingType, ResourceStatus); }
-				if (Instance->Strands.RestResource)				{ Instance->Strands.RestResource->Allocate(GraphBuilder, EHairResourceLoadingType::Async, ResourceStatus); }
-				if (Instance->Strands.ClusterCullingResource)	{ Instance->Strands.ClusterCullingResource->Allocate(GraphBuilder, EHairResourceLoadingType::Async, ResourceStatus); }
-				if (Instance->Strands.InterpolationResource)	{ Instance->Strands.InterpolationResource->Allocate(GraphBuilder, EHairResourceLoadingType::Async, ResourceStatus); }
+				if (Instance->Strands.RestResource)				{ Instance->Strands.RestResource->Allocate(GraphBuilder, EHairResourceLoadingType::Async, ResourceStatus, RequestedCurveCount); }
+				if (Instance->Strands.ClusterCullingResource)	{ Instance->Strands.ClusterCullingResource->Allocate(GraphBuilder, EHairResourceLoadingType::Async, ResourceStatus, RequestedCurveCount); }
+				if (Instance->Strands.InterpolationResource)	{ Instance->Strands.InterpolationResource->Allocate(GraphBuilder, EHairResourceLoadingType::Async, ResourceStatus, RequestedCurveCount); }
 
 				if (Instance->Strands.DeformedRootResource)		{ Instance->Strands.DeformedRootResource->Allocate(GraphBuilder, LoadingType, ResourceStatus); Instance->Strands.DeformedRootResource->AllocateLOD(GraphBuilder, MeshLODIndex, LoadingType, ResourceStatus); }
 				if (Instance->Strands.DeformedResource)			{ Instance->Strands.DeformedResource->Allocate(GraphBuilder, LoadingType, ResourceStatus); }
@@ -1017,7 +1025,7 @@ static void RunHairLODSelection(
 			}
 
 			// Only switch LOD if the data are ready to be used
-			const bool bIsLODDataReady = !!(ResourceStatus & EHairResourceStatus::Loading);
+			const bool bIsLODDataReady = !ResourceStatus.HasStatus(EHairResourceStatus::EStatus::Loading) || (IsHairStrandsContinousLODEnabled() && ResourceStatus.AvailableCurveCount > 0);
 			if (bIsLODDataReady)
 			{
 				EHairBindingType BindingType = Instance->HairGroupPublicData->GetBindingType(IntLODIndex);
@@ -1058,12 +1066,10 @@ static void RunHairLODSelection(
 				check(bIsLODDataReady);
 				check(Instance->Strands.RestResource);
 
-				const uint32 ActiveCurveCount = ComputeActiveCurveCount(Instance->HairGroupPublicData->ContinuousLODScreenSize, Instance->HairGroupPublicData->RestCurveCount);
-				const FPackedHairCurve Curve = Instance->Strands.RestResource->CurveData[ActiveCurveCount - 1];
-				const uint32 ActiveVertexCount = ActiveCurveCount < Instance->HairGroupPublicData->RestCurveCount ? (Curve.PointOffset + Curve.PointCount) : Instance->HairGroupPublicData->RestPointCount;
-
-				Instance->HairGroupPublicData->ContinuousLODCurveCount = ActiveCurveCount;
-				Instance->HairGroupPublicData->ContinuousLODPointCount = ActiveVertexCount;
+				const uint32 EffectiveCurveCount = FMath::Min(ResourceStatus.AvailableCurveCount, RequestedCurveCount);
+				Instance->HairGroupPublicData->ContinuousLODCurveCount = EffectiveCurveCount;
+				Instance->HairGroupPublicData->ContinuousLODPointCount = EffectiveCurveCount > 0 ? Instance->Strands.Data->Header.CurveToPointCount[EffectiveCurveCount - 1] : 0;
+				check(Instance->HairGroupPublicData->ContinuousLODPointCount <= Instance->HairGroupPublicData->RestPointCount);
 			}
 			return true;
 		};
diff --git a/Engine/Plugins/Runtime/HairStrands/Source/HairStrandsCore/Private/GroomResources.cpp b/Engine/Plugins/Runtime/HairStrands/Source/HairStrandsCore/Private/GroomResources.cpp
--- a/Engine/Plugins/Runtime/HairStrands/Source/HairStrandsCore/Private/GroomResources.cpp
+++ b/Engine/Plugins/Runtime/HairStrands/Source/HairStrandsCore/Private/GroomResources.cpp
@@ -103,117 +103,133 @@ const TCHAR* ToHairResourceDebugName(const TCHAR* In, FHairResourceName& InDebug
 /////////////////////////////////////////////////////////////////////////////////////////
 // IO/DDC request
 
+bool FHairStreamingRequest::IsNone() const
+{ 
 #if !WITH_EDITORONLY_DATA
-bool FHairResourceRequest::IsNone() const 		{ return IORequest.IsNone(); }
-bool FHairResourceRequest::IsCompleted() const	{ return IORequest.IsCompleted(); }
-void FHairResourceRequest::Request(FHairStrandsBulkCommon& In, bool bWait, bool bFillBulkdata, const FName& InOwnerName)
+	return IORequest.IsNone();
+#else
+	return Chunks.Num() == 0;
+#endif
+}
+bool FHairStreamingRequest::IsCompleted()
 { 
-	if (In.GetResourceCount() != 0)
+#if !WITH_EDITORONLY_DATA
+	if (IORequest.IsCompleted())
 	{
-		In.Read_IO(IORequest);
-		if (bWait)
+		for (FHairStreamingRequest::FChunk& Chunk : Chunks)
 		{
-			IORequest.Wait();
+			Chunk.Status = FHairStreamingRequest::FChunk::Completed;
 		}
 	}
-}
+	return Chunks.Num() == 0 || IORequest.IsCompleted(); 
 #else
-bool FHairResourceRequest::IsNone() const 		{ return Chunks.Num() == 0; }
-bool FHairResourceRequest::IsCompleted() const
-{ 
-	for (const FHairChunkRequest& Chunk : Chunks)
+	for (const FHairStreamingRequest::FChunk& Chunk : Chunks)
 	{
-		if (Chunk.Status != FHairChunkRequest::Completed)
+		if (Chunk.Status != FHairStreamingRequest::FChunk::Completed)
 		{
 			return false;
 		}
 	}
 	return true; 
+#endif
 }
-
 // Request fullfil 2 use cases:
 // * Load DDC data and upload them to GPU
 // * Load DDC data and store them into bulkdata for serialization
-void FHairResourceRequest::Request(FHairStrandsBulkCommon& In, bool bWait, bool bFillBulkdata, const FName& InOwnerName)
+void FHairStreamingRequest::Request(uint32 InRequestedCurveCount, FHairStrandsBulkCommon& In, bool bWait, bool bFillBulkdata, const FName& InOwnerName)
 {
-	if (In.GetResourceCount() == 0)
+	if (In.GetResourceCount() == 0 || InRequestedCurveCount == 0)
 	{
+		CurveCount = 0;
 		return;
 	}
+	CurveCount = InRequestedCurveCount;
 
-	using namespace UE::DerivedData;
-	TArray<FCacheGetChunkRequest> Requests;
-	In.Read_DDC(Requests);
-
-	// Add trackin data
-	Chunks.SetNum(0, false);
-	Chunks.Reserve(Requests.Num());
-	for (FCacheGetChunkRequest& R : Requests)
-	{
-		FHairChunkRequest& Chunk = Chunks.AddDefaulted_GetRef();
-		Chunk.Offset = 0;
-		Chunk.Size = 0;
-		Chunk.Status = FHairChunkRequest::EStatus::Pending;
-		Chunk.Container = (FHairBulkContainer*)R.UserData;
-		R.UserData = (uint64)&Chunk;
-		R.Name = InOwnerName.ToString();
-	}
-
-	DDCRequestOwner = MakeUnique<FRequestOwner>(bWait ? UE::DerivedData::EPriority::Blocking : UE::DerivedData::EPriority::Normal); // <= Move this onto the resource (Buffer, in order to ensure no race condition)
-
-	//FRequestBarrier Barrier(*DDCRequestOwner);	// This is a critical section on the owner. It does not constrain ordering
-	GetCache().GetChunks(Requests, *DDCRequestOwner,
-	[this, &In](FCacheGetChunkResponse&& Response)
-	{
-		if (Response.Status == UE::DerivedData::EStatus::Ok)
-		{
-			FHairChunkRequest& Chunk = *(FHairChunkRequest*)Response.UserData;
-			Chunk.Data 	 = MoveTemp(Response.RawData);
-			Chunk.Offset = Response.RawOffset;
-			Chunk.Size 	 = Response.RawSize;
-			Chunk.Status = FHairChunkRequest::Completed;
-		}
-		else
+	#if !WITH_EDITORONLY_DATA
+	{ 
+		In.Read_IO(this, IORequest);
+		if (bWait)
 		{
-			FHairChunkRequest& Chunk = *(FHairChunkRequest*)Response.UserData;
-			Chunk.Status = FHairChunkRequest::Failed;
-			UE_LOG(LogHairStrands, Error, TEXT("[Groom] DDC request failed for '%s' (Key:%s) "), *Response.Name, *In.DerivedDataKey);
+			IORequest.Wait();
 		}
-	});
-
-	// TODO remove for resources supporting streaming request
-	bWait = true;
-
-	if (bWait)
+	}
+	#else
 	{
-		DDCRequestOwner->Wait();
-		if (!IsCompleted())
+		using namespace UE::DerivedData;
+		TArray<FCacheGetChunkRequest> Requests;
+		In.Read_DDC(this, Requests);
+	
+		check(DDCRequestOwner == nullptr);
+		DDCRequestOwner = MakeUnique<FRequestOwner>(bWait ? UE::DerivedData::EPriority::Blocking : UE::DerivedData::EPriority::Normal); // <= Move this onto the resource (Buffer, in order to ensure no race condition)
+	
+		//FRequestBarrier Barrier(*DDCRequestOwner);	// This is a critical section on the owner. It does not constrain ordering
+		GetCache().GetChunks(Requests, *DDCRequestOwner,
+		[this, &In, bFillBulkdata, InOwnerName](FCacheGetChunkResponse && Response)
 		{
-			return;
-		}
-	}
+			if (Response.Status == UE::DerivedData::EStatus::Ok)
+			{
+				FHairStreamingRequest::FChunk& Chunk = *(FHairStreamingRequest::FChunk*)Response.UserData;
+				Chunk.Data_DDC 	= MoveTemp(Response.RawData);
+				Chunk.Offset 	= Response.RawOffset;
+				Chunk.Size 		= Response.RawSize;
+				Chunk.TotalSize = Response.RawOffset + Response.RawSize;
+				Chunk.Status 	= FHairStreamingRequest::FChunk::Completed;
+		
+				// Upload the total amount of loaded data
+				// HAIR_TODO: Disable for now, reupload all the data until issue with RDG/DX12 source/target buffer is solved
+				// Chunk.Container->LoadedSize = Chunk.TotalSize;
 
-	// Fill in bytebulkdata
-	if (bFillBulkdata)
-	{		
-		for (FHairChunkRequest& Chunk : Chunks)
-		{
-			check(Chunk.Status == FHairChunkRequest::Completed);
-			if (Chunk.Container)
+				// Optional fill in of bytebulkdata container
+				if (bFillBulkdata)
+				{
+					check(Chunk.Container);
+					FByteBulkData& BulkData = Chunk.Container->Data;
+	
+					// The buffer is then stored into bulk data
+					BulkData.Lock(LOCK_READ_WRITE);
+					void* DstData = BulkData.Realloc(Chunk.Size);
+					FMemory::Memcpy(DstData, Chunk.GetData(), Chunk.Size);
+					BulkData.Unlock();
+
+					//Chunk.Release();
+				}
+			}
+			else
 			{
-				FByteBulkData& BulkData = Chunk.Container->Data;
-
-				// The buffer is then stored into bulk data
-				BulkData.Lock(LOCK_READ_WRITE);
-				void* DstData = BulkData.Realloc(Chunk.Size);
-				FMemory::Memcpy(DstData, Chunk.GetData(), Chunk.Size);
-				BulkData.Unlock();
-				Chunk.Release();
+				FHairStreamingRequest::FChunk& Chunk = *(FHairStreamingRequest::FChunk*)Response.UserData;
+				Chunk.Status = FHairStreamingRequest::FChunk::Failed;
+				UE_LOG(LogHairStrands, Error, TEXT("[Groom] DDC request failed for '%s' (Key:%s) "), *InOwnerName.ToString(), *In.DerivedDataKey);
 			}
+		});
+
+		// Optional wait on DDC response
+		if (bWait || bFillBulkdata)
+		{
+			DDCRequestOwner->Wait();
+			check(IsCompleted());
 		}
 	}
+	#endif
+}
+
+const uint8* FHairStreamingRequest::FChunk::GetData() const 
+{ 
+	check(Status == Completed); 
+#if !WITH_EDITORONLY_DATA
+	check(Container);
+	return (const uint8*)Data_IO.GetData();
+#else
+	return (const uint8*)Data_DDC.GetData(); 
+#endif
 }
+void FHairStreamingRequest::FChunk::Release() 
+{ 
+#if !WITH_EDITORONLY_DATA
+	TIoStatusOr<uint8*> Out = Data_IO.Release();
+#else
+	Data_DDC.Reset(); 
 #endif
+}
 
 /////////////////////////////////////////////////////////////////////////////////////////
 // FRWBuffer utils 
@@ -306,6 +322,9 @@ inline void InternalSetBulkDataFlags(FByteBulkData& In)
 #endif
 }
 
+///////////////////////////////////////////////////////////////////////////////////////////////////
+// Regular loading from BulkData
+
 void InternalCreateBufferRDG_FromBulkData(FRDGBuilder& GraphBuilder, FByteBulkData& InBulkData, FRDGExternalBuffer& Out, const EPixelFormat OutFormat, const FRDGBufferDesc Desc, const TCHAR* DebugName, const FName& OwnerName)
 {
 	InternalSetBulkDataFlags(InBulkData);
@@ -359,6 +378,113 @@ void InternalCreateByteAddressBufferRDG_FromBulkData(FRDGBuilder& GraphBuilder,
 }
 
 
+///////////////////////////////////////////////////////////////////////////////////////////////////
+// HairBulkData loading
+
+static FRDGBufferRef InternalCreateBufferRDG_FromHairBulkData(FRDGBuilder& GraphBuilder, FHairBulkContainer& InData, const FRDGBufferRef& In, const FRDGBufferDesc& BufferDesc, const FRDGBufferDesc& UploadDesc, const TCHAR* DebugName, const FName& OwnerName)
+{
+	check(InData.ChunkRequest);
+	FHairStreamingRequest::FChunk& InChunk = *InData.ChunkRequest;
+
+	const bool bCopyData = In != nullptr;
+	// For now always upload the entire data set as partial upload + GPU copy result in D3D12 error due to buffer aliasing HAIR_TODO
+	const bool bNeedCreateBuffer = In == nullptr || BufferDesc.GetSize() == UploadDesc.GetSize() || bCopyData;
+	if (bNeedCreateBuffer)
+	{	
+		check(BufferDesc.GetSize() >= UploadDesc.GetSize());
+
+		FRDGBufferRef Out = GraphBuilder.CreateBuffer(BufferDesc, DebugName, ERDGBufferFlags::MultiFrame);
+		GraphBuilder.QueueBufferUpload(Out, InChunk.GetData(), InChunk.Size, ERDGInitialDataFlags::None);
+		InChunk.Release();
+		return Out;
+	}
+#if 0 // HAIR_TODO
+	else if (bCopyData && InChunk.Size > 0)
+	{
+		// 1. Create new buffer
+		FRDGBufferRef Out = GraphBuilder.CreateBuffer(BufferDesc, DebugName, ERDGBufferFlags::MultiFrame);
+
+		// 2. Copy existing data from the old buffer to the new buffer
+		AddCopyBufferPass(GraphBuilder, Out, 0, In, 0, In->Desc.GetSize());
+
+		// 3. Upload missing data
+		FRDGBufferRef UploadBuffer = GraphBuilder.CreateBuffer(UploadDesc, DebugName, ERDGBufferFlags::MultiFrame);
+		GraphBuilder.QueueBufferUpload(UploadBuffer, InChunk.GetData(), InChunk.Size, ERDGInitialDataFlags::None);
+		InChunk.Release();
+
+		// 4. Append new data to the new buffer
+		AddCopyBufferPass(GraphBuilder, Out, InChunk.Offset, UploadBuffer, 0, InChunk.Size);
+		return Out;
+	}
+#endif
+	else
+	{
+		return nullptr;
+	}
+}
+
+template<typename FormatType>
+void InternalCreateVertexBufferRDG_FromHairBulkData(FRDGBuilder& GraphBuilder, FHairBulkContainer& InChunk, uint32 InDataCount, FRDGExternalBuffer& Out, const TCHAR* DebugName, const FName& OwnerName, EHairResourceUsageType UsageType)
+{
+	// Fallback for non-streamble resources (e.g. guides)
+	if (InChunk.ChunkRequest == nullptr)
+	{
+		InternalCreateVertexBufferRDG_FromBulkData<FormatType>(GraphBuilder, InChunk.Data, InDataCount, Out, DebugName, OwnerName, UsageType);
+		return;
+	}
+
+	const FRDGBufferRef In = Out.Buffer ? Register(GraphBuilder, Out, ERDGImportedBufferFlags::None).Buffer : nullptr;
+	const FRDGBufferDesc BufferDesc = ApplyUsage(FRDGBufferDesc::CreateBufferDesc(FormatType::SizeInByte, FMath::DivideAndRoundUp(InChunk.ChunkRequest->TotalSize, FormatType::SizeInByte)), UsageType);
+	const FRDGBufferDesc UploadDesc = ApplyUsage(FRDGBufferDesc::CreateBufferDesc(FormatType::SizeInByte, FMath::DivideAndRoundUp(InChunk.ChunkRequest->Size, FormatType::SizeInByte)), UsageType);
+	if (BufferDesc.GetSize() == 0 || UploadDesc.GetSize()==0) { Out.Buffer = nullptr; return; }
+	if (FRDGBufferRef Buffer = InternalCreateBufferRDG_FromHairBulkData(GraphBuilder, InChunk, In, BufferDesc, UploadDesc, DebugName, OwnerName))
+	{
+		ConvertToExternalBufferWithViews(GraphBuilder, Buffer, Out, FormatType::Format);
+	}
+}
+
+template<typename FormatType>
+void InternalCreateStructuredBufferRDG_FromHairBulkData(FRDGBuilder& GraphBuilder, FHairBulkContainer& InChunk, uint32 InDataCount, FRDGExternalBuffer& Out, const TCHAR* DebugName, const FName& OwnerName, EHairResourceUsageType UsageType)
+{
+	// Fallback for non-streamble resources (e.g. guides)
+	if (InChunk.ChunkRequest == nullptr)
+	{
+		InternalCreateStructuredBufferRDG_FromHairBulkData<FormatType>(GraphBuilder, InChunk.Data, InDataCount, Out, DebugName, OwnerName, UsageType);
+		return;
+	}
+
+	const FRDGBufferRef In = Out.Buffer ? Register(GraphBuilder, Out, ERDGImportedBufferFlags::None).Buffer : nullptr;
+	const FRDGBufferDesc BufferDesc = ApplyUsage(FRDGBufferDesc::CreateStructuredDesc(FormatType::SizeInByte, FMath::DivideAndRoundUp(InChunk.ChunkRequest->TotalSize, FormatType::SizeInByte)), UsageType);
+	const FRDGBufferDesc UploadDesc = ApplyUsage(FRDGBufferDesc::CreateStructuredDesc(FormatType::SizeInByte, FMath::DivideAndRoundUp(InChunk.ChunkRequest->Size, FormatType::SizeInByte)), UsageType);
+	if (BufferDesc.GetSize() == 0 || UploadDesc.GetSize()==0) { Out.Buffer = nullptr; return; }
+	if (FRDGBufferRef Buffer = InternalCreateBufferRDG_FromHairBulkData(GraphBuilder, InChunk, In, BufferDesc, UploadDesc, DebugName, OwnerName))
+	{
+		ConvertToExternalBufferWithViews(GraphBuilder, Buffer, Out);
+	}
+}
+
+void InternalCreateByteAddressBufferRDG_FromHairBulkData(FRDGBuilder& GraphBuilder, FHairBulkContainer& InChunk, FRDGExternalBuffer& Out, const TCHAR* DebugName, const FName& OwnerName, EHairResourceUsageType UsageType)
+{
+	// Fallback for non-streamble resources (e.g. guides)
+	if (InChunk.ChunkRequest == nullptr)
+	{
+		InternalCreateByteAddressBufferRDG_FromBulkData(GraphBuilder, InChunk.Data, Out, DebugName, OwnerName, UsageType);
+		return;
+	}
+
+	const FRDGBufferRef In = Out.Buffer ? Register(GraphBuilder, Out, ERDGImportedBufferFlags::None).Buffer : nullptr;
+	const FRDGBufferDesc BufferDesc = ApplyUsage(FRDGBufferDesc::CreateByteAddressDesc(InChunk.ChunkRequest->TotalSize), UsageType); // HAIR_STREAMING
+	const FRDGBufferDesc UploadDesc = ApplyUsage(FRDGBufferDesc::CreateByteAddressDesc(InChunk.ChunkRequest->Size), UsageType); // HAIR_STREAMING
+	if (BufferDesc.GetSize() == 0 || UploadDesc.GetSize()==0) { Out.Buffer = nullptr; return; }
+	if (FRDGBufferRef Buffer = InternalCreateBufferRDG_FromHairBulkData(GraphBuilder, InChunk, In, BufferDesc, UploadDesc, DebugName, OwnerName))
+	{
+		ConvertToExternalBufferWithViews(GraphBuilder, Buffer, Out);
+	}
+}
+
+///////////////////////////////////////////////////////////////////////////////////////////////////
+// Regular data loading
+
 static FRDGBufferRef InternalCreateVertexBuffer(
 	FRDGBuilder& GraphBuilder,
 	const TCHAR* Name,
@@ -604,7 +730,7 @@ void FHairCommonResource::InitRHI()
 {
 	if (bIsInitialized || AllocationType == EHairStrandsAllocationType::Deferred || GUsingNullRHI) { return; }
 
-	check(InternalIsDataLoaded());
+	check(InternalIsDataLoaded(HAIR_MAX_NUM_CURVE_PER_GROUP));
 
 	if (bUseRenderGraph)
 	{
@@ -628,40 +754,88 @@ void FHairCommonResource::ReleaseRHI()
 
 void FHairCommonResource::Allocate(FRDGBuilder& GraphBuilder, EHairResourceLoadingType LoadingType)
 {
-	EHairResourceStatus Status = EHairResourceStatus::None;
-	Allocate(GraphBuilder, LoadingType, Status);
+	EHairResourceStatus Status;
+	Status.Status = EHairResourceStatus::EStatus::None;
+	Status.AvailableCurveCount = HAIR_MAX_NUM_CURVE_PER_GROUP;
+	Allocate(GraphBuilder, LoadingType, Status, HAIR_MAX_NUM_CURVE_PER_GROUP);
 }
 
 void FHairCommonResource::Allocate(FRDGBuilder& GraphBuilder, EHairResourceLoadingType LoadingType, EHairResourceStatus& Status)
 {
-	if (bIsInitialized) { Status |= EHairResourceStatus::Valid; return; }
-
-	if (LoadingType == EHairResourceLoadingType::Async && !InternalIsDataLoaded()) { Status |= EHairResourceStatus::Loading; return; }
+	Allocate(GraphBuilder, LoadingType, Status, HAIR_MAX_NUM_CURVE_PER_GROUP);
+}
 
+void FHairCommonResource::Allocate(FRDGBuilder& GraphBuilder, EHairResourceLoadingType LoadingType, EHairResourceStatus& Status, uint32 InRequestedCurveCount)
+{
 	check(AllocationType == EHairStrandsAllocationType::Deferred);
-	FRenderResource::InitResource(); // Call RenderResource InitResource() so that the resources is marked as initialized
-	InternalAllocate(GraphBuilder);
-	bIsInitialized = true;
-	Status |= EHairResourceStatus::Valid;
+
+	// Check if there is already a request in flight	
+
+	if (LoadingType == EHairResourceLoadingType::Sync)
+	{
+		if (!bIsInitialized)
+		{
+			FRenderResource::InitResource(); // Call RenderResource InitResource() so that the resources is marked as initialized
+			InternalAllocate(GraphBuilder);
+			bIsInitialized = true;
+		}
+		Status |= EHairResourceStatus::EStatus::Valid;
+
+	}
+	else if (LoadingType == EHairResourceLoadingType::Async)
+	{
+		// 1. If all requested curve are already loaded, nothing to do
+		if (bIsInitialized && MaxAvailableCurveCount >= InRequestedCurveCount) 
+		{ 
+			Status |= EHairResourceStatus::EStatus::Valid; 
+		}
+		// 2. If more curves are requested, issue a streaming request
+		else if (InternalIsDataLoaded(InRequestedCurveCount))
+		{ 
+			// 2.1 Curve data are availble, and update GPU resources
+			if (!bIsInitialized)
+			{
+				FRenderResource::InitResource(); // Call RenderResource InitResource() so that the resources is marked as initialized
+			}
+			InternalAllocate(GraphBuilder);
+			bIsInitialized = true;
+
+			// Update the max curve count available
+			MaxAvailableCurveCount = StreamingRequest.CurveCount;
+
+			// Reset streaming request. When the request is delete, the DDC request becomes cancelled. 
+			StreamingRequest = FHairStreamingRequest();
+
+			Status |= EHairResourceStatus::EStatus::Valid;
+		}
+		else
+		{
+			// 2.2 Curve data are not availble yet, postpone (new) resources creation
+			Status |= EHairResourceStatus::EStatus::Loading; 
+		}
+	}
+
+	Status.AddAvailableCurve(MaxAvailableCurveCount);
 }
 
 void FHairCommonResource::AllocateLOD(FRDGBuilder& GraphBuilder, int32 LODIndex, EHairResourceLoadingType LoadingType, EHairResourceStatus& Status)
 {
 	// When using a async loading, sub-resources allocation requires the common/main resource to be already initialized.
-	if (LoadingType == EHairResourceLoadingType::Async && (!InternalIsLODDataLoaded(LODIndex) || !bIsInitialized)) { Status |= EHairResourceStatus::Loading; return; }
+	if (LoadingType == EHairResourceLoadingType::Async && (!InternalIsLODDataLoaded(LODIndex) || !bIsInitialized)) { Status |= EHairResourceStatus::EStatus::Loading; return; }
 
 	// Sanity check. 
 	check(AllocationType == EHairStrandsAllocationType::Deferred);
 
 	InternalAllocateLOD(GraphBuilder, LODIndex);
-	Status |= EHairResourceStatus::Valid;
+	Status |= EHairResourceStatus::EStatus::Valid;
 }
 
 void FHairCommonResource::StreamInData()
 {
 	if (!bIsInitialized)
 	{
-		InternalIsDataLoaded();
+		// TODO
+		InternalIsDataLoaded(HAIR_MAX_NUM_CURVE_PER_GROUP);
 	}
 }
 
@@ -669,6 +843,7 @@ void FHairCommonResource::StreamInLODData(int32 LODIndex)
 {
 	if (!bIsInitialized)
 	{
+		// TODO
 		InternalIsLODDataLoaded(LODIndex);
 	}
 }
@@ -914,69 +1089,61 @@ FHairStrandsRestResource::FHairStrandsRestResource(FHairStrandsBulkData& InBulkD
 	FHairCommonResource(EHairStrandsAllocationType::Deferred, InResourceName, InOwnerName),
 	PositionBuffer(), PointAttributeBuffer(), CurveAttributeBuffer(), PointToCurveBuffer(), BulkData(InBulkData), CurveType(InCurveType)
 {
+	MaxAvailableCurveCount = 0;
+
 	// Sanity check
 	check(!!(BulkData.Header.Flags & FHairStrandsBulkData::DataFlags_HasData));
 }
 
-bool FHairStrandsRestResource::InternalIsDataLoaded()
+bool FHairStrandsRestResource::InternalIsDataLoaded(uint32 InRequestedCurveCount)
 {
-	if (BulkDataRequest.IsNone())
+	if (StreamingRequest.IsNone())
 	{
-		BulkDataRequest.Request(BulkData, true, true, OwnerName);
+		StreamingRequest.Request(InRequestedCurveCount, BulkData, false, false, OwnerName);
 	}
-	return BulkDataRequest.IsCompleted();
+	return StreamingRequest.IsCompleted();
 }
 
 bool IsHairStrandsContinousLODEnabled();
 void FHairStrandsRestResource::InternalAllocate(FRDGBuilder& GraphBuilder)
 {
-	// When the request is delete, the DDC request becomes cancelled. 
-	// So the request need to be completed by now
-	//check(BulkDataRequest.IsCompleted());
-	BulkDataRequest = FHairResourceRequest();
+	// If we enter this function, the request need to be completed
+	check(StreamingRequest.IsCompleted());
 
 	const uint32 PointCount = BulkData.GetNumPoints();
 	const uint32 CurveCount = BulkData.GetNumCurves();
 
 	// 1. Lock data, which force the loading data from files (on non-editor build/cooked data). These data are then uploaded to the GPU
 	// 2. A local copy is done by the buffer uploader. This copy is discarded once the uploading is done.
-	InternalCreateVertexBufferRDG_FromBulkData<FHairStrandsPositionFormat>(GraphBuilder, BulkData.Data.Positions.Data, PointCount, PositionBuffer, ToHairResourceDebugName(HAIRSTRANDS_RESOUCE_NAME(CurveType, Hair.StrandsRest_PositionBuffer), ResourceName), OwnerName, EHairResourceUsageType::Static);
-	InternalCreateByteAddressBufferRDG_FromBulkData(GraphBuilder, BulkData.Data.CurveAttributes.Data, CurveAttributeBuffer, ToHairResourceDebugName(HAIRSTRANDS_RESOUCE_NAME(CurveType, Hair.StrandsRest_CurveAttributeBuffer), ResourceName), OwnerName, EHairResourceUsageType::Static);
+	InternalCreateVertexBufferRDG_FromHairBulkData<FHairStrandsPositionFormat>(GraphBuilder, BulkData.Data.Positions, PointCount, PositionBuffer, ToHairResourceDebugName(HAIRSTRANDS_RESOUCE_NAME(CurveType, Hair.StrandsRest_PositionBuffer), ResourceName), OwnerName, EHairResourceUsageType::Static);
+	InternalCreateByteAddressBufferRDG_FromHairBulkData(GraphBuilder, BulkData.Data.CurveAttributes, CurveAttributeBuffer, ToHairResourceDebugName(HAIRSTRANDS_RESOUCE_NAME(CurveType, Hair.StrandsRest_CurveAttributeBuffer), ResourceName), OwnerName, EHairResourceUsageType::Static);
 	if (!!(BulkData.Header.Flags & FHairStrandsBulkData::DataFlags_HasPointAttribute))
 	{
-		InternalCreateByteAddressBufferRDG_FromBulkData(GraphBuilder, BulkData.Data.PointAttributes.Data, PointAttributeBuffer, ToHairResourceDebugName(HAIRSTRANDS_RESOUCE_NAME(CurveType, Hair.StrandsRest_PointAttributeBuffer), ResourceName), OwnerName, EHairResourceUsageType::Static);
+		InternalCreateByteAddressBufferRDG_FromHairBulkData(GraphBuilder, BulkData.Data.PointAttributes, PointAttributeBuffer, ToHairResourceDebugName(HAIRSTRANDS_RESOUCE_NAME(CurveType, Hair.StrandsRest_PointAttributeBuffer), ResourceName), OwnerName, EHairResourceUsageType::Static);
 	}
-	else
+	else if (PointAttributeBuffer.Buffer == nullptr)
 	{
 		TArray<uint32> DummyAttribute;
 		DummyAttribute.Add(0u);
 		InternalCreateByteAddressBufferRDG(GraphBuilder, DummyAttribute, EPixelFormat::PF_R32_UINT, PointAttributeBuffer, ToHairResourceDebugName(HAIRSTRANDS_RESOUCE_NAME(CurveType, Hair.StrandsRest_PointAttributeBuffer), ResourceName), OwnerName, EHairResourceUsageType::Static);
 		GraphBuilder.UseExternalAccessMode(Register(GraphBuilder, PointAttributeBuffer, ERDGImportedBufferFlags::CreateSRV).Buffer, ERHIAccess::SRVMask);
 	}
-	InternalCreateVertexBufferRDG_FromBulkData<FHairStrandsCurveFormat>(GraphBuilder, BulkData.Data.Curves.Data, CurveCount, CurveBuffer, ToHairResourceDebugName(HAIRSTRANDS_RESOUCE_NAME(CurveType, Hair.StrandsRest_CurveBuffer), ResourceName), OwnerName, EHairResourceUsageType::Static);
+	InternalCreateVertexBufferRDG_FromHairBulkData<FHairStrandsCurveFormat>(GraphBuilder, BulkData.Data.Curves, CurveCount, CurveBuffer, ToHairResourceDebugName(HAIRSTRANDS_RESOUCE_NAME(CurveType, Hair.StrandsRest_CurveBuffer), ResourceName), OwnerName, EHairResourceUsageType::Static);
 	if (!!(BulkData.Header.Flags & FHairStrandsBulkData::DataFlags_Has16bitsCurveIndex))
 	{
-		InternalCreateVertexBufferRDG_FromBulkData<FHairStrandsPointToCurveFormat16>(GraphBuilder, BulkData.Data.PointToCurve.Data, PointCount, PointToCurveBuffer, ToHairResourceDebugName(HAIRSTRANDS_RESOUCE_NAME(CurveType, Hair.StrandsRest_PointToCurveBuffer), ResourceName), OwnerName, EHairResourceUsageType::Static);
+		InternalCreateVertexBufferRDG_FromHairBulkData<FHairStrandsPointToCurveFormat16>(GraphBuilder, BulkData.Data.PointToCurve, PointCount, PointToCurveBuffer, ToHairResourceDebugName(HAIRSTRANDS_RESOUCE_NAME(CurveType, Hair.StrandsRest_PointToCurveBuffer), ResourceName), OwnerName, EHairResourceUsageType::Static);
 	}
 	else
 	{
-		InternalCreateVertexBufferRDG_FromBulkData<FHairStrandsPointToCurveFormat32>(GraphBuilder, BulkData.Data.PointToCurve.Data, PointCount, PointToCurveBuffer, ToHairResourceDebugName(HAIRSTRANDS_RESOUCE_NAME(CurveType, Hair.StrandsRest_PointToCurveBuffer), ResourceName), OwnerName, EHairResourceUsageType::Static);
+		InternalCreateVertexBufferRDG_FromHairBulkData<FHairStrandsPointToCurveFormat32>(GraphBuilder, BulkData.Data.PointToCurve, PointCount, PointToCurveBuffer, ToHairResourceDebugName(HAIRSTRANDS_RESOUCE_NAME(CurveType, Hair.StrandsRest_PointToCurveBuffer), ResourceName), OwnerName, EHairResourceUsageType::Static);
 	}
 
-	TArray<FVector4f> RestOffset;
-	RestOffset.Add((FVector3f)BulkData.GetPositionOffset());// LWC_TODO: precision loss
-	InternalCreateVertexBufferRDG<FHairStrandsPositionOffsetFormat>(GraphBuilder, RestOffset, PositionOffsetBuffer, ToHairResourceDebugName(HAIRSTRANDS_RESOUCE_NAME(CurveType, Hair.StrandsRest_PositionOffsetBuffer), ResourceName), OwnerName, EHairResourceUsageType::Static, ERDGInitialDataFlags::None);
-	GraphBuilder.UseExternalAccessMode(Register(GraphBuilder, PositionOffsetBuffer, ERDGImportedBufferFlags::CreateSRV).Buffer, ERHIAccess::SRVMask);
-
-	// Copy curve data for CPU lookup
-	if (IsHairStrandsContinousLODEnabled())
-	{
-		CurveData.SetNum(CurveCount);
-		uint32* Data = (uint32*)CurveData.GetData();
-
-		const uint32* ReadOnlyData = (const uint32*)BulkData.Data.Curves.Data.Lock(LOCK_READ_ONLY);
-		memcpy(Data, ReadOnlyData, sizeof(uint32) * CurveCount);
-		BulkData.Data.Curves.Data.Unlock();
+	if (PositionOffsetBuffer.Buffer == nullptr)
+	{	
+		TArray<FVector4f> RestOffset;
+		RestOffset.Add((FVector3f)BulkData.GetPositionOffset());// LWC_TODO: precision loss
+		InternalCreateVertexBufferRDG<FHairStrandsPositionOffsetFormat>(GraphBuilder, RestOffset, PositionOffsetBuffer, ToHairResourceDebugName(HAIRSTRANDS_RESOUCE_NAME(CurveType, Hair.StrandsRest_PositionOffsetBuffer), ResourceName), OwnerName, EHairResourceUsageType::Static, ERDGInitialDataFlags::None);
+		GraphBuilder.UseExternalAccessMode(Register(GraphBuilder, PositionOffsetBuffer, ERDGImportedBufferFlags::CreateSRV).Buffer, ERHIAccess::SRVMask);
 	}
 }
 
@@ -1016,6 +1183,7 @@ void FHairStrandsRestResource::InternalRelease()
 	PointToCurveBuffer.Release();
 	TangentBuffer.Release();
 	CurveBuffer.Release();
+	MaxAvailableCurveCount = 0;
 }
 
 /////////////////////////////////////////////////////////////////////////////////////////
@@ -1057,10 +1225,11 @@ FRDGExternalBuffer& FHairStrandsDeformedResource::GetDeformerBuffer(FRDGBuilder&
 FRDGExternalBuffer& FHairStrandsDeformedResource::GetDeformerCurveAttributeBuffer(FRDGBuilder& GraphBuilder)
 {
 	// Deformer curve attributes
-	// Bulkdata might be not filled in directly when first called. Hence the extra size check
-	if (DeformerCurveAttributeBuffer.Buffer == nullptr && BulkData.Data.CurveAttributes.Data.GetBulkDataSize() > 0)
+	if (DeformerCurveAttributeBuffer.Buffer == nullptr)
 	{
-		InternalCreateByteAddressBufferRDG(GraphBuilder, BulkData.Data.CurveAttributes.Data.GetBulkDataSize(), DeformerCurveAttributeBuffer, ToHairResourceDebugName(HAIRSTRANDS_RESOUCE_NAME(CurveType, Hair.StrandsDeformed_DeformerCurveAttributeBuffer), ResourceName), OwnerName, EHairResourceUsageType::Dynamic);
+		const uint32 AllocSize = FMath::DivideAndRoundUp(BulkData.Header.CurveCount, BulkData.Header.Strides.CurveAttributeChunkElementCount) * BulkData.Header.Strides.CurveAttributeChunkStride;
+		check(AllocSize > 0);
+		InternalCreateByteAddressBufferRDG(GraphBuilder, AllocSize, DeformerCurveAttributeBuffer, ToHairResourceDebugName(HAIRSTRANDS_RESOUCE_NAME(CurveType, Hair.StrandsDeformed_DeformerCurveAttributeBuffer), ResourceName), OwnerName, EHairResourceUsageType::Dynamic);
 	}
 	return DeformerCurveAttributeBuffer;
 }
@@ -1068,10 +1237,11 @@ FRDGExternalBuffer& FHairStrandsDeformedResource::GetDeformerCurveAttributeBuffe
 FRDGExternalBuffer& FHairStrandsDeformedResource::GetDeformerPointAttributeBuffer(FRDGBuilder& GraphBuilder)
 {
 	// Deformer point attributes
-	// Bulkdata might be not filled in directly when first called. Hence the extra size check
 	if (DeformerPointAttributeBuffer.Buffer == nullptr && (BulkData.Header.Flags & FHairStrandsBulkData::DataFlags_HasPointAttribute) && BulkData.Data.PointAttributes.Data.GetBulkDataSize() > 0)
 	{
-		InternalCreateByteAddressBufferRDG(GraphBuilder, BulkData.Data.PointAttributes.Data.GetBulkDataSize(), DeformerPointAttributeBuffer, ToHairResourceDebugName(HAIRSTRANDS_RESOUCE_NAME(CurveType, Hair.StrandsDeformedt_DeformerPointAttributeBuffer), ResourceName), OwnerName, EHairResourceUsageType::Dynamic);
+		const uint32 AllocSize = FMath::DivideAndRoundUp(BulkData.Header.PointCount, BulkData.Header.Strides.PointAttributeChunkElementCount) * BulkData.Header.Strides.PointAttributeChunkStride;
+		check(AllocSize > 0);
+		InternalCreateByteAddressBufferRDG(GraphBuilder, AllocSize, DeformerPointAttributeBuffer, ToHairResourceDebugName(HAIRSTRANDS_RESOUCE_NAME(CurveType, Hair.StrandsDeformedt_DeformerPointAttributeBuffer), ResourceName), OwnerName, EHairResourceUsageType::Dynamic);
 	}
 	return DeformerPointAttributeBuffer;
 }
@@ -1095,22 +1265,21 @@ FHairStrandsClusterCullingResource::FHairStrandsClusterCullingResource(FHairStra
 	FHairCommonResource(EHairStrandsAllocationType::Deferred, InResourceName, InOwnerName),
 	BulkData(InBulkData)
 {
-
+	MaxAvailableCurveCount = 0;
 }
 
-bool FHairStrandsClusterCullingResource::InternalIsDataLoaded()
+bool FHairStrandsClusterCullingResource::InternalIsDataLoaded(uint32 InRequestedCurveCount)
 {
-	if (BulkDataRequest.IsNone())
+	if (StreamingRequest.IsNone())
 	{
-		BulkDataRequest.Request(BulkData, true, true, OwnerName);
+		StreamingRequest.Request(InRequestedCurveCount, BulkData, true, true, OwnerName); // HAIR_STREAMING : it still use the old resource creation using bulk data mapping
 	}
-	return BulkDataRequest.IsCompleted();
+	return StreamingRequest.IsCompleted();
 }
 
 void FHairStrandsClusterCullingResource::InternalAllocate(FRDGBuilder& GraphBuilder)
 {
-	check(BulkDataRequest.IsCompleted());
-	BulkDataRequest = FHairResourceRequest();
+	check(StreamingRequest.IsCompleted());
 
 	if (ValidateHairBulkData())
 	{
@@ -1130,7 +1299,7 @@ void FHairStrandsClusterCullingResource::InternalRelease()
 	ClusterLODInfoBuffer.Release();
 	ClusterVertexIdBuffer.Release();
 	VertexToClusterIdBuffer.Release();
-	BulkDataRequest = FHairResourceRequest();
+	MaxAvailableCurveCount = 0;
 }
 
 /////////////////////////////////////////////////////////////////////////////////////////
@@ -1157,7 +1326,7 @@ void FHairStrandsRestRootResource::PopulateFromRootData()
 	LODRequests.SetNum(BulkData.MeshProjectionLODs.Num());
 }
 
-bool FHairStrandsRestRootResource::InternalIsDataLoaded()
+bool FHairStrandsRestRootResource::InternalIsDataLoaded(uint32 InRequestedCurveCount)
 {	
 	return true;
 }
@@ -1522,31 +1691,35 @@ FHairStrandsInterpolationResource::FHairStrandsInterpolationResource(FHairStrand
 	FHairCommonResource(EHairStrandsAllocationType::Deferred, InResourceName, InOwnerName),
 	InterpolationBuffer(), BulkData(InBulkData)
 {
+	MaxAvailableCurveCount = 0;
+
 	// Sanity check
 	check(!!(BulkData.Header.Flags & FHairStrandsInterpolationBulkData::DataFlags_HasData));
 }
 
-bool FHairStrandsInterpolationResource::InternalIsDataLoaded()
+bool FHairStrandsInterpolationResource::InternalIsDataLoaded(uint32 InRequestedCurveCount)
 {
-	if (BulkDataRequest.IsNone())
+	if (StreamingRequest.IsNone())
 	{
-		BulkDataRequest.Request(BulkData, true, true, OwnerName);
+		StreamingRequest.Request(InRequestedCurveCount, BulkData, false, false, OwnerName);
 	}
-	return BulkDataRequest.IsCompleted();
+	return StreamingRequest.IsCompleted();
 }
 
 void FHairStrandsInterpolationResource::InternalAllocate(FRDGBuilder& GraphBuilder)
 {
-	BulkDataRequest = FHairResourceRequest();
-	InternalCreateByteAddressBufferRDG_FromBulkData(GraphBuilder, BulkData.Data.Interpolation.Data, InterpolationBuffer, ToHairResourceDebugName(TEXT("Hair.StrandsInterpolation_InterpolationBuffer"), ResourceName), OwnerName, EHairResourceUsageType::Static);
-	InternalCreateVertexBufferRDG_FromBulkData<FHairStrandsRootIndexFormat>(GraphBuilder, BulkData.Data.SimRootPointIndex.Data, BulkData.Header.SimPointCount, SimRootPointIndexBuffer, ToHairResourceDebugName(TEXT("Hair.StrandsInterpolation_SimRootPointIndex"), ResourceName), OwnerName, EHairResourceUsageType::Static);
+	// If we enter this function, the request need to be completed
+	check(StreamingRequest.IsCompleted());
+
+	InternalCreateByteAddressBufferRDG_FromHairBulkData(GraphBuilder, BulkData.Data.Interpolation, InterpolationBuffer, ToHairResourceDebugName(TEXT("Hair.StrandsInterpolation_InterpolationBuffer"), ResourceName), OwnerName, EHairResourceUsageType::Static);
+	InternalCreateVertexBufferRDG_FromHairBulkData<FHairStrandsRootIndexFormat>(GraphBuilder, BulkData.Data.SimRootPointIndex, BulkData.Header.SimPointCount, SimRootPointIndexBuffer, ToHairResourceDebugName(TEXT("Hair.StrandsInterpolation_SimRootPointIndex"), ResourceName), OwnerName, EHairResourceUsageType::Static);
 }
 
 void FHairStrandsInterpolationResource::InternalRelease()
 {
 	InterpolationBuffer.Release();
 	SimRootPointIndexBuffer.Release();
-	BulkDataRequest = FHairResourceRequest();
+	MaxAvailableCurveCount = 0;
 }
 
 /////////////////////////////////////////////////////////////////////////////////////////
diff --git a/Engine/Plugins/Runtime/HairStrands/Source/HairStrandsCore/Private/HairStrandsDatas.cpp b/Engine/Plugins/Runtime/HairStrands/Source/HairStrandsCore/Private/HairStrandsDatas.cpp
--- a/Engine/Plugins/Runtime/HairStrands/Source/HairStrandsCore/Private/HairStrandsDatas.cpp
+++ b/Engine/Plugins/Runtime/HairStrands/Source/HairStrandsCore/Private/HairStrandsDatas.cpp
@@ -4,6 +4,7 @@
 #include "UObject/ReleaseObjectVersion.h"
 #include "UObject/UE5ReleaseStreamObjectVersion.h"
 #include "HairAttributes.h"
+#include "IO/IoDispatcher.h"
 
 #if WITH_EDITOR
 #include "DerivedDataCache.h"
@@ -131,34 +132,38 @@ void FHairStrandsBulkCommon::Write_DDC(UObject* Owner, TArray<UE::DerivedData::F
 #endif
 }
 
-void FHairStrandsBulkCommon::Read_DDC(TArray<UE::DerivedData::FCacheGetChunkRequest>& Out)
+void FHairStrandsBulkCommon::Read_DDC(FHairStreamingRequest* In, TArray<UE::DerivedData::FCacheGetChunkRequest>& Out)
 {
 #if WITH_EDITORONLY_DATA
 	FHairStrandsBulkCommon::FQuery Q;
 	Q.Type = FHairStrandsBulkCommon::FQuery::ReadDDC;
 	Q.OutReadDDC = &Out;
 	Q.DerivedDataKey = &DerivedDataKey;
+	Q.StreamingRequest = In;
+	Q.StreamingRequest->Chunks.Reserve(GetResourceCount()); // This ensures that Chunks array is never reallocated, which would invalidate pointers to FChunk
 	GetResources(Q);
 #endif
 }
 
-void FHairStrandsBulkCommon::Read_IO(FBulkDataBatchRequest& Out)
+void FHairStrandsBulkCommon::Read_IO(FHairStreamingRequest* In, FBulkDataBatchRequest& Out)
 {
 	FBulkDataBatchRequest::FBatchBuilder Batch = Out.NewBatch(GetResourceCount());
 
 	FHairStrandsBulkCommon::FQuery Q;
 	Q.Type = FHairStrandsBulkCommon::FQuery::ReadIO;
 	Q.OutReadIO = &Batch;
+	Q.StreamingRequest = In;
+	Q.StreamingRequest->Chunks.Reserve(GetResourceCount()); // This ensures that Chunks array is never reallocated, which would invalidate pointers to FChunk
 	GetResources(Q);
 	Q.OutReadIO->Issue(Out);
 }
-void FHairStrandsBulkCommon::Write_IO(FArchive& Ar, UObject* Owner)
+void FHairStrandsBulkCommon::Write_IO(UObject* Owner, FArchive& Out)
 {
-	GetResourceVersion(Ar);
+	GetResourceVersion(Out);
 
 	FHairStrandsBulkCommon::FQuery Q;
 	Q.Type = FHairStrandsBulkCommon::FQuery::ReadWriteIO;
-	Q.OutWriteIO = &Ar;
+	Q.OutWriteIO = &Out;
 	Q.Owner = Owner;
 	GetResources(Q);
 }
@@ -171,10 +176,10 @@ void FHairStrandsBulkCommon::Serialize(FArchive& Ar, UObject* Owner)
 
 void FHairStrandsBulkCommon::SerializeData(FArchive& Ar, UObject* Owner)
 {
-	Write_IO(Ar, Owner);
+	Write_IO(Owner, Ar);
 }
 
-void FHairStrandsBulkCommon::FQuery::Add(FHairBulkContainer& In, const TCHAR* InSuffix) 
+void FHairStrandsBulkCommon::FQuery::Add(FHairBulkContainer& In, const TCHAR* InSuffix, uint32 InOffset, uint32 InSize) 
 {
 	check(Type != None);
 #if WITH_EDITORONLY_DATA
@@ -196,22 +201,52 @@ void FHairStrandsBulkCommon::FQuery::Add(FHairBulkContainer& In, const TCHAR* In
 	}
 	else if (Type == ReadDDC)
 	{
+		// 1. Add chunk request to the streaming request. The chunk will hold the request result.
+		check(StreamingRequest);
+		FHairStreamingRequest::FChunk& Chunk = StreamingRequest->Chunks.AddDefaulted_GetRef();
+		Chunk.Status 	= FHairStreamingRequest::FChunk::EStatus::Pending;
+		Chunk.Container = &In;
+		Chunk.Size 		= InSize;
+		Chunk.Offset 	= InOffset;
+		Chunk.TotalSize = InOffset + InSize;
+		In.ChunkRequest = &Chunk;
+
+		// 2. Fill in actual DDC request
 		check(OutReadDDC);
 		using namespace UE::DerivedData;
 		FCacheGetChunkRequest& Out = OutReadDDC->AddDefaulted_GetRef();
-		Out.Id			= FValueId::Null; // HairStrands::HairStrandsValueId : This is only needed for cache record, not cache value.
+		Out.Id			= FValueId::Null; 	// HairStrands::HairStrandsValueId : This is only needed for cache record, not cache value.
 		Out.Key			= ConvertLegacyCacheKey(*DerivedDataKey + InSuffix);
-		Out.RawOffset	= 0; 			//In.Upload.Offset;
-		Out.RawSize		= MAX_uint64; 	//In.Upload.Size;
+		Out.RawOffset	= 0;				//InSize != 0 ? InOffset : 0;
+		Out.RawSize		= MAX_uint64;		//InSize != 0 ? InSize : MAX_uint64;
 		Out.RawHash		= FIoHash();
-		Out.UserData	= (uint64)&In;
+		Out.UserData	= (uint64)&Chunk;
+		if (Owner) { Out.Name = Owner->GetPathName(); }
 	}
 	else 
 #endif
 	if (Type == ReadIO)
 	{
+		// 0. If no size value is provided, use the entire resource
+		if (InSize == 0)
+		{
+			InOffset = 0;
+			InSize = In.Data.GetBulkDataSize();
+		}
+
+		// 1. Add chunk request to the streaming request. The chunk will hold the request result.
+		check(StreamingRequest);
+		FHairStreamingRequest::FChunk& Chunk = StreamingRequest->Chunks.AddDefaulted_GetRef();
+		Chunk.Status 	= FHairStreamingRequest::FChunk::EStatus::Pending;
+		Chunk.Container = &In;
+		Chunk.Size 		= InSize;
+		Chunk.Offset 	= InOffset;
+		Chunk.TotalSize = InOffset + InSize;
+		In.ChunkRequest = &Chunk;
+
+		// 2. Fill in actual DDC request
 		check(OutReadIO);
-		OutReadIO->Read(In.Data);
+		OutReadIO->Read(In.Data, InOffset, InSize, EAsyncIOPriorityAndFlags::AIOP_Normal, Chunk.Data_IO);
 	}
 	else
 	{
@@ -251,6 +286,8 @@ void FHairStrandsBulkData::SerializeHeader(FArchive& Ar, UObject* Owner)
 	Ar << Header.ImportedAttributes;
 	Ar << Header.ImportedAttributeFlags;
 
+	Ar << Header.CurveToPointCount;
+
 	Ar << Header.Strides.PositionStride;
 	Ar << Header.Strides.CurveStride;
 	Ar << Header.Strides.PointToCurveStride;
@@ -279,16 +316,26 @@ void FHairStrandsBulkData::GetResources(FHairStrandsBulkCommon::FQuery& Out)
 	static_assert(sizeof(FHairStrandsPointToCurveFormat32::BulkType) == sizeof(FHairStrandsPointToCurveFormat32::Type));
 	static_assert(sizeof(FHairStrandsRootIndexFormat::BulkType) == sizeof(FHairStrandsRootIndexFormat::Type)); 
 
+	// Translate requested curve count into chunk/offset/size to be read
+	uint32 PointCount = 0;
+	uint32 CurveCount = 0;
+	if (Out.Type == FHairStrandsBulkCommon::FQuery::ReadIO || Out.Type == FHairStrandsBulkCommon::FQuery::ReadDDC)
+	{
+		CurveCount = FMath::Min(Header.CurveCount, Out.GetCurveCount());
+		PointCount = CurveCount > 0 ? Header.CurveToPointCount[CurveCount -1] : 0;
+	}
+
+	// For now Load all data, HAIR_STREAMING
 	if (!!(Header.Flags & DataFlags_HasData))
 	{
-		Out.Add(Data.Positions, TEXT("_Positions"));
-		Out.Add(Data.CurveAttributes, TEXT("_CurveAttributes"));
+		Out.Add(Data.Positions, 			TEXT("_Positions"), 		Data.Positions.LoadedSize, 		PointCount * Header.Strides.PositionStride);
+		Out.Add(Data.CurveAttributes, 		TEXT("_CurveAttributes"), 	0, 								0);
 		if (Header.Flags & DataFlags_HasPointAttribute)
 		{
-			Out.Add(Data.PointAttributes, TEXT("_PointAttributes"));
+			Out.Add(Data.PointAttributes, 	TEXT("_PointAttributes"), 	0, 								0);
 		}
-		Out.Add(Data.PointToCurve, TEXT("_PointToCurve"));
-		Out.Add(Data.Curves, TEXT("_Curves"));
+		Out.Add(Data.PointToCurve, 			TEXT("_PointToCurve"), 		Data.PointToCurve.LoadedSize,	PointCount * Header.Strides.PointToCurveStride);
+		Out.Add(Data.Curves, 				TEXT("_Curves"), 			Data.Curves.LoadedSize, 		CurveCount * Header.Strides.CurveStride);
 	}
 }
 
@@ -360,8 +407,8 @@ void FHairStrandsInterpolationBulkData::GetResources(FHairStrandsBulkCommon::FQu
 
 	if (Header.Flags & DataFlags_HasData)
 	{
-		Out.Add(Data.Interpolation, TEXT("_Interpolation"));
-		Out.Add(Data.SimRootPointIndex, TEXT("_SimRootPointIndex"));
+		Out.Add(Data.Interpolation, TEXT("_Interpolation"), 0, 0); // HAIR_STREAMING
+		Out.Add(Data.SimRootPointIndex, TEXT("_SimRootPointIndex"), 0, 0);
 	}
 }
 
@@ -424,22 +471,22 @@ void FHairStrandsClusterCullingBulkData::GetResources(FHairStrandsBulkCommon::FQ
 {
 	if (Header.ClusterLODCount)
 	{
-		Out.Add(Data.ClusterLODInfos, TEXT("_ClusterLODInfos"));
+		Out.Add(Data.ClusterLODInfos, TEXT("_ClusterLODInfos"), 0, 0); // Load all data
 	}
 
 	if (Header.VertexCount)
 	{
-		Out.Add(Data.VertexToClusterIds, TEXT("_VertexToClusterIds"));
+		Out.Add(Data.VertexToClusterIds, TEXT("_VertexToClusterIds"), 0, 0); // Load all data
 	}
 
 	if (Header.VertexLODCount)
 	{
-		Out.Add(Data.ClusterVertexIds, TEXT("_ClusterVertexIds"));
+		Out.Add(Data.ClusterVertexIds, TEXT("_ClusterVertexIds"), 0, 0); // Load all data
 	}
 
 	if (Header.ClusterCount)
 	{
-		Out.Add(Data.PackedClusterInfos, TEXT("_PackedClusterInfos"));
+		Out.Add(Data.PackedClusterInfos, TEXT("_PackedClusterInfos"), 0, 0); // Load all data
 	}
 
 	if (ValidateHairBulkData() && (Out.Type == FHairStrandsBulkCommon::FQuery::WriteDDC || Out.Type == FHairStrandsBulkCommon::FQuery::ReadWriteIO))
@@ -450,6 +497,8 @@ void FHairStrandsClusterCullingBulkData::GetResources(FHairStrandsBulkCommon::FQ
 
 void FHairStrandsClusterCullingBulkData::Validate(bool bIsSaving)
 {
+	return;
+
 	if (Header.ClusterCount == 0)
 	{
 		return;
diff --git a/Engine/Plugins/Runtime/HairStrands/Source/HairStrandsCore/Public/GroomInstance.h b/Engine/Plugins/Runtime/HairStrands/Source/HairStrandsCore/Public/GroomInstance.h
--- a/Engine/Plugins/Runtime/HairStrands/Source/HairStrandsCore/Public/GroomInstance.h
+++ b/Engine/Plugins/Runtime/HairStrands/Source/HairStrandsCore/Public/GroomInstance.h
@@ -259,6 +259,7 @@ struct HAIRSTRANDSCORE_API FHairGroupInstance : public FHairStrandsInstance
 	bool					bForceCards = false;
 	bool					bUpdatePositionOffset = false;
 	bool					bCastShadow = true;
+	bool 					bSupportStreaming = true;
 	
 	// Deformed component to extract the bone buffer 
 	UMeshComponent*	 DeformedComponent = nullptr;
diff --git a/Engine/Plugins/Runtime/HairStrands/Source/HairStrandsCore/Public/GroomResources.h b/Engine/Plugins/Runtime/HairStrands/Source/HairStrandsCore/Public/GroomResources.h
--- a/Engine/Plugins/Runtime/HairStrands/Source/HairStrandsCore/Public/GroomResources.h
+++ b/Engine/Plugins/Runtime/HairStrands/Source/HairStrandsCore/Public/GroomResources.h
@@ -61,49 +61,27 @@ enum class EHairResourceLoadingType : uint8
 	Sync
 };
 
-enum class EHairResourceStatus : uint8
+struct EHairResourceStatus
 {
-	None = 0,
-	Loading = 1,
-	Valid = 2
-};
-
-FORCEINLINE EHairResourceStatus  operator& (EHairResourceStatus A, EHairResourceStatus B) { return static_cast<EHairResourceStatus>(static_cast<uint8>(A) & static_cast<uint8>(B)); }
-FORCEINLINE EHairResourceStatus  operator| (EHairResourceStatus A, EHairResourceStatus B) { return static_cast<EHairResourceStatus>(static_cast<uint8>(A) | static_cast<uint8>(B)); }
-FORCEINLINE EHairResourceStatus& operator|=(EHairResourceStatus&A, EHairResourceStatus B) { return A=A|B; }
-FORCEINLINE bool operator! (EHairResourceStatus A) { return static_cast<uint8>(A) != 0; }
-
-EHairResourceLoadingType GetHairResourceLoadingType(EHairGeometryType InGeometryType, int32 InLODIndex);
-
-struct FHairChunkRequest
-{
-	enum EStatus { None, Pending, Completed, Failed };
-	FSharedBuffer Data;
-	uint32 Offset = 0;
-	uint32 Size = 0;
+	enum class EStatus : uint8
+	{
+		None = 0,
+		Loading = 1,
+		Valid = 2
+	};
+	void AddAvailableCurve(uint32 In) { AvailableCurveCount = FMath::Min(AvailableCurveCount, In); }
+	bool HasStatus(EStatus In) const { return !!(static_cast<uint8>(Status) & static_cast<uint8>(In)); }
 	EStatus Status = EStatus::None;
-	FHairBulkContainer* Container = nullptr;
-
-	const uint8* GetData() const { check(Status == Completed); return (const uint8*)Data.GetData(); }
-	void Release() { Data.Reset(); }
+	uint32 AvailableCurveCount = 0;
 };
 
-struct FHairResourceRequest
-{
-	void Request(FHairStrandsBulkCommon& In, bool bWait=false, bool bFillBulkData=false, const FName& InOwnerName = NAME_None);
-	bool IsNone() const ;
-	bool IsCompleted() const;
+FORCEINLINE EHairResourceStatus  operator& (EHairResourceStatus In,  EHairResourceStatus::EStatus InStatus) { EHairResourceStatus Out = In; Out.Status = static_cast<EHairResourceStatus::EStatus>(static_cast<uint8>(Out.Status) & static_cast<uint8>(InStatus)); return Out; }
+FORCEINLINE EHairResourceStatus  operator| (EHairResourceStatus In,  EHairResourceStatus::EStatus InStatus) { EHairResourceStatus Out = In; Out.Status = static_cast<EHairResourceStatus::EStatus>(static_cast<uint8>(Out.Status) | static_cast<uint8>(InStatus)); return Out; }
+FORCEINLINE EHairResourceStatus& operator|=(EHairResourceStatus&Out, EHairResourceStatus::EStatus InStatus) { return Out=Out|InStatus; }
 
-#if !WITH_EDITORONLY_DATA
-	// IO
-	FBulkDataBatchRequest IORequest;
-#else
-	// DDC
-	FString PathName;
-	TArray<FHairChunkRequest> Chunks;
-	TUniquePtr<UE::DerivedData::FRequestOwner> DDCRequestOwner;
-#endif
-};
+FORCEINLINE bool operator! (EHairResourceStatus A) { return static_cast<uint8>(A.Status) != 0 && A.AvailableCurveCount > 0; }
+
+EHairResourceLoadingType GetHairResourceLoadingType(EHairGeometryType InGeometryType, int32 InLODIndex);
 
 /* Hair resouces which whom allocation can be deferred */
 struct FHairCommonResource : public FRenderResource
@@ -118,6 +96,7 @@ struct FHairCommonResource : public FRenderResource
 	/* Init/Release buffers (FHairCommonResource) */
 	void Allocate(FRDGBuilder& GraphBuilder, EHairResourceLoadingType LoadingType);
 	void Allocate(FRDGBuilder& GraphBuilder, EHairResourceLoadingType LoadingType, EHairResourceStatus& Status);
+	void Allocate(FRDGBuilder& GraphBuilder, EHairResourceLoadingType LoadingType, EHairResourceStatus& Status, uint32 InRequestedCurveCount);
 	void AllocateLOD(FRDGBuilder& GraphBuilder, int32 LODIndex, EHairResourceLoadingType LoadingType, EHairResourceStatus& Status);
 
 	void StreamInData();
@@ -127,7 +106,7 @@ struct FHairCommonResource : public FRenderResource
 	virtual void InternalAllocate(FRDGBuilder& GraphBuilder) {}
 	virtual void InternalAllocateLOD(FRDGBuilder& GraphBuilder, int32 LODIndex) {}
 	virtual void InternalRelease() {}
-	virtual bool InternalIsDataLoaded() { return true; }
+	virtual bool InternalIsDataLoaded(uint32 InRequestedCurveCount) { return true; }
 	virtual bool InternalIsLODDataLoaded(int32 LODIndex) { return true; }
 
 	bool bUseRenderGraph = true;
@@ -137,6 +116,10 @@ struct FHairCommonResource : public FRenderResource
 	/* Store (precis) debug name */
 	FHairResourceName ResourceName;
 	FName OwnerName;
+
+	/* Handle streaming request */
+	FHairStreamingRequest StreamingRequest;
+	uint32 MaxAvailableCurveCount = HAIR_MAX_NUM_CURVE_PER_GROUP;
 };
 
 /* Render buffers for root deformation for dynamic meshes */
@@ -149,7 +132,7 @@ struct FHairStrandsRestRootResource : public FHairCommonResource
 	virtual void InternalAllocate(FRDGBuilder& GraphBuilder) override;
 	virtual void InternalAllocateLOD(FRDGBuilder& GraphBuilder, int32 LODIndex) override;
 	virtual void InternalRelease() override;
-	virtual bool InternalIsDataLoaded() override;
+	virtual bool InternalIsDataLoaded(uint32 InRequestedCurveCount) override;
 	virtual bool InternalIsLODDataLoaded(int32 LODIndex) override;
 
 	/* Get the resource name */
@@ -317,7 +300,7 @@ struct FHairStrandsRestResource : public FHairCommonResource
 	/* Init/Release buffers */
 	virtual void InternalAllocate(FRDGBuilder& GraphBuilder) override;
 	virtual void InternalRelease() override;
-	virtual bool InternalIsDataLoaded() override;
+	virtual bool InternalIsDataLoaded(uint32 InRequestedCurveCount) override;
 
 	/* Get the resource name */
 	virtual FString GetFriendlyName() const override { return TEXT("FHairStrandsResource"); }
@@ -364,9 +347,6 @@ struct FHairStrandsRestResource : public FHairCommonResource
 	/* Reference to the hair strands render data */
 	FHairStrandsBulkData& BulkData;
 
-	/* Handle to bulk data request */
-	FHairResourceRequest BulkDataRequest;
-
 	/* Type of curves */
 	const EHairStrandsResourcesType CurveType;
 
@@ -463,7 +443,7 @@ struct FHairStrandsClusterCullingResource : public FHairCommonResource
 	/* Init/Release buffers */
 	virtual void InternalAllocate(FRDGBuilder& GraphBuilder) override;
 	virtual void InternalRelease() override;
-	virtual bool InternalIsDataLoaded() override;
+	virtual bool InternalIsDataLoaded(uint32 InRequestedCurveCount) override;
 
 	/* Get the resource name */
 	virtual FString GetFriendlyName() const override { return TEXT("FHairStrandsClusterResource"); }
@@ -490,9 +470,6 @@ struct FHairStrandsClusterCullingResource : public FHairCommonResource
 	FRDGExternalBuffer ClusterVertexIdBuffer;
 
 	FHairStrandsClusterCullingBulkData& BulkData;
-
-	/* Handle to bulk data request */
-	FHairResourceRequest BulkDataRequest;
 };
 
 struct FHairStrandsInterpolationResource : public FHairCommonResource
@@ -503,7 +480,7 @@ struct FHairStrandsInterpolationResource : public FHairCommonResource
 	/* Init/Release buffers */
 	virtual void InternalAllocate(FRDGBuilder& GraphBuilder) override;
 	virtual void InternalRelease() override;
-	virtual bool InternalIsDataLoaded() override;
+	virtual bool InternalIsDataLoaded(uint32 InRequestedCurveCount) override;
 
 	/* Get the resource name */
 	virtual FString GetFriendlyName() const override { return TEXT("FHairStrandsInterplationResource"); }
@@ -524,9 +501,6 @@ struct FHairStrandsInterpolationResource : public FHairCommonResource
 
 	/* Reference to the hair strands interpolation render data */
 	FHairStrandsInterpolationBulkData& BulkData;
-
-	/* Handle to bulk data request */
-	FHairResourceRequest BulkDataRequest;
 };
 
 #if RHI_RAYTRACING
diff --git a/Engine/Plugins/Runtime/HairStrands/Source/HairStrandsCore/Public/HairStrandsDatas.h b/Engine/Plugins/Runtime/HairStrands/Source/HairStrandsCore/Public/HairStrandsDatas.h
--- a/Engine/Plugins/Runtime/HairStrands/Source/HairStrandsCore/Public/HairStrandsDatas.h
+++ b/Engine/Plugins/Runtime/HairStrands/Source/HairStrandsCore/Public/HairStrandsDatas.h
@@ -9,6 +9,12 @@
 #include "RenderGraphResources.h"
 #include "Serialization/BulkData.h"
 #include "HairStrandsDefinitions.h"
+#include "IO/IoDispatcher.h"
+#include "Memory/SharedBuffer.h"
+
+#if WITH_EDITORONLY_DATA
+#include "DerivedDataRequestOwner.h"
+#endif
 
 DECLARE_LOG_CATEGORY_EXTERN(LogHairStrands, Log, All);
 
@@ -232,11 +238,55 @@ struct FHairStrandsRootUtils
 	static float	 PackUVsToFloat(const FVector2f& UV);
 };
 
+struct FHairStrandsBulkCommon;
+struct FHairBulkContainer;
+
+// Streaming request are used for reading hair strands data from DDC or IO
+// The request are translated later to FHairStrandsBulkCommon::FQuery for appropriate reading.
+// FHairStrandsBulkCommon::FQuery abstract DDC/IO/Read/Write for bulk data
+//
+// A query is processed as follow:
+//  FHairStreamingRequest -> FChunk -> FQuery
+struct FHairStreamingRequest
+{
+	struct FChunk
+	{
+		enum EStatus { None, Pending, Completed, Failed };
+		FSharedBuffer Data_DDC;
+		FIoBuffer Data_IO;
+		uint32 Offset = 0;				// Offset to the requested data
+		uint32 Size = 0;				// Size of the requested data
+		uint32 TotalSize = 0; 			// Size of the total data (existing + requested)
+		EStatus Status = EStatus::None;	// Status of the current request
+		FHairBulkContainer* Container = nullptr;
+
+		const uint8* GetData() const;
+		void Release();
+	};
+
+	void Request(uint32 InRequestedCurveCount, FHairStrandsBulkCommon& In, bool bWait=false, bool bFillBulkData=false, const FName& InOwnerName = NAME_None);
+	bool IsNone() const;
+	bool IsCompleted();
+	bool HasPendingRequest() const { return CurveCount > 0; }
+
+#if !WITH_EDITORONLY_DATA
+	// IO
+	FBulkDataBatchRequest IORequest;
+#else
+	// DDC
+	FString PathName;
+	TUniquePtr<UE::DerivedData::FRequestOwner> DDCRequestOwner;
+#endif
+	TArray<FChunk> Chunks;
+	uint32 CurveCount = 0;
+};
+
 struct FHairBulkContainer
 {
 	// Streaming
-	TBitArray<> AvailableChunks;
+	uint32 LoadedSize = 0;
 	FByteBulkData Data;
+	FHairStreamingRequest::FChunk* ChunkRequest = nullptr;
 
 	// Forward BulkData functions
 	bool IsBulkDataLoaded() const 		{ return Data.IsBulkDataLoaded(); }
@@ -255,16 +305,18 @@ struct HAIRSTRANDSCORE_API FHairStrandsBulkCommon
 	void SerializeData(FArchive& Ar, UObject* Owner);
 
 	void Write_DDC(UObject* Owner, TArray<UE::DerivedData::FCachePutValueRequest>& Out);
-	void Read_DDC(TArray<UE::DerivedData::FCacheGetChunkRequest>& Out);
-	void Write_IO(FArchive& Ar, UObject* Owner);
-	void Read_IO(FBulkDataBatchRequest& Out);
+	void Read_DDC(FHairStreamingRequest* In, TArray<UE::DerivedData::FCacheGetChunkRequest>& Out);
+	void Write_IO(UObject* Owner, FArchive& Out);
+	void Read_IO(FHairStreamingRequest* In, FBulkDataBatchRequest& Out);
 
 	struct FQuery
 	{
-		void Add(FHairBulkContainer& In, const TCHAR* InSuffix);
+		void Add(FHairBulkContainer& In, const TCHAR* InSuffix, uint32 InOffset=0, uint32 InSize=0);
+		uint32 GetCurveCount() const { check(StreamingRequest); return StreamingRequest->CurveCount; }
 
-		enum EQueryType { None, ReadDDC, WriteDDC, ReadIO, ReadWriteIO /* aka regular Serialize() */};
+		enum EQueryType { None, ReadDDC, WriteDDC, ReadIO, ReadWriteIO /* i.e. regular Serialize() */};
 		EQueryType Type = None;
+		FHairStreamingRequest* StreamingRequest = nullptr;
 		FBulkDataBatchRequest::FBatchBuilder* 			OutReadIO = nullptr;
 		FArchive*										OutWriteIO = nullptr;
 	#if WITH_EDITORONLY_DATA
@@ -500,6 +552,9 @@ struct HAIRSTRANDSCORE_API FHairStrandsBulkData : FHairStrandsBulkCommon
 		uint32 ImportedAttributes = 0;
 		uint32 ImportedAttributeFlags = 0;
 
+		// Map 'curve' count to 'point' count (used for CLOD)
+		TArray<uint32> CurveToPointCount;
+
 		// Data strides
 		struct FStrides
 		{
@@ -599,6 +654,7 @@ struct FHairStrandsRootBulkData
 
 		/* Map each root onto the unique triangle Id (per-root) */
 		FByteBulkData RootToUniqueTriangleIndexBuffer;
+
 		/* Root's barycentric (per-root) */
 		FByteBulkData RootBarycentricBuffer;
 
