Fix private include issue with forward declaration.

diff --git a/Engine/Plugins/Runtime/HairStrands/Source/HairStrandsCore/Public/HairStrandsDatas.h b/Engine/Plugins/Runtime/HairStrands/Source/HairStrandsCore/Public/HairStrandsDatas.h
--- a/Engine/Plugins/Runtime/HairStrands/Source/HairStrandsCore/Public/HairStrandsDatas.h
+++ b/Engine/Plugins/Runtime/HairStrands/Source/HairStrandsCore/Public/HairStrandsDatas.h
@@ -12,10 +12,6 @@
 #include "IO/IoDispatcher.h"
 #include "Memory/SharedBuffer.h"
 
-#if WITH_EDITORONLY_DATA
-#include "DerivedDataRequestOwner.h"
-#endif
-
 DECLARE_LOG_CATEGORY_EXTERN(LogHairStrands, Log, All);
 
 enum class EHairAttribute : uint8;
@@ -23,6 +19,7 @@ namespace UE::DerivedData
 {
 	struct FCacheGetChunkRequest;
 	struct FCachePutValueRequest;
+	class FRequestOwner;
 }
 
 struct FPackedHairVertex
