Fix missing include causing forward declaration to fail.

diff --git a/Engine/Plugins/Runtime/HairStrands/Source/HairStrandsCore/Private/GroomComponent.cpp b/Engine/Plugins/Runtime/HairStrands/Source/HairStrandsCore/Private/GroomComponent.cpp
--- a/Engine/Plugins/Runtime/HairStrands/Source/HairStrandsCore/Private/GroomComponent.cpp
+++ b/Engine/Plugins/Runtime/HairStrands/Source/HairStrandsCore/Private/GroomComponent.cpp
@@ -42,6 +42,11 @@
 #include "PSOPrecache.h"
 #include "SceneInterface.h"
 
+#if WITH_EDITORONLY_DATA
+#include "DerivedDataCache.h"
+#include "DerivedDataRequestOwner.h"
+#endif
+
 #include UE_INLINE_GENERATED_CPP_BY_NAME(GroomComponent)
 LLM_DECLARE_TAG(Groom);
 
