Add fine grain streaming for interpolation data.

diff --git a/Engine/Plugins/Runtime/HairStrands/Source/HairStrandsCore/Private/GroomAsset.cpp b/Engine/Plugins/Runtime/HairStrands/Source/HairStrandsCore/Private/GroomAsset.cpp
--- a/Engine/Plugins/Runtime/HairStrands/Source/HairStrandsCore/Private/GroomAsset.cpp
+++ b/Engine/Plugins/Runtime/HairStrands/Source/HairStrandsCore/Private/GroomAsset.cpp
@@ -1773,9 +1773,9 @@ static void InternalSerializeStrand(FArchive& Ar, UObject* Owner, FHairGroupPlat
 	// When cooking data, force loading of *all* bulk data prior to saving them
 	if (Ar.IsCooking() && Ar.IsSaving())
 	{
-		{ FHairStreamingRequest R; R.Request(HAIR_MAX_NUM_CURVE_PER_GROUP, StrandData.BulkData,               true /*bWait*/, true /*bFillBulkdata*/, Owner->GetFName()); }
-		{ FHairStreamingRequest R; R.Request(HAIR_MAX_NUM_CURVE_PER_GROUP, StrandData.InterpolationBulkData,  true /*bWait*/, true /*bFillBulkdata*/, Owner->GetFName()); }
-		{ FHairStreamingRequest R; R.Request(HAIR_MAX_NUM_CURVE_PER_GROUP, StrandData.ClusterCullingBulkData, true /*bWait*/, true /*bFillBulkdata*/, Owner->GetFName()); }
+		{ FHairStreamingRequest R; R.Request(HAIR_MAX_NUM_CURVE_PER_GROUP, HAIR_MAX_NUM_POINT_PER_GROUP, StrandData.BulkData,               true /*bWait*/, true /*bFillBulkdata*/, Owner->GetFName()); }
+		{ FHairStreamingRequest R; R.Request(HAIR_MAX_NUM_CURVE_PER_GROUP, HAIR_MAX_NUM_POINT_PER_GROUP, StrandData.InterpolationBulkData,  true /*bWait*/, true /*bFillBulkdata*/, Owner->GetFName()); }
+		{ FHairStreamingRequest R; R.Request(HAIR_MAX_NUM_CURVE_PER_GROUP, HAIR_MAX_NUM_POINT_PER_GROUP, StrandData.ClusterCullingBulkData, true /*bWait*/, true /*bFillBulkdata*/, Owner->GetFName()); }
 	}
 
 	if (!Ar.IsCooking() || !StrandData.bIsCookedOut)
diff --git a/Engine/Plugins/Runtime/HairStrands/Source/HairStrandsCore/Private/GroomBuilder.cpp b/Engine/Plugins/Runtime/HairStrands/Source/HairStrandsCore/Private/GroomBuilder.cpp
--- a/Engine/Plugins/Runtime/HairStrands/Source/HairStrandsCore/Private/GroomBuilder.cpp
+++ b/Engine/Plugins/Runtime/HairStrands/Source/HairStrandsCore/Private/GroomBuilder.cpp
@@ -40,7 +40,7 @@ static FAutoConsoleVariableRef CVarHairGroupIndexBuilder_MaxVoxelResolution(TEXT
 
 FString FGroomBuilder::GetVersion()
 {
-	return TEXT("v8r22");
+	return TEXT("v8r23");
 }
 
 namespace FHairStrandsDecimation
@@ -1451,6 +1451,7 @@ namespace HairInterpolationBuilder
 
 		OutBulkData.Header.Flags = FHairStrandsInterpolationBulkData::DataFlags_HasData;
 		OutBulkData.Header.PointCount = PointCount;
+		OutBulkData.Header.Strides.SimRootPointIndexStride = FHairStrandsRootIndexFormat::SizeInByte;
 
 		if (HairInterpolation.bUseUniqueGuide)
 		{
@@ -1465,6 +1466,7 @@ namespace HairInterpolationBuilder
 			static_assert(sizeof(FHairStrandsInterpolationFormat::Type) == 4u);
 
 			OutBulkData.Header.Flags |= FHairStrandsInterpolationBulkData::DataFlags_HasSingleGuideData;
+			OutBulkData.Header.Strides.InterpolationStride = sizeof(FHairInterpolation1GuideVertex);
 
 			TArray<FHairStrandsInterpolationFormat::Type> OutPointsInterpolation;
 			OutPointsInterpolation.SetNum(PointCount);
@@ -1507,6 +1509,8 @@ namespace HairInterpolationBuilder
 			static_assert(sizeof(FHairInterpolation3GuidesVertex) == HAIR_INTERPOLATION_3GUIDE_STRIDE);
 			static_assert(sizeof(FHairStrandsInterpolationFormat::Type) == 4u);
 
+			OutBulkData.Header.Strides.InterpolationStride = sizeof(FHairInterpolation3GuidesVertex);
+
 			TArray<FHairStrandsInterpolationFormat::Type> OutPointsInterpolation;
 			OutPointsInterpolation.SetNum(PointCount * 4);
 			FHairInterpolation3GuidesVertex* Out = (FHairInterpolation3GuidesVertex*)OutPointsInterpolation.GetData();
@@ -1535,6 +1539,7 @@ namespace HairInterpolationBuilder
 			HairStrandsBuilder::CopyToBulkData<FHairStrandsInterpolationFormat>(OutBulkData.Data.Interpolation, OutPointsInterpolation);
 		}
 
+		// Map point to root point
 		{
 			OutBulkData.Header.SimPointCount = SimDatas.GetNumPoints();
 			const uint32 SimCurveCount = SimDatas.GetNumCurves();
diff --git a/Engine/Plugins/Runtime/HairStrands/Source/HairStrandsCore/Private/GroomManager.cpp b/Engine/Plugins/Runtime/HairStrands/Source/HairStrandsCore/Private/GroomManager.cpp
--- a/Engine/Plugins/Runtime/HairStrands/Source/HairStrandsCore/Private/GroomManager.cpp
+++ b/Engine/Plugins/Runtime/HairStrands/Source/HairStrandsCore/Private/GroomManager.cpp
@@ -935,10 +935,12 @@ static void RunHairLODSelection(
 				GeometryType = EHairGeometryType::NoneGeometry;
 			}
 
-			uint32 RequestedCurveCount  = Instance->HairGroupPublicData->RestCurveCount;
+			uint32 RequestedCurveCount = Instance->HairGroupPublicData->RestCurveCount;
+			uint32 RequestedPointCount = Instance->HairGroupPublicData->RestPointCount;
 			if (GeometryType == EHairGeometryType::Strands && IsHairStrandsContinousLODEnabled() && Instance->bSupportStreaming)
 			{
 				RequestedCurveCount = ComputeActiveCurveCount(Instance->HairGroupPublicData->ContinuousLODScreenSize, Instance->HairGroupPublicData->RestCurveCount);
+				RequestedPointCount = (Instance->Strands.RestResource) ? Instance->Strands.RestResource->BulkData.Header.CurveToPointCount[RequestedCurveCount -1] : 0;
 			}
 
 			const bool bSimulationEnable			= Instance->HairGroupPublicData->IsSimulationEnable(IntLODIndex);
@@ -1007,9 +1009,9 @@ static void RunHairLODSelection(
 				Instance->HairGroupPublicData->Allocate(GraphBuilder);
 
 				if (Instance->Strands.RestRootResource)			{ Instance->Strands.RestRootResource->Allocate(GraphBuilder, LoadingType, ResourceStatus); Instance->Strands.RestRootResource->AllocateLOD(GraphBuilder, MeshLODIndex, LoadingType, ResourceStatus); }
-				if (Instance->Strands.RestResource)				{ Instance->Strands.RestResource->Allocate(GraphBuilder, EHairResourceLoadingType::Async, ResourceStatus, RequestedCurveCount); }
-				if (Instance->Strands.ClusterCullingResource)	{ Instance->Strands.ClusterCullingResource->Allocate(GraphBuilder, EHairResourceLoadingType::Async, ResourceStatus, RequestedCurveCount); }
-				if (Instance->Strands.InterpolationResource)	{ Instance->Strands.InterpolationResource->Allocate(GraphBuilder, EHairResourceLoadingType::Async, ResourceStatus, RequestedCurveCount); }
+				if (Instance->Strands.RestResource)				{ Instance->Strands.RestResource->Allocate(GraphBuilder, EHairResourceLoadingType::Async, ResourceStatus, RequestedCurveCount, RequestedPointCount); }
+				if (Instance->Strands.ClusterCullingResource)	{ Instance->Strands.ClusterCullingResource->Allocate(GraphBuilder, EHairResourceLoadingType::Async, ResourceStatus, RequestedCurveCount, RequestedPointCount); }
+				if (Instance->Strands.InterpolationResource)	{ Instance->Strands.InterpolationResource->Allocate(GraphBuilder, EHairResourceLoadingType::Async, ResourceStatus, RequestedCurveCount, RequestedPointCount); }
 
 				if (Instance->Strands.DeformedRootResource)		{ Instance->Strands.DeformedRootResource->Allocate(GraphBuilder, LoadingType, ResourceStatus); Instance->Strands.DeformedRootResource->AllocateLOD(GraphBuilder, MeshLODIndex, LoadingType, ResourceStatus); }
 				if (Instance->Strands.DeformedResource)			{ Instance->Strands.DeformedResource->Allocate(GraphBuilder, LoadingType, ResourceStatus); }
@@ -1063,6 +1065,10 @@ static void RunHairLODSelection(
 			// Adapt the number curve/point based on Bound/CPU screen-size
 			if (GeometryType == EHairGeometryType::Strands && IsHairStrandsContinousLODEnabled())
 			{
+				if (!bIsLODDataReady)
+				{
+					return false;
+				}
 				check(bIsLODDataReady);
 				check(Instance->Strands.RestResource);
 
diff --git a/Engine/Plugins/Runtime/HairStrands/Source/HairStrandsCore/Private/GroomResources.cpp b/Engine/Plugins/Runtime/HairStrands/Source/HairStrandsCore/Private/GroomResources.cpp
--- a/Engine/Plugins/Runtime/HairStrands/Source/HairStrandsCore/Private/GroomResources.cpp
+++ b/Engine/Plugins/Runtime/HairStrands/Source/HairStrandsCore/Private/GroomResources.cpp
@@ -136,14 +136,16 @@ bool FHairStreamingRequest::IsCompleted()
 // Request fullfil 2 use cases:
 // * Load DDC data and upload them to GPU
 // * Load DDC data and store them into bulkdata for serialization
-void FHairStreamingRequest::Request(uint32 InRequestedCurveCount, FHairStrandsBulkCommon& In, bool bWait, bool bFillBulkdata, const FName& InOwnerName)
+void FHairStreamingRequest::Request(uint32 InRequestedCurveCount, uint32 InRequestedPointCount, FHairStrandsBulkCommon& In, bool bWait, bool bFillBulkdata, const FName& InOwnerName)
 {
 	if (In.GetResourceCount() == 0 || InRequestedCurveCount == 0)
 	{
 		CurveCount = 0;
+		PointCount = 0;
 		return;
 	}
 	CurveCount = InRequestedCurveCount;
+	PointCount = InRequestedPointCount;
 
 	#if !WITH_EDITORONLY_DATA
 	{ 
@@ -730,7 +732,7 @@ void FHairCommonResource::InitRHI()
 {
 	if (bIsInitialized || AllocationType == EHairStrandsAllocationType::Deferred || GUsingNullRHI) { return; }
 
-	check(InternalIsDataLoaded(HAIR_MAX_NUM_CURVE_PER_GROUP));
+	check(InternalIsDataLoaded(HAIR_MAX_NUM_CURVE_PER_GROUP, HAIR_MAX_NUM_POINT_PER_GROUP));
 
 	if (bUseRenderGraph)
 	{
@@ -757,15 +759,15 @@ void FHairCommonResource::Allocate(FRDGBuilder& GraphBuilder, EHairResourceLoadi
 	EHairResourceStatus Status;
 	Status.Status = EHairResourceStatus::EStatus::None;
 	Status.AvailableCurveCount = HAIR_MAX_NUM_CURVE_PER_GROUP;
-	Allocate(GraphBuilder, LoadingType, Status, HAIR_MAX_NUM_CURVE_PER_GROUP);
+	Allocate(GraphBuilder, LoadingType, Status, HAIR_MAX_NUM_CURVE_PER_GROUP, HAIR_MAX_NUM_POINT_PER_GROUP);
 }
 
 void FHairCommonResource::Allocate(FRDGBuilder& GraphBuilder, EHairResourceLoadingType LoadingType, EHairResourceStatus& Status)
 {
-	Allocate(GraphBuilder, LoadingType, Status, HAIR_MAX_NUM_CURVE_PER_GROUP);
+	Allocate(GraphBuilder, LoadingType, Status, HAIR_MAX_NUM_CURVE_PER_GROUP, HAIR_MAX_NUM_POINT_PER_GROUP);
 }
 
-void FHairCommonResource::Allocate(FRDGBuilder& GraphBuilder, EHairResourceLoadingType LoadingType, EHairResourceStatus& Status, uint32 InRequestedCurveCount)
+void FHairCommonResource::Allocate(FRDGBuilder& GraphBuilder, EHairResourceLoadingType LoadingType, EHairResourceStatus& Status, uint32 InRequestedCurveCount, uint32 InRequestedPointCount)
 {
 	check(AllocationType == EHairStrandsAllocationType::Deferred);
 
@@ -778,6 +780,8 @@ void FHairCommonResource::Allocate(FRDGBuilder& GraphBuilder, EHairResourceLoadi
 			FRenderResource::InitResource(); // Call RenderResource InitResource() so that the resources is marked as initialized
 			InternalAllocate(GraphBuilder);
 			bIsInitialized = true;
+
+			MaxAvailableCurveCount = InRequestedCurveCount;
 		}
 		Status |= EHairResourceStatus::EStatus::Valid;
 
@@ -790,7 +794,7 @@ void FHairCommonResource::Allocate(FRDGBuilder& GraphBuilder, EHairResourceLoadi
 			Status |= EHairResourceStatus::EStatus::Valid; 
 		}
 		// 2. If more curves are requested, issue a streaming request
-		else if (InternalIsDataLoaded(InRequestedCurveCount))
+		else if (InternalIsDataLoaded(InRequestedCurveCount, InRequestedPointCount))
 		{ 
 			// 2.1 Curve data are availble, and update GPU resources
 			if (!bIsInitialized)
@@ -835,7 +839,7 @@ void FHairCommonResource::StreamInData()
 	if (!bIsInitialized)
 	{
 		// TODO
-		InternalIsDataLoaded(HAIR_MAX_NUM_CURVE_PER_GROUP);
+		InternalIsDataLoaded(HAIR_MAX_NUM_CURVE_PER_GROUP, HAIR_MAX_NUM_POINT_PER_GROUP);
 	}
 }
 
@@ -1095,11 +1099,11 @@ FHairStrandsRestResource::FHairStrandsRestResource(FHairStrandsBulkData& InBulkD
 	check(!!(BulkData.Header.Flags & FHairStrandsBulkData::DataFlags_HasData));
 }
 
-bool FHairStrandsRestResource::InternalIsDataLoaded(uint32 InRequestedCurveCount)
+bool FHairStrandsRestResource::InternalIsDataLoaded(uint32 InRequestedCurveCount, uint32 InRequestedPointCount)
 {
 	if (StreamingRequest.IsNone())
 	{
-		StreamingRequest.Request(InRequestedCurveCount, BulkData, false, false, OwnerName);
+		StreamingRequest.Request(InRequestedCurveCount, InRequestedPointCount, BulkData, false, false, OwnerName);
 	}
 	return StreamingRequest.IsCompleted();
 }
@@ -1268,11 +1272,11 @@ FHairStrandsClusterCullingResource::FHairStrandsClusterCullingResource(FHairStra
 	MaxAvailableCurveCount = 0;
 }
 
-bool FHairStrandsClusterCullingResource::InternalIsDataLoaded(uint32 InRequestedCurveCount)
+bool FHairStrandsClusterCullingResource::InternalIsDataLoaded(uint32 InRequestedCurveCount, uint32 InRequestedPointCount)
 {
 	if (StreamingRequest.IsNone())
 	{
-		StreamingRequest.Request(InRequestedCurveCount, BulkData, true, true, OwnerName); // HAIR_STREAMING : it still use the old resource creation using bulk data mapping
+		StreamingRequest.Request(InRequestedCurveCount, InRequestedPointCount, BulkData, true, true, OwnerName); // HAIR_STREAMING : it still use the old resource creation using bulk data mapping
 	}
 	return StreamingRequest.IsCompleted();
 }
@@ -1326,11 +1330,6 @@ void FHairStrandsRestRootResource::PopulateFromRootData()
 	LODRequests.SetNum(BulkData.MeshProjectionLODs.Num());
 }
 
-bool FHairStrandsRestRootResource::InternalIsDataLoaded(uint32 InRequestedCurveCount)
-{	
-	return true;
-}
-
 bool FHairStrandsRestRootResource::InternalIsLODDataLoaded(int32 LODIndex)
 {
 	bool bIsLoading = false;
@@ -1697,11 +1696,11 @@ FHairStrandsInterpolationResource::FHairStrandsInterpolationResource(FHairStrand
 	check(!!(BulkData.Header.Flags & FHairStrandsInterpolationBulkData::DataFlags_HasData));
 }
 
-bool FHairStrandsInterpolationResource::InternalIsDataLoaded(uint32 InRequestedCurveCount)
+bool FHairStrandsInterpolationResource::InternalIsDataLoaded(uint32 InRequestedCurveCount, uint32 InRequestedPointCount)
 {
 	if (StreamingRequest.IsNone())
 	{
-		StreamingRequest.Request(InRequestedCurveCount, BulkData, false, false, OwnerName);
+		StreamingRequest.Request(InRequestedCurveCount, InRequestedPointCount, BulkData, false, false, OwnerName);
 	}
 	return StreamingRequest.IsCompleted();
 }
diff --git a/Engine/Plugins/Runtime/HairStrands/Source/HairStrandsCore/Private/HairStrandsDatas.cpp b/Engine/Plugins/Runtime/HairStrands/Source/HairStrandsCore/Private/HairStrandsDatas.cpp
--- a/Engine/Plugins/Runtime/HairStrands/Source/HairStrandsCore/Private/HairStrandsDatas.cpp
+++ b/Engine/Plugins/Runtime/HairStrands/Source/HairStrandsCore/Private/HairStrandsDatas.cpp
@@ -217,8 +217,8 @@ void FHairStrandsBulkCommon::FQuery::Add(FHairBulkContainer& In, const TCHAR* In
 		FCacheGetChunkRequest& Out = OutReadDDC->AddDefaulted_GetRef();
 		Out.Id			= FValueId::Null; 	// HairStrands::HairStrandsValueId : This is only needed for cache record, not cache value.
 		Out.Key			= ConvertLegacyCacheKey(*DerivedDataKey + InSuffix);
-		Out.RawOffset	= 0;				//InSize != 0 ? InOffset : 0;
-		Out.RawSize		= MAX_uint64;		//InSize != 0 ? InSize : MAX_uint64;
+		Out.RawOffset	= InSize != 0 ? InOffset : 0;
+		Out.RawSize		= InSize != 0 ? InSize : MAX_uint64;
 		Out.RawHash		= FIoHash();
 		Out.UserData	= (uint64)&Chunk;
 		if (Owner) { Out.Name = Owner->GetPathName(); }
@@ -325,14 +325,17 @@ void FHairStrandsBulkData::GetResources(FHairStrandsBulkCommon::FQuery& Out)
 		PointCount = CurveCount > 0 ? Header.CurveToPointCount[CurveCount -1] : 0;
 	}
 
+	const uint32 PointAttributeSize = FMath::DivideAndRoundUp(PointCount, Header.Strides.PointAttributeChunkElementCount) * Header.Strides.PointAttributeChunkStride;
+	const uint32 CurveAttributeSize = FMath::DivideAndRoundUp(CurveCount, Header.Strides.CurveAttributeChunkElementCount) * Header.Strides.CurveAttributeChunkStride;
+
 	// For now Load all data, HAIR_STREAMING
 	if (!!(Header.Flags & DataFlags_HasData))
 	{
 		Out.Add(Data.Positions, 			TEXT("_Positions"), 		Data.Positions.LoadedSize, 		PointCount * Header.Strides.PositionStride);
-		Out.Add(Data.CurveAttributes, 		TEXT("_CurveAttributes"), 	0, 								0);
+		Out.Add(Data.CurveAttributes, 		TEXT("_CurveAttributes"), 	Data.Positions.LoadedSize, 		CurveAttributeSize);
 		if (Header.Flags & DataFlags_HasPointAttribute)
 		{
-			Out.Add(Data.PointAttributes, 	TEXT("_PointAttributes"), 	0, 								0);
+			Out.Add(Data.PointAttributes, 	TEXT("_PointAttributes"), 	Data.Positions.LoadedSize, 		PointAttributeSize);
 		}
 		Out.Add(Data.PointToCurve, 			TEXT("_PointToCurve"), 		Data.PointToCurve.LoadedSize,	PointCount * Header.Strides.PointToCurveStride);
 		Out.Add(Data.Curves, 				TEXT("_Curves"), 			Data.Curves.LoadedSize, 		CurveCount * Header.Strides.CurveStride);
@@ -394,6 +397,9 @@ void FHairStrandsInterpolationBulkData::SerializeHeader(FArchive& Ar, UObject* O
 	Ar << Header.Flags;
 	Ar << Header.PointCount;
 	Ar << Header.SimPointCount;
+
+	Ar << Header.Strides.InterpolationStride;
+	Ar << Header.Strides.SimRootPointIndexStride;
 }
 
 uint32 FHairStrandsInterpolationBulkData::GetResourceCount() const
@@ -407,8 +413,15 @@ void FHairStrandsInterpolationBulkData::GetResources(FHairStrandsBulkCommon::FQu
 
 	if (Header.Flags & DataFlags_HasData)
 	{
-		Out.Add(Data.Interpolation, TEXT("_Interpolation"), 0, 0); // HAIR_STREAMING
-		Out.Add(Data.SimRootPointIndex, TEXT("_SimRootPointIndex"), 0, 0);
+		// Translate requested curve count into chunk/offset/size to be read
+		uint32 PointCount = 0;
+		if (Out.Type == FHairStrandsBulkCommon::FQuery::ReadIO || Out.Type == FHairStrandsBulkCommon::FQuery::ReadDDC)
+		{
+			PointCount = FMath::Min(Header.PointCount, Out.GetPointCount());
+		}
+
+		Out.Add(Data.Interpolation, 	TEXT("_Interpolation"), 	Data.Interpolation.LoadedSize, 		PointCount * Header.Strides.InterpolationStride);
+		Out.Add(Data.SimRootPointIndex, TEXT("_SimRootPointIndex"), Data.SimRootPointIndex.LoadedSize, 	Header.SimPointCount * Header.Strides.SimRootPointIndexStride); // Load all data at once (guide data)
 	}
 }
 
diff --git a/Engine/Plugins/Runtime/HairStrands/Source/HairStrandsCore/Public/GroomResources.h b/Engine/Plugins/Runtime/HairStrands/Source/HairStrandsCore/Public/GroomResources.h
--- a/Engine/Plugins/Runtime/HairStrands/Source/HairStrandsCore/Public/GroomResources.h
+++ b/Engine/Plugins/Runtime/HairStrands/Source/HairStrandsCore/Public/GroomResources.h
@@ -96,7 +96,7 @@ struct FHairCommonResource : public FRenderResource
 	/* Init/Release buffers (FHairCommonResource) */
 	void Allocate(FRDGBuilder& GraphBuilder, EHairResourceLoadingType LoadingType);
 	void Allocate(FRDGBuilder& GraphBuilder, EHairResourceLoadingType LoadingType, EHairResourceStatus& Status);
-	void Allocate(FRDGBuilder& GraphBuilder, EHairResourceLoadingType LoadingType, EHairResourceStatus& Status, uint32 InRequestedCurveCount);
+	void Allocate(FRDGBuilder& GraphBuilder, EHairResourceLoadingType LoadingType, EHairResourceStatus& Status, uint32 InRequestedCurveCount, uint32 InRequestedPointCount);
 	void AllocateLOD(FRDGBuilder& GraphBuilder, int32 LODIndex, EHairResourceLoadingType LoadingType, EHairResourceStatus& Status);
 
 	void StreamInData();
@@ -106,7 +106,7 @@ struct FHairCommonResource : public FRenderResource
 	virtual void InternalAllocate(FRDGBuilder& GraphBuilder) {}
 	virtual void InternalAllocateLOD(FRDGBuilder& GraphBuilder, int32 LODIndex) {}
 	virtual void InternalRelease() {}
-	virtual bool InternalIsDataLoaded(uint32 InRequestedCurveCount) { return true; }
+	virtual bool InternalIsDataLoaded(uint32 InRequestedCurveCount, uint32 InRequestedPointCount) { return true; }
 	virtual bool InternalIsLODDataLoaded(int32 LODIndex) { return true; }
 
 	bool bUseRenderGraph = true;
@@ -132,7 +132,6 @@ struct FHairStrandsRestRootResource : public FHairCommonResource
 	virtual void InternalAllocate(FRDGBuilder& GraphBuilder) override;
 	virtual void InternalAllocateLOD(FRDGBuilder& GraphBuilder, int32 LODIndex) override;
 	virtual void InternalRelease() override;
-	virtual bool InternalIsDataLoaded(uint32 InRequestedCurveCount) override;
 	virtual bool InternalIsLODDataLoaded(int32 LODIndex) override;
 
 	/* Get the resource name */
@@ -300,7 +299,7 @@ struct FHairStrandsRestResource : public FHairCommonResource
 	/* Init/Release buffers */
 	virtual void InternalAllocate(FRDGBuilder& GraphBuilder) override;
 	virtual void InternalRelease() override;
-	virtual bool InternalIsDataLoaded(uint32 InRequestedCurveCount) override;
+	virtual bool InternalIsDataLoaded(uint32 InRequestedCurveCount, uint32 InRequestedPointCount) override;
 
 	/* Get the resource name */
 	virtual FString GetFriendlyName() const override { return TEXT("FHairStrandsResource"); }
@@ -443,7 +442,7 @@ struct FHairStrandsClusterCullingResource : public FHairCommonResource
 	/* Init/Release buffers */
 	virtual void InternalAllocate(FRDGBuilder& GraphBuilder) override;
 	virtual void InternalRelease() override;
-	virtual bool InternalIsDataLoaded(uint32 InRequestedCurveCount) override;
+	virtual bool InternalIsDataLoaded(uint32 InRequestedCurveCount, uint32 InRequestedPointCount) override;
 
 	/* Get the resource name */
 	virtual FString GetFriendlyName() const override { return TEXT("FHairStrandsClusterResource"); }
@@ -480,7 +479,7 @@ struct FHairStrandsInterpolationResource : public FHairCommonResource
 	/* Init/Release buffers */
 	virtual void InternalAllocate(FRDGBuilder& GraphBuilder) override;
 	virtual void InternalRelease() override;
-	virtual bool InternalIsDataLoaded(uint32 InRequestedCurveCount) override;
+	virtual bool InternalIsDataLoaded(uint32 InRequestedCurveCount, uint32 InRequestedPointCount) override;
 
 	/* Get the resource name */
 	virtual FString GetFriendlyName() const override { return TEXT("FHairStrandsInterplationResource"); }
diff --git a/Engine/Plugins/Runtime/HairStrands/Source/HairStrandsCore/Public/HairStrandsDatas.h b/Engine/Plugins/Runtime/HairStrands/Source/HairStrandsCore/Public/HairStrandsDatas.h
--- a/Engine/Plugins/Runtime/HairStrands/Source/HairStrandsCore/Public/HairStrandsDatas.h
+++ b/Engine/Plugins/Runtime/HairStrands/Source/HairStrandsCore/Public/HairStrandsDatas.h
@@ -261,10 +261,9 @@ struct FHairStreamingRequest
 		void Release();
 	};
 
-	void Request(uint32 InRequestedCurveCount, FHairStrandsBulkCommon& In, bool bWait=false, bool bFillBulkData=false, const FName& InOwnerName = NAME_None);
+	void Request(uint32 InRequestedCurveCount, uint32 InRequestedPointCount, FHairStrandsBulkCommon& In, bool bWait=false, bool bFillBulkData=false, const FName& InOwnerName = NAME_None);
 	bool IsNone() const;
 	bool IsCompleted();
-	bool HasPendingRequest() const { return CurveCount > 0; }
 
 #if !WITH_EDITORONLY_DATA
 	// IO
@@ -276,6 +275,7 @@ struct FHairStreamingRequest
 #endif
 	TArray<FChunk> Chunks;
 	uint32 CurveCount = 0;
+	uint32 PointCount = 0;
 };
 
 struct FHairBulkContainer
@@ -310,6 +310,7 @@ struct HAIRSTRANDSCORE_API FHairStrandsBulkCommon
 	{
 		void Add(FHairBulkContainer& In, const TCHAR* InSuffix, uint32 InOffset=0, uint32 InSize=0);
 		uint32 GetCurveCount() const { check(StreamingRequest); return StreamingRequest->CurveCount; }
+		uint32 GetPointCount() const { check(StreamingRequest); return StreamingRequest->PointCount; }
 
 		enum EQueryType { None, ReadDDC, WriteDDC, ReadIO, ReadWriteIO /* i.e. regular Serialize() */};
 		EQueryType Type = None;
@@ -384,6 +385,12 @@ struct HAIRSTRANDSCORE_API FHairStrandsInterpolationBulkData : FHairStrandsBulkC
 		uint32 Flags = 0;
 		uint32 PointCount = 0;
 		uint32 SimPointCount = 0;
+
+		struct FStrides
+		{
+			uint32 InterpolationStride = 0;
+			uint32 SimRootPointIndexStride = 0;
+		} Strides;
 	} Header;
 
 	struct FData
