Hair resource are now progressively updated (delta uplaod, opt. reallocation, and GPU copy).

diff --git a/Engine/Plugins/Runtime/HairStrands/Source/HairStrandsCore/Private/GroomAsset.cpp b/Engine/Plugins/Runtime/HairStrands/Source/HairStrandsCore/Private/GroomAsset.cpp
--- a/Engine/Plugins/Runtime/HairStrands/Source/HairStrandsCore/Private/GroomAsset.cpp
+++ b/Engine/Plugins/Runtime/HairStrands/Source/HairStrandsCore/Private/GroomAsset.cpp
@@ -1771,6 +1771,8 @@ static void InternalSerializeStrand(FArchive& Ar, UObject* Owner, FHairGroupPlat
 	Ar.UsingCustomVersion(FAnimObjectVersion::GUID);
 
 	// When cooking data, force loading of *all* bulk data prior to saving them
+	// Note: bFillBulkdata is true for filling in the bulkdata container prior to serialization. This also forces the resources loading 
+	// from the 'start' (i.e., without offset)
 	if (Ar.IsCooking() && Ar.IsSaving())
 	{
 		{ FHairStreamingRequest R; R.Request(HAIR_MAX_NUM_CURVE_PER_GROUP, HAIR_MAX_NUM_POINT_PER_GROUP, StrandData.BulkData,               true /*bWait*/, true /*bFillBulkdata*/, Owner->GetFName()); }
diff --git a/Engine/Plugins/Runtime/HairStrands/Source/HairStrandsCore/Private/GroomBuilder.cpp b/Engine/Plugins/Runtime/HairStrands/Source/HairStrandsCore/Private/GroomBuilder.cpp
--- a/Engine/Plugins/Runtime/HairStrands/Source/HairStrandsCore/Private/GroomBuilder.cpp
+++ b/Engine/Plugins/Runtime/HairStrands/Source/HairStrandsCore/Private/GroomBuilder.cpp
@@ -40,7 +40,7 @@ static FAutoConsoleVariableRef CVarHairGroupIndexBuilder_MaxVoxelResolution(TEXT
 
 FString FGroomBuilder::GetVersion()
 {
-	return TEXT("v8r23");
+	return TEXT("v8r24");
 }
 
 namespace FHairStrandsDecimation
diff --git a/Engine/Plugins/Runtime/HairStrands/Source/HairStrandsCore/Private/GroomResources.cpp b/Engine/Plugins/Runtime/HairStrands/Source/HairStrandsCore/Private/GroomResources.cpp
--- a/Engine/Plugins/Runtime/HairStrands/Source/HairStrandsCore/Private/GroomResources.cpp
+++ b/Engine/Plugins/Runtime/HairStrands/Source/HairStrandsCore/Private/GroomResources.cpp
@@ -119,6 +119,8 @@ bool FHairStreamingRequest::IsCompleted()
 		for (FHairStreamingRequest::FChunk& Chunk : Chunks)
 		{
 			Chunk.Status = FHairStreamingRequest::FChunk::Completed;
+			check(Chunk.Container);
+			Chunk.Container->LoadedSize = Chunk.TotalSize;
 		}
 	}
 	return Chunks.Num() == 0 || IORequest.IsCompleted(); 
@@ -157,6 +159,10 @@ void FHairStreamingRequest::Request(uint32 InRequestedCurveCount, uint32 InReque
 	}
 	#else
 	{
+		// When enabled, data can be loaded from an offset. Otherwisee, start from the beginning of the resource
+		// This is used when cooking data to force the loading of the entire resource (i.e., bSupportOffsetLoad=false)
+		bSupportOffsetLoad = !bFillBulkdata; 
+
 		using namespace UE::DerivedData;
 		TArray<FCacheGetChunkRequest> Requests;
 		In.Read_DDC(this, Requests);
@@ -178,8 +184,7 @@ void FHairStreamingRequest::Request(uint32 InRequestedCurveCount, uint32 InReque
 				Chunk.Status 	= FHairStreamingRequest::FChunk::Completed;
 		
 				// Upload the total amount of loaded data
-				// HAIR_TODO: Disable for now, reupload all the data until issue with RDG/DX12 source/target buffer is solved
-				// Chunk.Container->LoadedSize = Chunk.TotalSize;
+				Chunk.Container->LoadedSize = Chunk.TotalSize;
 
 				// Optional fill in of bytebulkdata container
 				if (bFillBulkdata)
@@ -388,10 +393,13 @@ static FRDGBufferRef InternalCreateBufferRDG_FromHairBulkData(FRDGBuilder& Graph
 	check(InData.ChunkRequest);
 	FHairStreamingRequest::FChunk& InChunk = *InData.ChunkRequest;
 
-	const bool bCopyData = In != nullptr;
-	// For now always upload the entire data set as partial upload + GPU copy result in D3D12 error due to buffer aliasing HAIR_TODO
-	const bool bNeedCreateBuffer = In == nullptr || BufferDesc.GetSize() == UploadDesc.GetSize() || bCopyData;
-	if (bNeedCreateBuffer)
+	const bool bCreate 		= In == nullptr || BufferDesc.GetSize() == UploadDesc.GetSize();
+	const bool bCopy 		= In != nullptr && InChunk.Size > 0;
+	const bool bReallocate 	= In != nullptr && In->Desc.GetSize() < InChunk.TotalSize;
+
+	check(InChunk.Size > 0);
+	// Either create a new buffer or append new data to existing buffer
+	if (bCreate)
 	{	
 		check(BufferDesc.GetSize() >= UploadDesc.GetSize());
 
@@ -400,25 +408,33 @@ static FRDGBufferRef InternalCreateBufferRDG_FromHairBulkData(FRDGBuilder& Graph
 		InChunk.Release();
 		return Out;
 	}
-#if 0 // HAIR_TODO
-	else if (bCopyData && InChunk.Size > 0)
+	else if (bCopy)
 	{
-		// 1. Create new buffer
-		FRDGBufferRef Out = GraphBuilder.CreateBuffer(BufferDesc, DebugName, ERDGBufferFlags::MultiFrame);
-
-		// 2. Copy existing data from the old buffer to the new buffer
-		AddCopyBufferPass(GraphBuilder, Out, 0, In, 0, In->Desc.GetSize());
+		// 1. If the current buffer is too small for storing the new data, reallocate it
+		FRDGBufferRef Out = In;
+		if (bReallocate)
+		{
+			// 1.1 Create new buffer
+			FRDGBufferDesc NewBufferDesc = BufferDesc;
+			NewBufferDesc.Usage |= EBufferUsageFlags::UnorderedAccess;
+			Out = GraphBuilder.CreateBuffer(NewBufferDesc, DebugName, ERDGBufferFlags::MultiFrame);
+	
+			// 1.2 Copy existing data from the old buffer to the new buffer
+			AddCopyBufferPass(GraphBuilder, Out, 0, In, 0, In->Desc.GetSize());
+		}
 
-		// 3. Upload missing data
+		// 2. Upload missing data
 		FRDGBufferRef UploadBuffer = GraphBuilder.CreateBuffer(UploadDesc, DebugName, ERDGBufferFlags::MultiFrame);
 		GraphBuilder.QueueBufferUpload(UploadBuffer, InChunk.GetData(), InChunk.Size, ERDGInitialDataFlags::None);
 		InChunk.Release();
 
-		// 4. Append new data to the new buffer
+		// 4. Append new data to the new/existing buffer
 		AddCopyBufferPass(GraphBuilder, Out, InChunk.Offset, UploadBuffer, 0, InChunk.Size);
+
+		// Return the new buffer if it needs to be extracted
 		return Out;
+		//return bReallocate ? Out : nullptr;
 	}
-#endif
 	else
 	{
 		return nullptr;
@@ -438,7 +454,8 @@ void InternalCreateVertexBufferRDG_FromHairBulkData(FRDGBuilder& GraphBuilder, F
 	const FRDGBufferRef In = Out.Buffer ? Register(GraphBuilder, Out, ERDGImportedBufferFlags::None).Buffer : nullptr;
 	const FRDGBufferDesc BufferDesc = ApplyUsage(FRDGBufferDesc::CreateBufferDesc(FormatType::SizeInByte, FMath::DivideAndRoundUp(InChunk.ChunkRequest->TotalSize, FormatType::SizeInByte)), UsageType);
 	const FRDGBufferDesc UploadDesc = ApplyUsage(FRDGBufferDesc::CreateBufferDesc(FormatType::SizeInByte, FMath::DivideAndRoundUp(InChunk.ChunkRequest->Size, FormatType::SizeInByte)), UsageType);
-	if (BufferDesc.GetSize() == 0 || UploadDesc.GetSize()==0) { Out.Buffer = nullptr; return; }
+	if (BufferDesc.GetSize() == 0) 	{ Out.Buffer = nullptr; return; }
+	if (UploadDesc.GetSize()==0) 	{ return; }
 	if (FRDGBufferRef Buffer = InternalCreateBufferRDG_FromHairBulkData(GraphBuilder, InChunk, In, BufferDesc, UploadDesc, DebugName, OwnerName))
 	{
 		ConvertToExternalBufferWithViews(GraphBuilder, Buffer, Out, FormatType::Format);
@@ -458,7 +475,8 @@ void InternalCreateStructuredBufferRDG_FromHairBulkData(FRDGBuilder& GraphBuilde
 	const FRDGBufferRef In = Out.Buffer ? Register(GraphBuilder, Out, ERDGImportedBufferFlags::None).Buffer : nullptr;
 	const FRDGBufferDesc BufferDesc = ApplyUsage(FRDGBufferDesc::CreateStructuredDesc(FormatType::SizeInByte, FMath::DivideAndRoundUp(InChunk.ChunkRequest->TotalSize, FormatType::SizeInByte)), UsageType);
 	const FRDGBufferDesc UploadDesc = ApplyUsage(FRDGBufferDesc::CreateStructuredDesc(FormatType::SizeInByte, FMath::DivideAndRoundUp(InChunk.ChunkRequest->Size, FormatType::SizeInByte)), UsageType);
-	if (BufferDesc.GetSize() == 0 || UploadDesc.GetSize()==0) { Out.Buffer = nullptr; return; }
+	if (BufferDesc.GetSize() == 0) 	{ Out.Buffer = nullptr; return; }
+	if (UploadDesc.GetSize()==0) 	{ return; }
 	if (FRDGBufferRef Buffer = InternalCreateBufferRDG_FromHairBulkData(GraphBuilder, InChunk, In, BufferDesc, UploadDesc, DebugName, OwnerName))
 	{
 		ConvertToExternalBufferWithViews(GraphBuilder, Buffer, Out);
@@ -475,9 +493,10 @@ void InternalCreateByteAddressBufferRDG_FromHairBulkData(FRDGBuilder& GraphBuild
 	}
 
 	const FRDGBufferRef In = Out.Buffer ? Register(GraphBuilder, Out, ERDGImportedBufferFlags::None).Buffer : nullptr;
-	const FRDGBufferDesc BufferDesc = ApplyUsage(FRDGBufferDesc::CreateByteAddressDesc(InChunk.ChunkRequest->TotalSize), UsageType); // HAIR_STREAMING
-	const FRDGBufferDesc UploadDesc = ApplyUsage(FRDGBufferDesc::CreateByteAddressDesc(InChunk.ChunkRequest->Size), UsageType); // HAIR_STREAMING
-	if (BufferDesc.GetSize() == 0 || UploadDesc.GetSize()==0) { Out.Buffer = nullptr; return; }
+	const FRDGBufferDesc BufferDesc = ApplyUsage(FRDGBufferDesc::CreateByteAddressDesc(InChunk.ChunkRequest->TotalSize), UsageType);
+	const FRDGBufferDesc UploadDesc = ApplyUsage(FRDGBufferDesc::CreateByteAddressDesc(InChunk.ChunkRequest->Size), UsageType);
+	if (BufferDesc.GetSize() == 0) 	{ Out.Buffer = nullptr; return; }
+	if (UploadDesc.GetSize()==0) 	{ return; }
 	if (FRDGBufferRef Buffer = InternalCreateBufferRDG_FromHairBulkData(GraphBuilder, InChunk, In, BufferDesc, UploadDesc, DebugName, OwnerName))
 	{
 		ConvertToExternalBufferWithViews(GraphBuilder, Buffer, Out);
@@ -1276,7 +1295,7 @@ bool FHairStrandsClusterCullingResource::InternalIsDataLoaded(uint32 InRequested
 {
 	if (StreamingRequest.IsNone())
 	{
-		StreamingRequest.Request(InRequestedCurveCount, InRequestedPointCount, BulkData, true, true, OwnerName); // HAIR_STREAMING : it still use the old resource creation using bulk data mapping
+		StreamingRequest.Request(InRequestedCurveCount, InRequestedPointCount, BulkData, true, true, OwnerName);
 	}
 	return StreamingRequest.IsCompleted();
 }
diff --git a/Engine/Plugins/Runtime/HairStrands/Source/HairStrandsCore/Private/HairStrandsDatas.cpp b/Engine/Plugins/Runtime/HairStrands/Source/HairStrandsCore/Private/HairStrandsDatas.cpp
--- a/Engine/Plugins/Runtime/HairStrands/Source/HairStrandsCore/Private/HairStrandsDatas.cpp
+++ b/Engine/Plugins/Runtime/HairStrands/Source/HairStrandsCore/Private/HairStrandsDatas.cpp
@@ -201,13 +201,15 @@ void FHairStrandsBulkCommon::FQuery::Add(FHairBulkContainer& In, const TCHAR* In
 	}
 	else if (Type == ReadDDC)
 	{
-		// 1. Add chunk request to the streaming request. The chunk will hold the request result.
 		check(StreamingRequest);
+		const bool bOffset = StreamingRequest->bSupportOffsetLoad;
+
+		// 1. Add chunk request to the streaming request. The chunk will hold the request result.
 		FHairStreamingRequest::FChunk& Chunk = StreamingRequest->Chunks.AddDefaulted_GetRef();
 		Chunk.Status 	= FHairStreamingRequest::FChunk::EStatus::Pending;
 		Chunk.Container = &In;
 		Chunk.Size 		= InSize;
-		Chunk.Offset 	= InOffset;
+		Chunk.Offset 	= bOffset ? InOffset : 0;
 		Chunk.TotalSize = InOffset + InSize;
 		In.ChunkRequest = &Chunk;
 
@@ -217,7 +219,7 @@ void FHairStrandsBulkCommon::FQuery::Add(FHairBulkContainer& In, const TCHAR* In
 		FCacheGetChunkRequest& Out = OutReadDDC->AddDefaulted_GetRef();
 		Out.Id			= FValueId::Null; 	// HairStrands::HairStrandsValueId : This is only needed for cache record, not cache value.
 		Out.Key			= ConvertLegacyCacheKey(*DerivedDataKey + InSuffix);
-		Out.RawOffset	= InSize != 0 ? InOffset : 0;
+		Out.RawOffset	= bOffset && InSize != 0 ? InOffset : 0;
 		Out.RawSize		= InSize != 0 ? InSize : MAX_uint64;
 		Out.RawHash		= FIoHash();
 		Out.UserData	= (uint64)&Chunk;
@@ -328,14 +330,13 @@ void FHairStrandsBulkData::GetResources(FHairStrandsBulkCommon::FQuery& Out)
 	const uint32 PointAttributeSize = FMath::DivideAndRoundUp(PointCount, Header.Strides.PointAttributeChunkElementCount) * Header.Strides.PointAttributeChunkStride;
 	const uint32 CurveAttributeSize = FMath::DivideAndRoundUp(CurveCount, Header.Strides.CurveAttributeChunkElementCount) * Header.Strides.CurveAttributeChunkStride;
 
-	// For now Load all data, HAIR_STREAMING
 	if (!!(Header.Flags & DataFlags_HasData))
 	{
 		Out.Add(Data.Positions, 			TEXT("_Positions"), 		Data.Positions.LoadedSize, 		PointCount * Header.Strides.PositionStride);
-		Out.Add(Data.CurveAttributes, 		TEXT("_CurveAttributes"), 	Data.Positions.LoadedSize, 		CurveAttributeSize);
+		Out.Add(Data.CurveAttributes, 		TEXT("_CurveAttributes"), 	Data.CurveAttributes.LoadedSize,CurveAttributeSize);
 		if (Header.Flags & DataFlags_HasPointAttribute)
 		{
-			Out.Add(Data.PointAttributes, 	TEXT("_PointAttributes"), 	Data.Positions.LoadedSize, 		PointAttributeSize);
+			Out.Add(Data.PointAttributes, 	TEXT("_PointAttributes"), 	Data.PointAttributes.LoadedSize,PointAttributeSize);
 		}
 		Out.Add(Data.PointToCurve, 			TEXT("_PointToCurve"), 		Data.PointToCurve.LoadedSize,	PointCount * Header.Strides.PointToCurveStride);
 		Out.Add(Data.Curves, 				TEXT("_Curves"), 			Data.Curves.LoadedSize, 		CurveCount * Header.Strides.CurveStride);
diff --git a/Engine/Plugins/Runtime/HairStrands/Source/HairStrandsCore/Public/HairStrandsDatas.h b/Engine/Plugins/Runtime/HairStrands/Source/HairStrandsCore/Public/HairStrandsDatas.h
--- a/Engine/Plugins/Runtime/HairStrands/Source/HairStrandsCore/Public/HairStrandsDatas.h
+++ b/Engine/Plugins/Runtime/HairStrands/Source/HairStrandsCore/Public/HairStrandsDatas.h
@@ -244,8 +244,19 @@ struct FHairBulkContainer;
 //
 // A query is processed as follow:
 //  FHairStreamingRequest -> FChunk -> FQuery
+//   ____________________________________________________________
+//  |                    FHairCommonResource                     |
+//  |_______________________            _________________________|
+//  | FHairStreamingRequest |          |  FHairStrandsBulkCommon |
+//	|                       |          |                         |
+//	|       FChunk  --------|--FQuery--|--> FHairBulkContainer   |
+//	|       FChunk  --------|--FQuery--|--> FHairBulkContainer   |
+//	|       FChunk  --------|--FQuery--|--> FHairBulkContainer   |
+//  |_______________________|          |_________________________|
+
 struct FHairStreamingRequest
 {
+	// Hold request at the container level (transient to the request)
 	struct FChunk
 	{
 		enum EStatus { None, Pending, Completed, Failed };
@@ -276,6 +287,10 @@ struct FHairStreamingRequest
 	TArray<FChunk> Chunks;
 	uint32 CurveCount = 0;
 	uint32 PointCount = 0;
+
+	// When enabled, data can be loaded from an offset. Otherwisee, start from the beginning of the resource
+	// This is used when cooking data to force the loading of the entire resource (i.e., bSupportOffsetLoad=false)
+	bool bSupportOffsetLoad = true; 
 };
 
 struct FHairBulkContainer
