Fix crash when loading CitySample for groom with cards only, having no guides data.

diff --git a/Engine/Plugins/Runtime/HairStrands/Source/HairStrandsCore/Private/HairStrandsDatas.cpp b/Engine/Plugins/Runtime/HairStrands/Source/HairStrandsCore/Private/HairStrandsDatas.cpp
--- a/Engine/Plugins/Runtime/HairStrands/Source/HairStrandsCore/Private/HairStrandsDatas.cpp
+++ b/Engine/Plugins/Runtime/HairStrands/Source/HairStrandsCore/Private/HairStrandsDatas.cpp
@@ -327,8 +327,8 @@ void FHairStrandsBulkData::GetResources(FHairStrandsBulkCommon::FQuery& Out)
 		PointCount = CurveCount > 0 ? Header.CurveToPointCount[CurveCount -1] : 0;
 	}
 
-	const uint32 PointAttributeSize = FMath::DivideAndRoundUp(PointCount, Header.Strides.PointAttributeChunkElementCount) * Header.Strides.PointAttributeChunkStride;
-	const uint32 CurveAttributeSize = FMath::DivideAndRoundUp(CurveCount, Header.Strides.CurveAttributeChunkElementCount) * Header.Strides.CurveAttributeChunkStride;
+	const uint32 PointAttributeSize = PointCount > 0 ? FMath::DivideAndRoundUp(PointCount, Header.Strides.PointAttributeChunkElementCount) * Header.Strides.PointAttributeChunkStride : 0;
+	const uint32 CurveAttributeSize = CurveCount > 0 ? FMath::DivideAndRoundUp(CurveCount, Header.Strides.CurveAttributeChunkElementCount) * Header.Strides.CurveAttributeChunkStride : 0;
 
 	if (!!(Header.Flags & DataFlags_HasData))
 	{
