Fix incorrect delta size during groom streaming update.

diff --git a/Engine/Plugins/Runtime/HairStrands/Source/HairStrandsCore/Private/HairStrandsDatas.cpp b/Engine/Plugins/Runtime/HairStrands/Source/HairStrandsCore/Private/HairStrandsDatas.cpp
--- a/Engine/Plugins/Runtime/HairStrands/Source/HairStrandsCore/Private/HairStrandsDatas.cpp
+++ b/Engine/Plugins/Runtime/HairStrands/Source/HairStrandsCore/Private/HairStrandsDatas.cpp
@@ -202,15 +202,17 @@ void FHairStrandsBulkCommon::FQuery::Add(FHairBulkContainer& In, const TCHAR* In
 	else if (Type == ReadDDC)
 	{
 		check(StreamingRequest);
-		const bool bOffset = StreamingRequest->bSupportOffsetLoad;
+		InOffset = StreamingRequest->bSupportOffsetLoad ? InOffset : 0;
+
+		check (InSize >= InOffset);
 
 		// 1. Add chunk request to the streaming request. The chunk will hold the request result.
 		FHairStreamingRequest::FChunk& Chunk = StreamingRequest->Chunks.AddDefaulted_GetRef();
 		Chunk.Status 	= FHairStreamingRequest::FChunk::EStatus::Pending;
 		Chunk.Container = &In;
-		Chunk.Size 		= InSize;
-		Chunk.Offset 	= bOffset ? InOffset : 0;
-		Chunk.TotalSize = InOffset + InSize;
+		Chunk.Size 		= InSize - InOffset;
+		Chunk.Offset 	= InOffset;
+		Chunk.TotalSize = InSize;
 		In.ChunkRequest = &Chunk;
 
 		// 2. Fill in actual DDC request
@@ -219,8 +221,8 @@ void FHairStrandsBulkCommon::FQuery::Add(FHairBulkContainer& In, const TCHAR* In
 		FCacheGetChunkRequest& Out = OutReadDDC->AddDefaulted_GetRef();
 		Out.Id			= FValueId::Null; 	// HairStrands::HairStrandsValueId : This is only needed for cache record, not cache value.
 		Out.Key			= ConvertLegacyCacheKey(*DerivedDataKey + InSuffix);
-		Out.RawOffset	= bOffset && InSize != 0 ? InOffset : 0;
-		Out.RawSize		= InSize != 0 ? InSize : MAX_uint64;
+		Out.RawOffset	= InOffset;
+		Out.RawSize		= InSize != 0 ? InSize-InOffset : MAX_uint64;
 		Out.RawHash		= FIoHash();
 		Out.UserData	= (uint64)&Chunk;
 		if (Owner) { Out.Name = Owner->GetPathName(); }
@@ -235,15 +237,16 @@ void FHairStrandsBulkCommon::FQuery::Add(FHairBulkContainer& In, const TCHAR* In
 			InOffset = 0;
 			InSize = In.Data.GetBulkDataSize();
 		}
+		check (InSize >= InOffset);
 
 		// 1. Add chunk request to the streaming request. The chunk will hold the request result.
 		check(StreamingRequest);
 		FHairStreamingRequest::FChunk& Chunk = StreamingRequest->Chunks.AddDefaulted_GetRef();
 		Chunk.Status 	= FHairStreamingRequest::FChunk::EStatus::Pending;
 		Chunk.Container = &In;
-		Chunk.Size 		= InSize;
+		Chunk.Size 		= InSize - InOffset;
 		Chunk.Offset 	= InOffset;
-		Chunk.TotalSize = InOffset + InSize;
+		Chunk.TotalSize = InSize;
 		In.ChunkRequest = &Chunk;
 
 		// 2. Fill in actual DDC request
