Split root data into header/data.

diff --git a/Engine/Plugins/Runtime/HairStrands/Source/HairStrandsCore/Private/GroomAsset.cpp b/Engine/Plugins/Runtime/HairStrands/Source/HairStrandsCore/Private/GroomAsset.cpp
--- a/Engine/Plugins/Runtime/HairStrands/Source/HairStrandsCore/Private/GroomAsset.cpp
+++ b/Engine/Plugins/Runtime/HairStrands/Source/HairStrandsCore/Private/GroomAsset.cpp
@@ -322,7 +322,7 @@ void DumpLoadedGroomBindingAssets(IConsoleVariable* InCVarPakTesterEnabled)
 				Total_GPUMemorySize_Strands+= GPUMemorySize_Strands;
 				Total_GPUMemorySize_Cards  += GPUMemorySize_Cards;
 
-				const uint32 SkelLODCount = AssetIt->HairGroupBulkDatas[GroupIt].RenRootBulkData.MeshProjectionLODs.Num();
+				const uint32 SkelLODCount = AssetIt->HairGroupBulkDatas[GroupIt].RenRootBulkData.GetLODCount();
 				if (bDetails)
 				{
 //					UE_LOG(LogHairStrands, Log, TEXT("--  No.  - LOD -    CPU Total (     Guides|    Strands|      Cards) -    GPU Total (     Guides|    Strands|      Cards) - Asset Name "));
diff --git a/Engine/Plugins/Runtime/HairStrands/Source/HairStrandsCore/Private/GroomBindingAsset.cpp b/Engine/Plugins/Runtime/HairStrands/Source/HairStrandsCore/Private/GroomBindingAsset.cpp
--- a/Engine/Plugins/Runtime/HairStrands/Source/HairStrandsCore/Private/GroomBindingAsset.cpp
+++ b/Engine/Plugins/Runtime/HairStrands/Source/HairStrandsCore/Private/GroomBindingAsset.cpp
@@ -417,26 +417,26 @@ bool UGroomBindingAsset::IsCompatible(const USkeletalMesh* InSkeletalMesh, const
 
 		for (const FHairGroupResource& Resource : InBinding->HairGroupResources)
 		{
-			if (Resource.SimRootResources && InSkeletalMesh->GetLODNum() != Resource.SimRootResources->BulkData.MeshProjectionLODs.Num())
+			if (Resource.SimRootResources && InSkeletalMesh->GetLODNum() != Resource.SimRootResources->GetLODCount())
 			{
 				if (bIssueWarning)
 				{
 					UE_LOG(LogHairStrands, Warning, TEXT("[Groom] The Groom binding (%s) does not have the same have LOD count (LOD sim:%d) than the skeletal mesh (%s, LOD:%d). The binding asset will not be used."),
 						*InBinding->GetName(),
-						Resource.SimRootResources->BulkData.MeshProjectionLODs.Num(),
+						Resource.SimRootResources->GetLODCount(),
 						*InSkeletalMesh->GetName(),
 						InSkeletalMesh->GetLODNum());
 				}
 				return false;
 			}
 
-			if (Resource.RenRootResources && InSkeletalMesh->GetLODNum() != Resource.RenRootResources->BulkData.MeshProjectionLODs.Num() && IsHairStrandsEnabled(EHairStrandsShaderType::Strands))
+			if (Resource.RenRootResources && InSkeletalMesh->GetLODNum() != Resource.RenRootResources->GetLODCount() && IsHairStrandsEnabled(EHairStrandsShaderType::Strands))
 			{
 				if (bIssueWarning)
 				{
 					UE_LOG(LogHairStrands, Warning, TEXT("[Groom] The Groom binding (%s) does not have the same have LOD count (LOD render:%d) than the skeletal mesh (%s, LOD:%d). The binding asset will not be used."),
 						*InBinding->GetName(),
-						Resource.RenRootResources->BulkData.MeshProjectionLODs.Num(),
+						Resource.RenRootResources->GetLODCount(),
 						*InSkeletalMesh->GetName(),
 						InSkeletalMesh->GetLODNum());
 				}
@@ -784,10 +784,10 @@ void UGroomBindingAsset::CacheDerivedDatas()
 			FGoomBindingGroupInfo& Info = GroupInfos[GroupIt];
 			const FHairGroupBulkData& BulkData = HairGroupBulkDatas[GroupIt];
 			{
-				Info.SimRootCount = BulkData.SimRootBulkData.RootCount;
-				Info.SimLODCount  = BulkData.SimRootBulkData.MeshProjectionLODs.Num();
-				Info.RenRootCount = BulkData.RenRootBulkData.RootCount;
-				Info.RenLODCount  = BulkData.RenRootBulkData.MeshProjectionLODs.Num();
+				Info.SimRootCount = BulkData.SimRootBulkData.GetRootCount();
+				Info.SimLODCount  = BulkData.SimRootBulkData.GetLODCount();
+				Info.RenRootCount = BulkData.RenRootBulkData.GetRootCount();
+				Info.RenLODCount  = BulkData.RenRootBulkData.GetLODCount();
 			}
 		}
 
diff --git a/Engine/Plugins/Runtime/HairStrands/Source/HairStrandsCore/Private/GroomBindingBuilder.cpp b/Engine/Plugins/Runtime/HairStrands/Source/HairStrandsCore/Private/GroomBindingBuilder.cpp
--- a/Engine/Plugins/Runtime/HairStrands/Source/HairStrandsCore/Private/GroomBindingBuilder.cpp
+++ b/Engine/Plugins/Runtime/HairStrands/Source/HairStrandsCore/Private/GroomBindingBuilder.cpp
@@ -57,7 +57,7 @@ static FAutoConsoleVariableRef CVarHairStrandsBindingBuilderWarningEnable(TEXT("
 FString FGroomBindingBuilder::GetVersion()
 {
 	// Important to update the version when groom building changes
-	return TEXT("2h");
+	return TEXT("3a");
 }
 
 ///////////////////////////////////////////////////////////////////////////////////////////////////
@@ -1876,6 +1876,12 @@ void CopyToBulkData(FByteBulkData& Out, const TArray<typename TFormatType::Type>
 	Out.Unlock();
 }
 
+template<typename TFormatType>
+void CopyToBulkData(FHairBulkContainer& Out, const TArray<typename TFormatType::Type>& Data)
+{
+	CopyToBulkData<TFormatType>(Out.Data, Data);
+}
+
 template<typename TFormatType>
 void CopyFromBulkData(TArray<typename TFormatType::Type>& Out, const FByteBulkData& In)
 {
@@ -1889,16 +1895,34 @@ void CopyFromBulkData(TArray<typename TFormatType::Type>& Out, const FByteBulkDa
 	In.Unlock();
 }
 
+template<typename TFormatType>
+void CopyFromBulkData(TArray<typename TFormatType::Type>& Out, const FHairBulkContainer& In)
+{
+	CopyFromBulkData<TFormatType>(Out, In.Data);
+}
+
+
 // Convert "root data" -> "root bulk data"
 static void BuildRootBulkData(
 	FHairStrandsRootBulkData& Out,
 	const FHairStrandsRootData& In)
 {
-	Out.RootCount = In.RootCount;
-	Out.PointCount = In.PointCount;
-
 	const uint32 MeshLODCount = In.MeshProjectionLODs.Num();
-	Out.MeshProjectionLODs.SetNum(MeshLODCount);
+
+	// Header
+	Out.Header.RootCount  = In.RootCount;
+	Out.Header.PointCount = In.PointCount;
+
+	Out.Header.Strides.RootToUniqueTriangleIndexBufferStride 	= FHairStrandsRootToUniqueTriangleIndexFormat::SizeInByte;
+	Out.Header.Strides.RootBarycentricBufferStride 				= FHairStrandsRootBarycentricFormat::SizeInByte;
+	Out.Header.Strides.UniqueTriangleIndexBufferStride 			= FHairStrandsUniqueTriangleIndexFormat::SizeInByte;
+	Out.Header.Strides.RestUniqueTrianglePositionBufferStride 	= FHairStrandsMeshTrianglePositionFormat::SizeInByte;
+
+	Out.Header.Strides.MeshInterpolationWeightsBufferStride 	= FHairStrandsWeightFormat::SizeInByte;
+	Out.Header.Strides.MeshSampleIndicesBufferStride 			= FHairStrandsIndexFormat::SizeInByte;
+	Out.Header.Strides.RestSamplePositionsBufferStride 			= FHairStrandsMeshTrianglePositionFormat::SizeInByte;
+
+	Out.Header.LODs.SetNum(MeshLODCount);
 	for (uint32 MeshLODIt = 0; MeshLODIt < MeshLODCount; ++MeshLODIt)
 	{
 		const bool bHasValidSamples =
@@ -1906,28 +1930,37 @@ static void BuildRootBulkData(
 			In.MeshProjectionLODs[MeshLODIt].MeshSampleIndicesBuffer.Num() > 0 &&
 			In.MeshProjectionLODs[MeshLODIt].RestSamplePositionsBuffer.Num() > 0;
 
-		Out.MeshProjectionLODs[MeshLODIt].LODIndex = In.MeshProjectionLODs[MeshLODIt].LODIndex;
-		Out.MeshProjectionLODs[MeshLODIt].SampleCount = bHasValidSamples ? In.MeshProjectionLODs[MeshLODIt].SampleCount : 0u;
-		Out.MeshProjectionLODs[MeshLODIt].UniqueTriangleCount = In.MeshProjectionLODs[MeshLODIt].UniqueTriangleIndexBuffer.Num();
+		Out.Header.LODs[MeshLODIt].LODIndex 			= In.MeshProjectionLODs[MeshLODIt].LODIndex;
+		Out.Header.LODs[MeshLODIt].SampleCount 			= bHasValidSamples ? In.MeshProjectionLODs[MeshLODIt].SampleCount : 0u;
+		Out.Header.LODs[MeshLODIt].UniqueTriangleCount 	= In.MeshProjectionLODs[MeshLODIt].UniqueTriangleIndexBuffer.Num();
+		Out.Header.LODs[MeshLODIt].UniqueSectionIndices = In.MeshProjectionLODs[MeshLODIt].UniqueSectionIds;
+	}
 
-		CopyToBulkData<FHairStrandsUniqueTriangleIndexFormat>(Out.MeshProjectionLODs[MeshLODIt].UniqueTriangleIndexBuffer, In.MeshProjectionLODs[MeshLODIt].UniqueTriangleIndexBuffer);
-		CopyToBulkData<FHairStrandsRootBarycentricFormat>(Out.MeshProjectionLODs[MeshLODIt].RootBarycentricBuffer, In.MeshProjectionLODs[MeshLODIt].RootBarycentricBuffer);
-		CopyToBulkData<FHairStrandsRootToUniqueTriangleIndexFormat>(Out.MeshProjectionLODs[MeshLODIt].RootToUniqueTriangleIndexBuffer, In.MeshProjectionLODs[MeshLODIt].RootToUniqueTriangleIndexBuffer);
-		CopyToBulkData<FHairStrandsMeshTrianglePositionFormat>(Out.MeshProjectionLODs[MeshLODIt].RestUniqueTrianglePositionBuffer, In.MeshProjectionLODs[MeshLODIt].RestUniqueTrianglePositionBuffer);
+	// Data
+	Out.Data.LODs.SetNum(MeshLODCount);
+	for (uint32 MeshLODIt = 0; MeshLODIt < MeshLODCount; ++MeshLODIt)
+	{
+		const bool bHasValidSamples =
+			In.MeshProjectionLODs[MeshLODIt].MeshInterpolationWeightsBuffer.Num() > 0 &&
+			In.MeshProjectionLODs[MeshLODIt].MeshSampleIndicesBuffer.Num() > 0 &&
+			In.MeshProjectionLODs[MeshLODIt].RestSamplePositionsBuffer.Num() > 0;
 
-		Out.MeshProjectionLODs[MeshLODIt].UniqueSectionIndices = In.MeshProjectionLODs[MeshLODIt].UniqueSectionIds;
+		CopyToBulkData<FHairStrandsUniqueTriangleIndexFormat>(Out.Data.LODs[MeshLODIt].UniqueTriangleIndexBuffer, In.MeshProjectionLODs[MeshLODIt].UniqueTriangleIndexBuffer);
+		CopyToBulkData<FHairStrandsRootBarycentricFormat>(Out.Data.LODs[MeshLODIt].RootBarycentricBuffer, In.MeshProjectionLODs[MeshLODIt].RootBarycentricBuffer);
+		CopyToBulkData<FHairStrandsRootToUniqueTriangleIndexFormat>(Out.Data.LODs[MeshLODIt].RootToUniqueTriangleIndexBuffer, In.MeshProjectionLODs[MeshLODIt].RootToUniqueTriangleIndexBuffer);
+		CopyToBulkData<FHairStrandsMeshTrianglePositionFormat>(Out.Data.LODs[MeshLODIt].RestUniqueTrianglePositionBuffer, In.MeshProjectionLODs[MeshLODIt].RestUniqueTrianglePositionBuffer);
 
 		if (bHasValidSamples)
 		{
-			CopyToBulkData<FHairStrandsWeightFormat>(Out.MeshProjectionLODs[MeshLODIt].MeshInterpolationWeightsBuffer, In.MeshProjectionLODs[MeshLODIt].MeshInterpolationWeightsBuffer);
-			CopyToBulkData<FHairStrandsIndexFormat>(Out.MeshProjectionLODs[MeshLODIt].MeshSampleIndicesBuffer, In.MeshProjectionLODs[MeshLODIt].MeshSampleIndicesBuffer);
-			CopyToBulkData<FHairStrandsMeshTrianglePositionFormat>(Out.MeshProjectionLODs[MeshLODIt].RestSamplePositionsBuffer, In.MeshProjectionLODs[MeshLODIt].RestSamplePositionsBuffer);
+			CopyToBulkData<FHairStrandsWeightFormat>(Out.Data.LODs[MeshLODIt].MeshInterpolationWeightsBuffer, In.MeshProjectionLODs[MeshLODIt].MeshInterpolationWeightsBuffer);
+			CopyToBulkData<FHairStrandsIndexFormat>(Out.Data.LODs[MeshLODIt].MeshSampleIndicesBuffer, In.MeshProjectionLODs[MeshLODIt].MeshSampleIndicesBuffer);
+			CopyToBulkData<FHairStrandsMeshTrianglePositionFormat>(Out.Data.LODs[MeshLODIt].RestSamplePositionsBuffer, In.MeshProjectionLODs[MeshLODIt].RestSamplePositionsBuffer);
 		}
 		else
 		{
-			Out.MeshProjectionLODs[MeshLODIt].MeshInterpolationWeightsBuffer.RemoveBulkData();
-			Out.MeshProjectionLODs[MeshLODIt].MeshSampleIndicesBuffer.RemoveBulkData();
-			Out.MeshProjectionLODs[MeshLODIt].RestSamplePositionsBuffer.RemoveBulkData();
+			Out.Data.LODs[MeshLODIt].MeshInterpolationWeightsBuffer.RemoveBulkData();
+			Out.Data.LODs[MeshLODIt].MeshSampleIndicesBuffer.RemoveBulkData();
+			Out.Data.LODs[MeshLODIt].RestSamplePositionsBuffer.RemoveBulkData();
 		}
 	}
 }
@@ -1937,30 +1970,31 @@ static void BuildRootData(
 	FHairStrandsRootData& Out,
 	const FHairStrandsRootBulkData& In)
 {
-	Out.RootCount = In.RootCount;
-	Out.PointCount = In.PointCount;
+	// TODO: do we need to force load the data into the bulk data prior to running the convertion (bulk -> data)?
 
-	const uint32 MeshLODCount = In.MeshProjectionLODs.Num();
+	Out.RootCount = In.Header.RootCount;
+	Out.PointCount = In.Header.PointCount;
+
+	const uint32 MeshLODCount = In.Header.LODs.Num();
 	Out.MeshProjectionLODs.SetNum(MeshLODCount);
 	for (uint32 MeshLODIt = 0; MeshLODIt < MeshLODCount; ++MeshLODIt)
 	{
-		const bool bHasValidSamples = In.MeshProjectionLODs[MeshLODIt].SampleCount > 0;
-
-		Out.MeshProjectionLODs[MeshLODIt].LODIndex = In.MeshProjectionLODs[MeshLODIt].LODIndex;
-		Out.MeshProjectionLODs[MeshLODIt].SampleCount = bHasValidSamples ? In.MeshProjectionLODs[MeshLODIt].SampleCount : 0u;
+		const bool bHasValidSamples = In.Header.LODs[MeshLODIt].SampleCount > 0;
 
-		CopyFromBulkData<FHairStrandsUniqueTriangleIndexFormat>(Out.MeshProjectionLODs[MeshLODIt].UniqueTriangleIndexBuffer, In.MeshProjectionLODs[MeshLODIt].UniqueTriangleIndexBuffer);
-		CopyFromBulkData<FHairStrandsRootToUniqueTriangleIndexFormat>(Out.MeshProjectionLODs[MeshLODIt].RootToUniqueTriangleIndexBuffer, In.MeshProjectionLODs[MeshLODIt].RootToUniqueTriangleIndexBuffer);
-		CopyFromBulkData<FHairStrandsRootBarycentricFormat>(Out.MeshProjectionLODs[MeshLODIt].RootBarycentricBuffer, In.MeshProjectionLODs[MeshLODIt].RootBarycentricBuffer);
-		CopyFromBulkData<FHairStrandsMeshTrianglePositionFormat>(Out.MeshProjectionLODs[MeshLODIt].RestUniqueTrianglePositionBuffer, In.MeshProjectionLODs[MeshLODIt].RestUniqueTrianglePositionBuffer);
+		Out.MeshProjectionLODs[MeshLODIt].LODIndex = In.Header.LODs[MeshLODIt].LODIndex;
+		Out.MeshProjectionLODs[MeshLODIt].SampleCount = bHasValidSamples ? In.Header.LODs[MeshLODIt].SampleCount : 0u;
+		Out.MeshProjectionLODs[MeshLODIt].UniqueSectionIds = In.Header.LODs[MeshLODIt].UniqueSectionIndices;
 
-		Out.MeshProjectionLODs[MeshLODIt].UniqueSectionIds = In.MeshProjectionLODs[MeshLODIt].UniqueSectionIndices;
+		CopyFromBulkData<FHairStrandsUniqueTriangleIndexFormat>(Out.MeshProjectionLODs[MeshLODIt].UniqueTriangleIndexBuffer, In.Data.LODs[MeshLODIt].UniqueTriangleIndexBuffer);
+		CopyFromBulkData<FHairStrandsRootToUniqueTriangleIndexFormat>(Out.MeshProjectionLODs[MeshLODIt].RootToUniqueTriangleIndexBuffer, In.Data.LODs[MeshLODIt].RootToUniqueTriangleIndexBuffer);
+		CopyFromBulkData<FHairStrandsRootBarycentricFormat>(Out.MeshProjectionLODs[MeshLODIt].RootBarycentricBuffer, In.Data.LODs[MeshLODIt].RootBarycentricBuffer);
+		CopyFromBulkData<FHairStrandsMeshTrianglePositionFormat>(Out.MeshProjectionLODs[MeshLODIt].RestUniqueTrianglePositionBuffer, In.Data.LODs[MeshLODIt].RestUniqueTrianglePositionBuffer);
 
 		if (bHasValidSamples)
 		{
-			CopyFromBulkData<FHairStrandsWeightFormat>(Out.MeshProjectionLODs[MeshLODIt].MeshInterpolationWeightsBuffer, In.MeshProjectionLODs[MeshLODIt].MeshInterpolationWeightsBuffer);
-			CopyFromBulkData<FHairStrandsIndexFormat>(Out.MeshProjectionLODs[MeshLODIt].MeshSampleIndicesBuffer, In.MeshProjectionLODs[MeshLODIt].MeshSampleIndicesBuffer);
-			CopyFromBulkData<FHairStrandsMeshTrianglePositionFormat>(Out.MeshProjectionLODs[MeshLODIt].RestSamplePositionsBuffer, In.MeshProjectionLODs[MeshLODIt].RestSamplePositionsBuffer);
+			CopyFromBulkData<FHairStrandsWeightFormat>(Out.MeshProjectionLODs[MeshLODIt].MeshInterpolationWeightsBuffer, In.Data.LODs[MeshLODIt].MeshInterpolationWeightsBuffer);
+			CopyFromBulkData<FHairStrandsIndexFormat>(Out.MeshProjectionLODs[MeshLODIt].MeshSampleIndicesBuffer, In.Data.LODs[MeshLODIt].MeshSampleIndicesBuffer);
+			CopyFromBulkData<FHairStrandsMeshTrianglePositionFormat>(Out.MeshProjectionLODs[MeshLODIt].RestSamplePositionsBuffer, In.Data.LODs[MeshLODIt].RestSamplePositionsBuffer);
 		}
 		else
 		{
diff --git a/Engine/Plugins/Runtime/HairStrands/Source/HairStrandsCore/Private/GroomComponent.cpp b/Engine/Plugins/Runtime/HairStrands/Source/HairStrandsCore/Private/GroomComponent.cpp
--- a/Engine/Plugins/Runtime/HairStrands/Source/HairStrandsCore/Private/GroomComponent.cpp
+++ b/Engine/Plugins/Runtime/HairStrands/Source/HairStrandsCore/Private/GroomComponent.cpp
@@ -2245,7 +2245,7 @@ static USkeletalMeshComponent* ValidateBindingAsset(
 					GroupIt < BindingAsset->HairGroupResources.Num() &&
 					CardsLODIndex < uint32(BindingAsset->HairGroupResources[GroupIt].CardsRootResources.Num()) &&
 					BindingAsset->HairGroupResources[GroupIt].CardsRootResources[CardsLODIndex] != nullptr &&
-					((SkeletalMeshComponent && SkeletalMeshComponent->GetSkeletalMeshAsset()) ? SkeletalMeshComponent->GetSkeletalMeshAsset()->GetLODInfoArray().Num() == BindingAsset->HairGroupResources[GroupIt].CardsRootResources[CardsLODIndex]->BulkData.MeshProjectionLODs.Num() : false);
+					((SkeletalMeshComponent && SkeletalMeshComponent->GetSkeletalMeshAsset()) ? SkeletalMeshComponent->GetSkeletalMeshAsset()->GetLODInfoArray().Num() == BindingAsset->HairGroupResources[GroupIt].CardsRootResources[CardsLODIndex]->GetLODCount() : false);
 
 				if (!bIsCardsBindingCompatible)
 				{
@@ -2613,7 +2613,7 @@ void UGroomComponent::InitResources(bool bIsBindingReloading)
 				check(GroupIt < LocalBindingAsset->HairGroupResources.Num());
 				if (USkeletalMeshComponent* SkeletalMeshComponent = Cast<USkeletalMeshComponent>(RegisteredMeshComponent))
 				{
-					check(SkeletalMeshComponent->GetSkeletalMeshAsset() ? SkeletalMeshComponent->GetSkeletalMeshAsset()->GetLODInfoArray().Num() == LocalBindingAsset->HairGroupResources[GroupIt].SimRootResources->BulkData.MeshProjectionLODs.Num() : false);
+					check(SkeletalMeshComponent->GetSkeletalMeshAsset() ? SkeletalMeshComponent->GetSkeletalMeshAsset()->GetLODInfoArray().Num() == LocalBindingAsset->HairGroupResources[GroupIt].SimRootResources->GetLODCount() : false);
 				}
 
 				HairGroupInstance->Guides.RestRootResource = LocalBindingAsset->HairGroupResources[GroupIt].SimRootResources;
@@ -2737,7 +2737,7 @@ void UGroomComponent::InitResources(bool bIsBindingReloading)
 				check(GroupIt < LocalBindingAsset->HairGroupResources.Num());
 				if (USkeletalMeshComponent* SkeletalMeshComponent = Cast<USkeletalMeshComponent>(RegisteredMeshComponent))
 				{
-					check(SkeletalMeshComponent->GetSkeletalMeshAsset() ? SkeletalMeshComponent->GetSkeletalMeshAsset()->GetLODInfoArray().Num() == LocalBindingAsset->HairGroupResources[GroupIt].RenRootResources->BulkData.MeshProjectionLODs.Num() : false);
+					check(SkeletalMeshComponent->GetSkeletalMeshAsset() ? SkeletalMeshComponent->GetSkeletalMeshAsset()->GetLODInfoArray().Num() == LocalBindingAsset->HairGroupResources[GroupIt].RenRootResources->GetLODCount() : false);
 				}
 
 				HairGroupInstance->Strands.RestRootResource = LocalBindingAsset->HairGroupResources[GroupIt].RenRootResources;
@@ -2842,7 +2842,7 @@ void UGroomComponent::InitResources(bool bIsBindingReloading)
 						check(GroupIt < LocalBindingAsset->HairGroupResources.Num());
 						if (USkeletalMeshComponent* SkeletalMeshComponent = Cast<USkeletalMeshComponent>(RegisteredMeshComponent))
 						{
-							check(SkeletalMeshComponent->GetSkeletalMeshAsset() ? SkeletalMeshComponent->GetSkeletalMeshAsset()->GetLODInfoArray().Num() == LocalBindingAsset->HairGroupResources[GroupIt].CardsRootResources[CardsLODIndex]->BulkData.MeshProjectionLODs.Num() : false);
+							check(SkeletalMeshComponent->GetSkeletalMeshAsset() ? SkeletalMeshComponent->GetSkeletalMeshAsset()->GetLODInfoArray().Num() == LocalBindingAsset->HairGroupResources[GroupIt].CardsRootResources[CardsLODIndex]->GetLODCount() : false);
 						}
 
 						InstanceLOD.Guides.RestRootResource = LocalBindingAsset->HairGroupResources[GroupIt].CardsRootResources[CardsLODIndex];
diff --git a/Engine/Plugins/Runtime/HairStrands/Source/HairStrandsCore/Private/GroomDebug.cpp b/Engine/Plugins/Runtime/HairStrands/Source/HairStrandsCore/Private/GroomDebug.cpp
--- a/Engine/Plugins/Runtime/HairStrands/Source/HairStrandsCore/Private/GroomDebug.cpp
+++ b/Engine/Plugins/Runtime/HairStrands/Source/HairStrandsCore/Private/GroomDebug.cpp
@@ -296,7 +296,7 @@ static void AddDebugProjectionHairPass(
 	}
 	
 	const EPrimitiveType PrimitiveType = GeometryType == EDebugProjectionHairType::HairFrame ? PT_LineList : GeometryType == EDebugProjectionHairType::HairTriangle ? PT_TriangleList : PT_LineList;
-	const uint32 RootCount = EDebugProjectionHairType::HairSamples == GeometryType ? 3 * RestRootResources->LODs[MeshLODIndex].SampleCount : RestRootResources->BulkData.RootCount;
+	const uint32 RootCount = EDebugProjectionHairType::HairSamples == GeometryType ? 3 * RestRootResources->LODs[MeshLODIndex].SampleCount : RestRootResources->GetRootCount();
 	const uint32 PrimitiveCount = RootCount;
 
 	if (PrimitiveCount == 0)
@@ -939,10 +939,10 @@ static void AddHairDebugPrintInstancePass(
 				const int32 MeshLODIndex = Instance->HairGroupPublicData->MeshLODIndex;
 				if (MeshLODIndex>=0 && Instance->Strands.RestRootResource)
 				{
-					Data1.X = Instance->Strands.RestRootResource->BulkData.MeshProjectionLODs[MeshLODIndex].UniqueSectionIndices.Num();
-					Data1.Y = Instance->Strands.RestRootResource->BulkData.MeshProjectionLODs[MeshLODIndex].UniqueTriangleCount;
-					Data1.Z = Instance->Strands.RestRootResource->BulkData.RootCount;
-					Data1.W = Instance->Strands.RestRootResource->BulkData.PointCount;
+					Data1.X = Instance->Strands.RestRootResource->BulkData.Header.LODs[MeshLODIndex].UniqueSectionIndices.Num();
+					Data1.Y = Instance->Strands.RestRootResource->BulkData.Header.LODs[MeshLODIndex].UniqueTriangleCount;
+					Data1.Z = Instance->Strands.RestRootResource->BulkData.Header.RootCount;
+					Data1.W = Instance->Strands.RestRootResource->BulkData.Header.PointCount;
 				}
 
 				{
diff --git a/Engine/Plugins/Runtime/HairStrands/Source/HairStrandsCore/Private/GroomResources.cpp b/Engine/Plugins/Runtime/HairStrands/Source/HairStrandsCore/Private/GroomResources.cpp
--- a/Engine/Plugins/Runtime/HairStrands/Source/HairStrandsCore/Private/GroomResources.cpp
+++ b/Engine/Plugins/Runtime/HairStrands/Source/HairStrandsCore/Private/GroomResources.cpp
@@ -1337,30 +1337,30 @@ FHairStrandsRestRootResource::FHairStrandsRestRootResource(FHairStrandsRootBulkD
 void FHairStrandsRestRootResource::PopulateFromRootData()
 {
 	uint32 LODIndex = 0;
-	LODs.Reserve(BulkData.MeshProjectionLODs.Num());
-	for (const FHairStrandsRootBulkData::FMeshProjectionLOD& MeshProjectionLOD : BulkData.MeshProjectionLODs)
+	LODs.Reserve(BulkData.Header.LODs.Num());
+	for (const FHairStrandsRootBulkData::FHeader::FLOD& InLOD : BulkData.Header.LODs)
 	{
-		FLOD& LOD = LODs.AddDefaulted_GetRef();
+		FLOD& OutLOD = LODs.AddDefaulted_GetRef();
 
-		LOD.LODIndex = MeshProjectionLOD.LODIndex;
-		LOD.Status = FLOD::EStatus::Invalid;
-		LOD.SampleCount = MeshProjectionLOD.SampleCount;
+		OutLOD.LODIndex 	= InLOD.LODIndex;
+		OutLOD.Status 		= FLOD::EStatus::Invalid;
+		OutLOD.SampleCount 	= InLOD.SampleCount;
 	}
-	LODRequests.SetNum(BulkData.MeshProjectionLODs.Num());
+	LODRequests.SetNum(BulkData.Header.LODs.Num());
 }
 
 bool FHairStrandsRestRootResource::InternalIsLODDataLoaded(int32 LODIndex)
 {
 	bool bIsLoading = false;
 	
-	check(LODs.Num() == BulkData.MeshProjectionLODs.Num());
+	check(LODs.Num() == BulkData.Header.LODs.Num());
 	if (LODIndex >= 0 && LODIndex < LODs.Num())
 	{
 		FBulkDataBatchRequest& LODRequest = LODRequests[LODIndex];
 		if (LODRequest.IsNone())
 		{
 			FBulkDataBatchRequest::FBatchBuilder Batch = FBulkDataBatchRequest::NewBatch(9);
-			FHairStrandsRootBulkData::FMeshProjectionLOD& CPUData = BulkData.MeshProjectionLODs[LODIndex];
+			FHairStrandsRootBulkData::FData::FLOD& CPUData = BulkData.Data.LODs[LODIndex];
 			const bool bHasValidCPUData = CPUData.RootBarycentricBuffer.GetBulkDataSize() > 0;
 			if (bHasValidCPUData)
 			{
@@ -1404,8 +1404,8 @@ void FHairStrandsRestRootResource::InternalAllocateLOD(FRDGBuilder& GraphBuilder
 {
 	// Sanity check ti insure that the 'common' part of FHairStrandsRestRootResource is already initialized
 	check(bIsInitialized);
-	check(BulkData.PointCount > 0);
-	check(LODs.Num() == BulkData.MeshProjectionLODs.Num());
+	check(BulkData.Header.PointCount > 0);
+	check(LODs.Num() == BulkData.Header.LODs.Num());
 	if (LODIndex >= 0 && LODIndex < LODs.Num())
 	{
 		FLOD& GPUData = LODs[LODIndex];
@@ -1417,44 +1417,45 @@ void FHairStrandsRestRootResource::InternalAllocateLOD(FRDGBuilder& GraphBuilder
 		
 		LODRequests[LODIndex] = FBulkDataBatchRequest();
 
-		FHairStrandsRootBulkData::FMeshProjectionLOD& CPUData = BulkData.MeshProjectionLODs[LODIndex];
+		const FHairStrandsRootBulkData::FHeader::FLOD& LODHeader = BulkData.Header.LODs[LODIndex];
+		FHairStrandsRootBulkData::FData::FLOD& CPUData = BulkData.Data.LODs[LODIndex];
 		const bool bHasValidCPUData = CPUData.RootBarycentricBuffer.GetBulkDataSize() > 0;
 		if (bHasValidCPUData)
 		{
 			GPUData.Status = FLOD::EStatus::Completed;
 			
-			InternalCreateVertexBufferRDG_FromBulkData<FHairStrandsRootBarycentricFormat>(GraphBuilder, CPUData.RootBarycentricBuffer, BulkData.RootCount, GPUData.RootBarycentricBuffer, ToHairResourceDebugName(HAIRSTRANDS_RESOUCE_NAME(CurveType, Hair.StrandsRestRoot_RootTriangleBarycentricBuffer), ResourceName), OwnerName, EHairResourceUsageType::Static);
-			InternalCreateVertexBufferRDG_FromBulkData<FHairStrandsRootToUniqueTriangleIndexFormat>(GraphBuilder, CPUData.RootToUniqueTriangleIndexBuffer, BulkData.RootCount, GPUData.RootToUniqueTriangleIndexBuffer, ToHairResourceDebugName(HAIRSTRANDS_RESOUCE_NAME(CurveType, Hair.StrandsRestRoot_RootToUniqueTriangleIndexBuffer), ResourceName), OwnerName, EHairResourceUsageType::Static);
-			InternalCreateVertexBufferRDG_FromBulkData<FHairStrandsUniqueTriangleIndexFormat>(GraphBuilder, CPUData.UniqueTriangleIndexBuffer, CPUData.UniqueTriangleCount, GPUData.UniqueTriangleIndexBuffer, ToHairResourceDebugName(HAIRSTRANDS_RESOUCE_NAME(CurveType, Hair.StrandsRestRoot_UniqueTriangleIndexBuffer), ResourceName), OwnerName, EHairResourceUsageType::Static);
-			InternalCreateVertexBufferRDG_FromBulkData<FHairStrandsMeshTrianglePositionFormat>(GraphBuilder, CPUData.RestUniqueTrianglePositionBuffer, CPUData.UniqueTriangleCount * 3, GPUData.RestUniqueTrianglePositionBuffer, ToHairResourceDebugName(HAIRSTRANDS_RESOUCE_NAME(CurveType, Hair.StrandsRestRoot_RestUniqueTrianglePosition0Buffer), ResourceName), OwnerName, EHairResourceUsageType::Static);
+			InternalCreateVertexBufferRDG_FromBulkData<FHairStrandsRootBarycentricFormat>(GraphBuilder, CPUData.RootBarycentricBuffer, BulkData.Header.RootCount, GPUData.RootBarycentricBuffer, ToHairResourceDebugName(HAIRSTRANDS_RESOUCE_NAME(CurveType, Hair.StrandsRestRoot_RootTriangleBarycentricBuffer), ResourceName), OwnerName, EHairResourceUsageType::Static);
+			InternalCreateVertexBufferRDG_FromBulkData<FHairStrandsRootToUniqueTriangleIndexFormat>(GraphBuilder, CPUData.RootToUniqueTriangleIndexBuffer, BulkData.Header.RootCount, GPUData.RootToUniqueTriangleIndexBuffer, ToHairResourceDebugName(HAIRSTRANDS_RESOUCE_NAME(CurveType, Hair.StrandsRestRoot_RootToUniqueTriangleIndexBuffer), ResourceName), OwnerName, EHairResourceUsageType::Static);
+			InternalCreateVertexBufferRDG_FromBulkData<FHairStrandsUniqueTriangleIndexFormat>(GraphBuilder, CPUData.UniqueTriangleIndexBuffer, LODHeader.UniqueTriangleCount, GPUData.UniqueTriangleIndexBuffer, ToHairResourceDebugName(HAIRSTRANDS_RESOUCE_NAME(CurveType, Hair.StrandsRestRoot_UniqueTriangleIndexBuffer), ResourceName), OwnerName, EHairResourceUsageType::Static);
+			InternalCreateVertexBufferRDG_FromBulkData<FHairStrandsMeshTrianglePositionFormat>(GraphBuilder, CPUData.RestUniqueTrianglePositionBuffer, LODHeader.UniqueTriangleCount * 3, GPUData.RestUniqueTrianglePositionBuffer, ToHairResourceDebugName(HAIRSTRANDS_RESOUCE_NAME(CurveType, Hair.StrandsRestRoot_RestUniqueTrianglePosition0Buffer), ResourceName), OwnerName, EHairResourceUsageType::Static);
 		}
 		else
 		{
 			GPUData.Status = FLOD::EStatus::Initialized;
 			
-			InternalCreateVertexBufferRDG<FHairStrandsRootBarycentricFormat>(GraphBuilder, BulkData.RootCount, GPUData.RootBarycentricBuffer, ToHairResourceDebugName(HAIRSTRANDS_RESOUCE_NAME(CurveType, Hair.StrandsRestRoot_RootBarycentricBuffer), ResourceName), OwnerName, EHairResourceUsageType::Dynamic);
-			InternalCreateVertexBufferRDG<FHairStrandsRootToUniqueTriangleIndexFormat>(GraphBuilder, BulkData.RootCount, GPUData.RootToUniqueTriangleIndexBuffer, ToHairResourceDebugName(HAIRSTRANDS_RESOUCE_NAME(CurveType, Hair.StrandsRestRoot_RootToUniqueTriangleIndexBuffer), ResourceName), OwnerName, EHairResourceUsageType::Dynamic);
-			InternalCreateVertexBufferRDG<FHairStrandsUniqueTriangleIndexFormat>(GraphBuilder, CPUData.UniqueTriangleCount, GPUData.UniqueTriangleIndexBuffer, ToHairResourceDebugName(HAIRSTRANDS_RESOUCE_NAME(CurveType, Hair.StrandsRestRoot_UniqueTriangleIndexBuffer), ResourceName), OwnerName, EHairResourceUsageType::Dynamic);
+			InternalCreateVertexBufferRDG<FHairStrandsRootBarycentricFormat>(GraphBuilder, BulkData.Header.RootCount, GPUData.RootBarycentricBuffer, ToHairResourceDebugName(HAIRSTRANDS_RESOUCE_NAME(CurveType, Hair.StrandsRestRoot_RootBarycentricBuffer), ResourceName), OwnerName, EHairResourceUsageType::Dynamic);
+			InternalCreateVertexBufferRDG<FHairStrandsRootToUniqueTriangleIndexFormat>(GraphBuilder, BulkData.Header.RootCount, GPUData.RootToUniqueTriangleIndexBuffer, ToHairResourceDebugName(HAIRSTRANDS_RESOUCE_NAME(CurveType, Hair.StrandsRestRoot_RootToUniqueTriangleIndexBuffer), ResourceName), OwnerName, EHairResourceUsageType::Dynamic);
+			InternalCreateVertexBufferRDG<FHairStrandsUniqueTriangleIndexFormat>(GraphBuilder, LODHeader.UniqueTriangleCount, GPUData.UniqueTriangleIndexBuffer, ToHairResourceDebugName(HAIRSTRANDS_RESOUCE_NAME(CurveType, Hair.StrandsRestRoot_UniqueTriangleIndexBuffer), ResourceName), OwnerName, EHairResourceUsageType::Dynamic);
 			
 			// Create buffers. Initialization will be done by render passes
-			InternalCreateVertexBufferRDG<FHairStrandsMeshTrianglePositionFormat>(GraphBuilder, CPUData.UniqueTriangleCount * 3, GPUData.RestUniqueTrianglePositionBuffer, ToHairResourceDebugName(HAIRSTRANDS_RESOUCE_NAME(CurveType, Hair.StrandsRestRoot_RestUniqueTrianglePosition0Buffer), ResourceName), OwnerName, EHairResourceUsageType::Dynamic);
+			InternalCreateVertexBufferRDG<FHairStrandsMeshTrianglePositionFormat>(GraphBuilder, LODHeader.UniqueTriangleCount * 3, GPUData.RestUniqueTrianglePositionBuffer, ToHairResourceDebugName(HAIRSTRANDS_RESOUCE_NAME(CurveType, Hair.StrandsRestRoot_RestUniqueTrianglePosition0Buffer), ResourceName), OwnerName, EHairResourceUsageType::Dynamic);
 		}
 
-		GPUData.SampleCount = CPUData.SampleCount;
+		GPUData.SampleCount = LODHeader.SampleCount;
 		const bool bHasValidCPUWeights = CPUData.MeshSampleIndicesBuffer.GetBulkDataSize() > 0;
 		if (bHasValidCPUWeights)
 		{
 			const uint32 InteroplationWeightCount = CPUData.MeshInterpolationWeightsBuffer.GetBulkDataSize() / sizeof(FHairStrandsWeightFormat::Type);
 			InternalCreateVertexBufferRDG_FromBulkData<FHairStrandsWeightFormat>(GraphBuilder, CPUData.MeshInterpolationWeightsBuffer, InteroplationWeightCount, GPUData.MeshInterpolationWeightsBuffer, ToHairResourceDebugName(HAIRSTRANDS_RESOUCE_NAME(CurveType, Hair.StrandsRestRoot_MeshInterpolationWeightsBuffer), ResourceName), OwnerName, EHairResourceUsageType::Static);
-			InternalCreateVertexBufferRDG_FromBulkData<FHairStrandsIndexFormat>(GraphBuilder, CPUData.MeshSampleIndicesBuffer, CPUData.SampleCount, GPUData.MeshSampleIndicesBuffer, ToHairResourceDebugName(HAIRSTRANDS_RESOUCE_NAME(CurveType, Hair.StrandsRestRoot_MeshSampleIndicesBuffer), ResourceName), OwnerName, EHairResourceUsageType::Static);
-			InternalCreateVertexBufferRDG_FromBulkData<FHairStrandsMeshTrianglePositionFormat>(GraphBuilder, CPUData.RestSamplePositionsBuffer, CPUData.SampleCount, GPUData.RestSamplePositionsBuffer, ToHairResourceDebugName(HAIRSTRANDS_RESOUCE_NAME(CurveType, Hair.StrandsRestRoot_RestSamplePositionsBuffer), ResourceName), OwnerName, EHairResourceUsageType::Static);
+			InternalCreateVertexBufferRDG_FromBulkData<FHairStrandsIndexFormat>(GraphBuilder, CPUData.MeshSampleIndicesBuffer, LODHeader.SampleCount, GPUData.MeshSampleIndicesBuffer, ToHairResourceDebugName(HAIRSTRANDS_RESOUCE_NAME(CurveType, Hair.StrandsRestRoot_MeshSampleIndicesBuffer), ResourceName), OwnerName, EHairResourceUsageType::Static);
+			InternalCreateVertexBufferRDG_FromBulkData<FHairStrandsMeshTrianglePositionFormat>(GraphBuilder, CPUData.RestSamplePositionsBuffer, LODHeader.SampleCount, GPUData.RestSamplePositionsBuffer, ToHairResourceDebugName(HAIRSTRANDS_RESOUCE_NAME(CurveType, Hair.StrandsRestRoot_RestSamplePositionsBuffer), ResourceName), OwnerName, EHairResourceUsageType::Static);
 		}
 		else
 		{
 			// TODO: do not allocate these resources, since they won't be used
-			InternalCreateVertexBufferRDG<FHairStrandsWeightFormat>(GraphBuilder, (CPUData.SampleCount+4) * (CPUData.SampleCount+4), GPUData.MeshInterpolationWeightsBuffer, ToHairResourceDebugName(HAIRSTRANDS_RESOUCE_NAME(CurveType, Hair.StrandsRestRoot_MeshInterpolationWeightsBuffer), ResourceName), OwnerName, EHairResourceUsageType::Dynamic);
-			InternalCreateVertexBufferRDG<FHairStrandsIndexFormat>(GraphBuilder, CPUData.SampleCount, GPUData.MeshSampleIndicesBuffer, ToHairResourceDebugName(HAIRSTRANDS_RESOUCE_NAME(CurveType, Hair.StrandsRestRoot_MeshSampleIndicesBuffer), ResourceName), OwnerName, EHairResourceUsageType::Dynamic);
-			InternalCreateVertexBufferRDG<FHairStrandsMeshTrianglePositionFormat>(GraphBuilder, CPUData.SampleCount, GPUData.RestSamplePositionsBuffer, ToHairResourceDebugName(HAIRSTRANDS_RESOUCE_NAME(CurveType, Hair.StrandsRestRoot_RestSamplePositionsBuffer), ResourceName), OwnerName, EHairResourceUsageType::Dynamic);
+			InternalCreateVertexBufferRDG<FHairStrandsWeightFormat>(GraphBuilder, (LODHeader.SampleCount+4) * (LODHeader.SampleCount+4), GPUData.MeshInterpolationWeightsBuffer, ToHairResourceDebugName(HAIRSTRANDS_RESOUCE_NAME(CurveType, Hair.StrandsRestRoot_MeshInterpolationWeightsBuffer), ResourceName), OwnerName, EHairResourceUsageType::Dynamic);
+			InternalCreateVertexBufferRDG<FHairStrandsIndexFormat>(GraphBuilder, LODHeader.SampleCount, GPUData.MeshSampleIndicesBuffer, ToHairResourceDebugName(HAIRSTRANDS_RESOUCE_NAME(CurveType, Hair.StrandsRestRoot_MeshSampleIndicesBuffer), ResourceName), OwnerName, EHairResourceUsageType::Dynamic);
+			InternalCreateVertexBufferRDG<FHairStrandsMeshTrianglePositionFormat>(GraphBuilder, LODHeader.SampleCount, GPUData.RestSamplePositionsBuffer, ToHairResourceDebugName(HAIRSTRANDS_RESOUCE_NAME(CurveType, Hair.StrandsRestRoot_RestSamplePositionsBuffer), ResourceName), OwnerName, EHairResourceUsageType::Dynamic);
 		}
 	}
 }
@@ -1492,7 +1493,7 @@ FHairStrandsDeformedRootResource::FHairStrandsDeformedRootResource(const FHairSt
 {
 	check(InRestResources);
 	uint32 LODIndex = 0;
-	RootCount = InRestResources->BulkData.RootCount;
+	RootCount = InRestResources->BulkData.Header.RootCount;
 	LODs.Reserve(InRestResources->LODs.Num());
 	for (const FHairStrandsRestRootResource::FLOD& InLOD : InRestResources->LODs)
 	{
@@ -1558,116 +1559,136 @@ void FHairStrandsDeformedRootResource::InternalRelease()
 
 /////////////////////////////////////////////////////////////////////////////////////////
 // Root bulk data
-FHairStrandsRootBulkData::FHairStrandsRootBulkData()
+
+void FHairStrandsRootBulkData::SerializeHeader(FArchive& Ar, UObject* Owner)
 {
+	Ar << Header.RootCount;
+	Ar << Header.PointCount;
 
-}
+	// Strides
+	{
+		Ar << Header.Strides.RootToUniqueTriangleIndexBufferStride;
+		Ar << Header.Strides.RootBarycentricBufferStride;
+		Ar << Header.Strides.UniqueTriangleIndexBufferStride;
+		Ar << Header.Strides.RestUniqueTrianglePositionBufferStride;
 
-bool FHairStrandsRootBulkData::HasProjectionData() const
-{
-	bool bIsValid = MeshProjectionLODs.Num() > 0;
-	bIsValid = RootCount > 0;
-	for (const FMeshProjectionLOD& LOD : MeshProjectionLODs)
+		Ar << Header.Strides.MeshInterpolationWeightsBufferStride;
+		Ar << Header.Strides.MeshSampleIndicesBufferStride;
+		Ar << Header.Strides.RestSamplePositionsBufferStride;
+	}
+
+	// LODS
+	uint32 LODCount = Header.LODs.Num();
+	Ar << LODCount;
+	if (Ar.IsLoading())
 	{
-		bIsValid = bIsValid && LOD.LODIndex > 0;
-		if (!bIsValid) break;
+		Header.LODs.SetNum(LODCount);
+		Data.LODs.SetNum(LODCount);
+	}
+	for (uint32 LODIt = 0; LODIt < LODCount; ++LODIt)
+	{
+		if (Ar.IsLoading())
+		{
+			Header.LODs[LODIt].LODIndex = LODIt;
+		}
+
+		Ar << Header.LODs[LODIt].UniqueTriangleCount;
+		Ar << Header.LODs[LODIt].SampleCount;
+		Ar << Header.LODs[LODIt].UniqueSectionIndices;
 	}
-	return bIsValid;
 }
 
 const TArray<uint32>& FHairStrandsRootBulkData::GetValidSectionIndices(int32 LODIndex) const
 {
-	check(LODIndex >= 0 && LODIndex < MeshProjectionLODs.Num());
-	return MeshProjectionLODs[LODIndex].UniqueSectionIndices;
+	check(LODIndex >= 0 && LODIndex < Header.LODs.Num());
+	return Header.LODs[LODIndex].UniqueSectionIndices;
 }
 
-static void InternalSerialize(FArchive& Ar, UObject* Owner, FHairStrandsRootBulkData::FMeshProjectionLOD& LOD)
+void FHairStrandsRootBulkData::SerializeData(FArchive& Ar, UObject* Owner)
 {
-	const int32 ChunkIndex = 0;
-	bool bAttemptFileMapping = false;
+	check(Header.LODs.Num() > 0);
+	check(Header.LODs.Num() == Data.LODs.Num());
 
-	if (Ar.IsSaving())
+	for (uint32 LODIt = 0, LODCount=Header.LODs.Num(); LODIt < LODCount; ++LODIt)
 	{
-		const uint32 BulkFlags = BULKDATA_Force_NOT_InlinePayload;
-		LOD.RootBarycentricBuffer.SetBulkDataFlags(BulkFlags);
-		LOD.RootToUniqueTriangleIndexBuffer.SetBulkDataFlags(BulkFlags);
-		LOD.UniqueTriangleIndexBuffer.SetBulkDataFlags(BulkFlags);
-		LOD.RestUniqueTrianglePositionBuffer.SetBulkDataFlags(BulkFlags);
-
-		LOD.MeshInterpolationWeightsBuffer.SetBulkDataFlags(BulkFlags);
-		LOD.MeshSampleIndicesBuffer.SetBulkDataFlags(BulkFlags);
-		LOD.RestSamplePositionsBuffer.SetBulkDataFlags(BulkFlags);
-	}
+		const FHeader::FLOD& HeaderLOD = Header.LODs[LODIt];
+		FData::FLOD& DataLOD = Data.LODs[LODIt];
 
-	Ar << LOD.LODIndex;
-	Ar << LOD.UniqueTriangleCount;
-	LOD.RootBarycentricBuffer.Serialize(Ar, Owner, ChunkIndex, bAttemptFileMapping);
-	LOD.RootToUniqueTriangleIndexBuffer.Serialize(Ar, Owner, ChunkIndex, bAttemptFileMapping);
-	LOD.UniqueTriangleIndexBuffer.Serialize(Ar, Owner, ChunkIndex, bAttemptFileMapping);
-	LOD.RestUniqueTrianglePositionBuffer.Serialize(Ar, Owner, ChunkIndex, bAttemptFileMapping);
-
-	Ar << LOD.SampleCount;
-	if (LOD.SampleCount)
-	{
-		LOD.MeshInterpolationWeightsBuffer.Serialize(Ar, Owner, ChunkIndex, bAttemptFileMapping);
-		LOD.MeshSampleIndicesBuffer.Serialize(Ar, Owner, ChunkIndex, bAttemptFileMapping);
-		LOD.RestSamplePositionsBuffer.Serialize(Ar, Owner, ChunkIndex, bAttemptFileMapping);
-	}
-		LOD.UniqueSectionIndices.BulkSerialize(Ar);
-}
+		const int32 ChunkIndex = 0;
+		bool bAttemptFileMapping = false;
+	
+		if (Ar.IsSaving())
+		{
+			const uint32 BulkFlags = BULKDATA_Force_NOT_InlinePayload;
+	
+			DataLOD.RootBarycentricBuffer.SetBulkDataFlags(BulkFlags);
+			DataLOD.RootToUniqueTriangleIndexBuffer.SetBulkDataFlags(BulkFlags);
+			DataLOD.UniqueTriangleIndexBuffer.SetBulkDataFlags(BulkFlags);
+			DataLOD.RestUniqueTrianglePositionBuffer.SetBulkDataFlags(BulkFlags);
+	
+			DataLOD.MeshInterpolationWeightsBuffer.SetBulkDataFlags(BulkFlags);
+			DataLOD.MeshSampleIndicesBuffer.SetBulkDataFlags(BulkFlags);
+			DataLOD.RestSamplePositionsBuffer.SetBulkDataFlags(BulkFlags);
+		}
 
-static void InternalSerialize(FArchive& Ar, UObject* Owner, TArray<FHairStrandsRootBulkData::FMeshProjectionLOD>& LODs)
-{
-	uint32 LODCount = LODs.Num();
-	Ar << LODCount;
-	if (Ar.IsLoading())
-	{
-		LODs.SetNum(LODCount);
-	}
-	for (uint32 LODIt = 0; LODIt < LODCount; ++LODIt)
-	{
-		InternalSerialize(Ar, Owner, LODs[LODIt]);
+		DataLOD.RootBarycentricBuffer.Serialize(Ar, Owner, ChunkIndex, bAttemptFileMapping);
+		DataLOD.RootToUniqueTriangleIndexBuffer.Serialize(Ar, Owner, ChunkIndex, bAttemptFileMapping);
+		DataLOD.UniqueTriangleIndexBuffer.Serialize(Ar, Owner, ChunkIndex, bAttemptFileMapping);
+		DataLOD.RestUniqueTrianglePositionBuffer.Serialize(Ar, Owner, ChunkIndex, bAttemptFileMapping);
+	
+		if (HeaderLOD.SampleCount)
+		{
+			DataLOD.MeshInterpolationWeightsBuffer.Serialize(Ar, Owner, ChunkIndex, bAttemptFileMapping);
+			DataLOD.MeshSampleIndicesBuffer.Serialize(Ar, Owner, ChunkIndex, bAttemptFileMapping);
+			DataLOD.RestSamplePositionsBuffer.Serialize(Ar, Owner, ChunkIndex, bAttemptFileMapping);
+		}
 	}
 }
 
 void FHairStrandsRootBulkData::Serialize(FArchive& Ar, UObject* Owner)
 {
-	const int32 ChunkIndex = 0;
-	bool bAttemptFileMapping = false;
-
-	if (!Ar.IsObjectReferenceCollector())
+	// Root data don't have UObject references
+	if (!Ar.IsLoading() && !Ar.IsSaving() && Ar.IsObjectReferenceCollector())
 	{
-		Ar << RootCount;
-		Ar << PointCount;
-		InternalSerialize(Ar, Owner, MeshProjectionLODs);
+		return;
 	}
+
+	SerializeHeader(Ar, Owner);
+	SerializeData(Ar, Owner);
 }
 
 void FHairStrandsRootBulkData::Reset()
 {
-	RootCount = 0;
-	for (FMeshProjectionLOD& LOD : MeshProjectionLODs)
+	// Header
+	Header.RootCount = 0;
+	Header.PointCount = 0;
+	Header.LODs.Empty();
+
+	// Data
+	for (FData::FLOD& LOD : Data.LODs)
 	{
+		// Binding
 		LOD.RootBarycentricBuffer.RemoveBulkData();
 		LOD.RootToUniqueTriangleIndexBuffer.RemoveBulkData();
 		LOD.UniqueTriangleIndexBuffer.RemoveBulkData();
 		LOD.RestUniqueTrianglePositionBuffer.RemoveBulkData();
-		LOD.MeshInterpolationWeightsBuffer.RemoveBulkData();
-		LOD.MeshSampleIndicesBuffer.RemoveBulkData();
-		LOD.RestSamplePositionsBuffer.RemoveBulkData();
 
 		// Reset the bulk byte buffer to ensure the (serialize) data size is reset to 0
 		LOD.RootBarycentricBuffer 			 = FByteBulkData();
 		LOD.RootToUniqueTriangleIndexBuffer  = FByteBulkData();
 		LOD.UniqueTriangleIndexBuffer		 = FByteBulkData();
 		LOD.RestUniqueTrianglePositionBuffer = FByteBulkData();
+
+		// RBF
+		LOD.MeshInterpolationWeightsBuffer.RemoveBulkData();
+		LOD.MeshSampleIndicesBuffer.RemoveBulkData();
+		LOD.RestSamplePositionsBuffer.RemoveBulkData();
+
 		LOD.MeshInterpolationWeightsBuffer 	 = FByteBulkData();
 		LOD.MeshSampleIndicesBuffer 		 = FByteBulkData();
 		LOD.RestSamplePositionsBuffer 		 = FByteBulkData();
-
-		LOD.UniqueSectionIndices.Empty();
 	}
-	MeshProjectionLODs.Empty();
+	Data.LODs.Empty();
 }
 
 /////////////////////////////////////////////////////////////////////////////////////////
diff --git a/Engine/Plugins/Runtime/HairStrands/Source/HairStrandsCore/Private/HairStrandsInterpolation.cpp b/Engine/Plugins/Runtime/HairStrands/Source/HairStrandsCore/Private/HairStrandsInterpolation.cpp
--- a/Engine/Plugins/Runtime/HairStrands/Source/HairStrandsCore/Private/HairStrandsInterpolation.cpp
+++ b/Engine/Plugins/Runtime/HairStrands/Source/HairStrandsCore/Private/HairStrandsInterpolation.cpp
@@ -391,7 +391,7 @@ static void AddDeformSimHairStrandsPass(
 			Parameters->SimPointToCurveBuffer = SimPointToCurveBuffer;
 		}
 
-		const uint32 RootCount = SimRestRootResources ? SimRestRootResources->BulkData.RootCount : 0;
+		const uint32 RootCount = SimRestRootResources ? SimRestRootResources->GetRootCount() : 0;
 
 		if(BoneBufferSRV)
 		{
@@ -622,7 +622,7 @@ static void AddHairStrandsInterpolationPass(
 
 	const bool bSupportDynamicMesh = 
 		RenRestRootResources &&
-		RenRestRootResources->BulkData.RootCount > 0 && 
+		RenRestRootResources->GetRootCount() > 0 && 
 		MeshLODIndex >= 0 && 
 		MeshLODIndex < RenRestRootResources->LODs.Num() &&
 		MeshLODIndex < RenDeformedRootResources->LODs.Num() &&
@@ -1067,7 +1067,7 @@ static void AddHairCardsDeformationPass(
 	const bool bSupportDynamicMesh =
 		Instance->BindingType == EHairBindingType::Skinning &&
 		RestRootResources &&
-		RestRootResources->BulkData.RootCount > 0 &&
+		RestRootResources->GetRootCount() > 0 &&
 		MeshLODIndex >= 0 &&
 		MeshLODIndex < RestRootResources->LODs.Num() &&
 		MeshLODIndex < DeformedRootResources->LODs.Num() &&
diff --git a/Engine/Plugins/Runtime/HairStrands/Source/HairStrandsCore/Private/HairStrandsMeshProjection.cpp b/Engine/Plugins/Runtime/HairStrands/Source/HairStrandsCore/Private/HairStrandsMeshProjection.cpp
--- a/Engine/Plugins/Runtime/HairStrands/Source/HairStrandsCore/Private/HairStrandsMeshProjection.cpp
+++ b/Engine/Plugins/Runtime/HairStrands/Source/HairStrandsCore/Private/HairStrandsMeshProjection.cpp
@@ -276,7 +276,7 @@ void AddHairStrandUpdateMeshTrianglesPass(
 	FHairStrandsRestRootResource* RestResources,
 	FHairStrandsDeformedRootResource* DeformedResources)
 {
-	const uint32 RootCount = RestResources->BulkData.RootCount;
+	const uint32 RootCount = RestResources->GetRootCount();
 	if (RootCount == 0 || LODIndex < 0)
 	{
 		return;
@@ -1075,7 +1075,7 @@ static void AddFollicleMaskPass(
 	FHairStrandsRestRootResource* RestResources,
 	FRDGTextureRef OutTexture)
 {
-	const uint32 RootCount = RestResources->BulkData.RootCount;
+	const uint32 RootCount = RestResources->GetRootCount();
 	if (LODIndex >= uint32(RestResources->LODs.Num()) || RootCount == 0)
 		return;
 
@@ -1091,7 +1091,7 @@ static void AddFollicleMaskPass(
 	Parameters->RootUVsBuffer = nullptr;
 	Parameters->OutputResolution = OutputResolution;
 	Parameters->MaxRootCount = RootCount;
-	Parameters->MaxUniqueTriangleIndex = RestResources->BulkData.MeshProjectionLODs[LODIndex].UniqueTriangleCount;
+	Parameters->MaxUniqueTriangleIndex = RestResources->BulkData.Header.LODs[LODIndex].UniqueTriangleCount;
 	Parameters->Channel = FMath::Min(Channel, 3u);
 	Parameters->KernelSizeInPixels = FMath::Clamp(KernelSizeInPixels, 2u, 200u);
 	Parameters->RenderTargets[0] = FRenderTargetBinding(OutTexture, bNeedClear ? ERenderTargetLoadAction::EClear : ERenderTargetLoadAction::ELoad, 0);
diff --git a/Engine/Plugins/Runtime/HairStrands/Source/HairStrandsCore/Public/GroomResources.h b/Engine/Plugins/Runtime/HairStrands/Source/HairStrandsCore/Public/GroomResources.h
--- a/Engine/Plugins/Runtime/HairStrands/Source/HairStrandsCore/Public/GroomResources.h
+++ b/Engine/Plugins/Runtime/HairStrands/Source/HairStrandsCore/Public/GroomResources.h
@@ -140,6 +140,10 @@ struct FHairStrandsRestRootResource : public FHairCommonResource
 	/* Populate GPU LOD data from RootData (this function doesn't initialize resources) */
 	void PopulateFromRootData();
 
+	// Accessors
+	uint32 GetLODCount() const { return BulkData.GetLODCount(); }
+	uint32 GetRootCount()const { return BulkData.GetRootCount(); }
+
 	/* Return the memory size for GPU resources */
 	uint32 GetResourcesSize() const 
 	{
diff --git a/Engine/Plugins/Runtime/HairStrands/Source/HairStrandsCore/Public/HairStrandsDatas.h b/Engine/Plugins/Runtime/HairStrands/Source/HairStrandsCore/Public/HairStrandsDatas.h
--- a/Engine/Plugins/Runtime/HairStrands/Source/HairStrandsCore/Public/HairStrandsDatas.h
+++ b/Engine/Plugins/Runtime/HairStrands/Source/HairStrandsCore/Public/HairStrandsDatas.h
@@ -642,17 +642,20 @@ struct HAIRSTRANDSCORE_API FHairStrandsDebugDatas
 /* Bulk data for root resources (GPU resources are stored into FHairStrandsRootResources) */
 struct FHairStrandsRootBulkData
 {
-	FHairStrandsRootBulkData();
+	void SerializeHeader(FArchive& Ar, UObject* Owner);
+	void SerializeData(FArchive& Ar, UObject* Owner);
 	void Serialize(FArchive& Ar, UObject* Owner);
+
 	void Reset();
-	bool HasProjectionData() const;
-	bool IsValid() const { return RootCount > 0; }
+	bool IsValid() const { return Header.RootCount > 0; }
 	const TArray<uint32>& GetValidSectionIndices(int32 LODIndex) const;
+	uint32 GetLODCount() const { return Header.LODs.Num(); }
+	uint32 GetRootCount()const { return Header.RootCount; }
 
 	uint32 GetDataSize() const
 	{
 		uint32 Total = 0;
-		for (const FMeshProjectionLOD& LOD : MeshProjectionLODs)
+		for (const FData::FLOD& LOD : Data.LODs)
 		{
 			Total += LOD.UniqueTriangleIndexBuffer.IsBulkDataLoaded() ?			LOD.UniqueTriangleIndexBuffer.GetBulkDataSize() : 0u;
 			Total += LOD.RootToUniqueTriangleIndexBuffer.IsBulkDataLoaded() ?	LOD.RootToUniqueTriangleIndexBuffer.GetBulkDataSize() : 0u;
@@ -661,52 +664,55 @@ struct FHairStrandsRootBulkData
 			Total += LOD.MeshInterpolationWeightsBuffer.IsBulkDataLoaded() ?	LOD.MeshInterpolationWeightsBuffer.GetBulkDataSize() : 0u;
 			Total += LOD.MeshSampleIndicesBuffer.IsBulkDataLoaded() ?			LOD.MeshSampleIndicesBuffer.GetBulkDataSize() : 0u;
 			Total += LOD.RestSamplePositionsBuffer.IsBulkDataLoaded() ?			LOD.RestSamplePositionsBuffer.GetBulkDataSize() : 0u;
-			Total += LOD.UniqueSectionIndices.GetAllocatedSize();
 		}
 		return Total;
 	}
 
-	struct FMeshProjectionLOD
-	{
-		int32 LODIndex = -1;
-		uint32 UniqueTriangleCount = 0;
-
-		/* Map each root onto the unique triangle Id (per-root) */
-		FByteBulkData RootToUniqueTriangleIndexBuffer;
-
-		/* Root's barycentric (per-root) */
-		FByteBulkData RootBarycentricBuffer;
-
-		/* Unique triangles list from skeleton mesh section IDs and triangle IDs (per-unique-triangle) */
-		FByteBulkData UniqueTriangleIndexBuffer;
-
-		/* Rest triangle positions (per-unique-triangle) */
-		FByteBulkData RestUniqueTrianglePositionBuffer;
-
-		/* Number of samples used for the mesh interpolation */
-		uint32 SampleCount = 0;
-
-		/* Store the hair interpolation weights | Size = SamplesCount * SamplesCount (per-sample)*/
-		FByteBulkData MeshInterpolationWeightsBuffer;
-
-		/* Store the samples vertex indices (per-sample) */
-		FByteBulkData MeshSampleIndicesBuffer;
-
-		/* Store the samples rest positions (per-sample) */
-		FByteBulkData RestSamplePositionsBuffer;
-
-		/* Store the mesh section indices which are relevant for this root LOD data */
-		TArray<uint32> UniqueSectionIndices;
-	};
-
-	/* Number of roots */
-	uint32 RootCount = 0;
+	struct FHeader
+	{		
+		struct FLOD
+		{
+			int32  LODIndex = -1;
+			uint32 UniqueTriangleCount = 0;			
+			uint32 SampleCount = 0; 				// Number of samples used for the mesh interpolation
+			TArray<uint32> UniqueSectionIndices; 	// Store the mesh section indices which are relevant for this root LOD data
+		};
 
-	/* Number of control points */
-	uint32 PointCount = 0;
+		struct FStrides
+		{
+			uint32 RootToUniqueTriangleIndexBufferStride = 0;
+			uint32 RootBarycentricBufferStride = 0;
+			uint32 UniqueTriangleIndexBufferStride = 0;
+			uint32 RestUniqueTrianglePositionBufferStride = 0;
+
+			uint32 MeshInterpolationWeightsBufferStride = 0;
+			uint32 MeshSampleIndicesBufferStride = 0;
+			uint32 RestSamplePositionsBufferStride = 0;
+		};
+
+		uint32 RootCount = 0;						// Number of roots
+		uint32 PointCount = 0;						// Number of control points
+		FStrides Strides;
+		TArray<FLOD> LODs;
+	} Header;
 
-	/* Store the hair projection information for each mesh LOD */
-	TArray<FMeshProjectionLOD> MeshProjectionLODs;
+	struct FData
+	{
+		struct FLOD
+		{
+			// Binding
+			FByteBulkData RootToUniqueTriangleIndexBuffer; 	// Map each root onto the unique triangle Id (per-root)
+			FByteBulkData RootBarycentricBuffer; 			// Root's barycentric (per-root)
+			FByteBulkData UniqueTriangleIndexBuffer; 		// Unique triangles list from skeleton mesh section IDs and triangle IDs (per-unique-triangle)
+			FByteBulkData RestUniqueTrianglePositionBuffer;	// Rest triangle positions (per-unique-triangle)
+
+			// RBF
+			FByteBulkData MeshInterpolationWeightsBuffer; 	// Store the hair interpolation weights | Size = SamplesCount * SamplesCount (per-sample
+			FByteBulkData MeshSampleIndicesBuffer; 			// Store the samples vertex indices (per-sample)
+			FByteBulkData RestSamplePositionsBuffer; 		// Store the samples rest positions (per-sample)
+		};
+		TArray<FLOD> LODs;
+	} Data;
 };
 
 /* Source data for building root bulk data */
