Move groom strands/cards/resources into their respective files.

diff --git a/Engine/Plugins/Runtime/HairStrands/Source/HairStrandsCore/Private/GroomResources.cpp b/Engine/Plugins/Runtime/HairStrands/Source/HairStrandsCore/Private/GroomResources.cpp
--- a/Engine/Plugins/Runtime/HairStrands/Source/HairStrandsCore/Private/GroomResources.cpp
+++ b/Engine/Plugins/Runtime/HairStrands/Source/HairStrandsCore/Private/GroomResources.cpp
@@ -50,6 +50,8 @@ static int32 GHairStrandsDebugVoxel_MaxSegmentPerVoxel = 2048;
 static FAutoConsoleVariableRef CVarHairStrandsDebugVoxel_WorldSize(TEXT("r.HairStrands.DebugData.VoxelSize"), GHairStrandsDebugVoxel_WorldSize, TEXT("Voxel size use for creating debug data."));
 static FAutoConsoleVariableRef CVarHairStrandsDebugVoxel_MaxSegmentPerVoxel(TEXT("r.HairStrands.DebugData.MaxSegmentPerVoxel"), GHairStrandsDebugVoxel_MaxSegmentPerVoxel, TEXT("Max number of segments per Voxel size when creating debug data."));
 
+bool IsHairStrandsContinousLODEnabled();
+
 bool ValidateHairBulkData()
 {
 	return GHairStrandsBulkData_Validation > 0;
@@ -100,144 +102,6 @@ const TCHAR* ToHairResourceDebugName(const TCHAR* In, FHairResourceName& InDebug
 #endif
 }
 
-/////////////////////////////////////////////////////////////////////////////////////////
-// IO/DDC request
-
-bool FHairStreamingRequest::IsNone() const
-{ 
-#if !WITH_EDITORONLY_DATA
-	return IORequest.IsNone();
-#else
-	return Chunks.Num() == 0;
-#endif
-}
-bool FHairStreamingRequest::IsCompleted()
-{ 
-#if !WITH_EDITORONLY_DATA
-	if (IORequest.IsCompleted())
-	{
-		for (FHairStreamingRequest::FChunk& Chunk : Chunks)
-		{
-			Chunk.Status = FHairStreamingRequest::FChunk::Completed;
-			check(Chunk.Container);
-			Chunk.Container->LoadedSize = Chunk.TotalSize;
-		}
-	}
-	return Chunks.Num() == 0 || IORequest.IsCompleted(); 
-#else
-	for (const FHairStreamingRequest::FChunk& Chunk : Chunks)
-	{
-		if (Chunk.Status != FHairStreamingRequest::FChunk::Completed)
-		{
-			return false;
-		}
-	}
-	return true; 
-#endif
-}
-// Request fullfil 2 use cases:
-// * Load DDC data and upload them to GPU
-// * Load DDC data and store them into bulkdata for serialization
-void FHairStreamingRequest::Request(uint32 InRequestedCurveCount, uint32 InRequestedPointCount, FHairStrandsBulkCommon& In, bool bWait, bool bFillBulkdata, const FName& InOwnerName)
-{
-	if (In.GetResourceCount() == 0 || InRequestedCurveCount == 0)
-	{
-		CurveCount = 0;
-		PointCount = 0;
-		return;
-	}
-	CurveCount = InRequestedCurveCount;
-	PointCount = InRequestedPointCount;
-
-	#if !WITH_EDITORONLY_DATA
-	{ 
-		In.Read_IO(this, IORequest);
-		if (bWait)
-		{
-			IORequest.Wait();
-		}
-	}
-	#else
-	{
-		// When enabled, data can be loaded from an offset. Otherwisee, start from the beginning of the resource
-		// This is used when cooking data to force the loading of the entire resource (i.e., bSupportOffsetLoad=false)
-		bSupportOffsetLoad = !bFillBulkdata; 
-
-		using namespace UE::DerivedData;
-		TArray<FCacheGetChunkRequest> Requests;
-		In.Read_DDC(this, Requests);
-	
-		check(DDCRequestOwner == nullptr);
-		DDCRequestOwner = MakeUnique<FRequestOwner>(bWait ? UE::DerivedData::EPriority::Blocking : UE::DerivedData::EPriority::Normal); // <= Move this onto the resource (Buffer, in order to ensure no race condition)
-	
-		//FRequestBarrier Barrier(*DDCRequestOwner);	// This is a critical section on the owner. It does not constrain ordering
-		GetCache().GetChunks(Requests, *DDCRequestOwner,
-		[this, &In, bFillBulkdata, InOwnerName](FCacheGetChunkResponse && Response)
-		{
-			if (Response.Status == UE::DerivedData::EStatus::Ok)
-			{
-				FHairStreamingRequest::FChunk& Chunk = *(FHairStreamingRequest::FChunk*)Response.UserData;
-				Chunk.Data_DDC 	= MoveTemp(Response.RawData);
-				Chunk.Offset 	= Response.RawOffset;
-				Chunk.Size 		= Response.RawSize;
-				Chunk.TotalSize = Response.RawOffset + Response.RawSize;
-				Chunk.Status 	= FHairStreamingRequest::FChunk::Completed;
-		
-				// Upload the total amount of loaded data
-				Chunk.Container->LoadedSize = Chunk.TotalSize;
-
-				// Optional fill in of bytebulkdata container
-				if (bFillBulkdata)
-				{
-					check(Chunk.Container);
-					FByteBulkData& BulkData = Chunk.Container->Data;
-	
-					// The buffer is then stored into bulk data
-					BulkData.Lock(LOCK_READ_WRITE);
-					void* DstData = BulkData.Realloc(Chunk.Size);
-					FMemory::Memcpy(DstData, Chunk.GetData(), Chunk.Size);
-					BulkData.Unlock();
-
-					//Chunk.Release();
-				}
-			}
-			else
-			{
-				FHairStreamingRequest::FChunk& Chunk = *(FHairStreamingRequest::FChunk*)Response.UserData;
-				Chunk.Status = FHairStreamingRequest::FChunk::Failed;
-				UE_LOG(LogHairStrands, Error, TEXT("[Groom] DDC request failed for '%s' (Key:%s) "), *InOwnerName.ToString(), *In.DerivedDataKey);
-			}
-		});
-
-		// Optional wait on DDC response
-		if (bWait || bFillBulkdata)
-		{
-			DDCRequestOwner->Wait();
-			check(IsCompleted());
-		}
-	}
-	#endif
-}
-
-const uint8* FHairStreamingRequest::FChunk::GetData() const 
-{ 
-	check(Status == Completed); 
-#if !WITH_EDITORONLY_DATA
-	check(Container);
-	return (const uint8*)Data_IO.GetData();
-#else
-	return (const uint8*)Data_DDC.GetData(); 
-#endif
-}
-void FHairStreamingRequest::FChunk::Release() 
-{ 
-#if !WITH_EDITORONLY_DATA
-	TIoStatusOr<uint8*> Out = Data_IO.Release();
-#else
-	Data_DDC.Reset(); 
-#endif
-}
-
 /////////////////////////////////////////////////////////////////////////////////////////
 // FRWBuffer utils 
 void UploadDataToBuffer(FReadBuffer& OutBuffer, uint32 DataSizeInBytes, const void* InCpuData)
@@ -1127,7 +991,6 @@ bool FHairStrandsRestResource::InternalIsDataLoaded(uint32 InRequestedCurveCount
 	return StreamingRequest.IsCompleted();
 }
 
-bool IsHairStrandsContinousLODEnabled();
 void FHairStrandsRestResource::InternalAllocate(FRDGBuilder& GraphBuilder)
 {
 	// If we enter this function, the request need to be completed
@@ -1557,173 +1420,6 @@ void FHairStrandsDeformedRootResource::InternalRelease()
 	LODs.Empty();
 }
 
-/////////////////////////////////////////////////////////////////////////////////////////
-// Root bulk data
-
-void FHairStrandsRootBulkData::SerializeHeader(FArchive& Ar, UObject* Owner)
-{
-	Ar << Header.RootCount;
-	Ar << Header.PointCount;
-
-	// Strides
-	{
-		Ar << Header.Strides.RootToUniqueTriangleIndexBufferStride;
-		Ar << Header.Strides.RootBarycentricBufferStride;
-		Ar << Header.Strides.UniqueTriangleIndexBufferStride;
-		Ar << Header.Strides.RestUniqueTrianglePositionBufferStride;
-
-		Ar << Header.Strides.MeshInterpolationWeightsBufferStride;
-		Ar << Header.Strides.MeshSampleIndicesBufferStride;
-		Ar << Header.Strides.RestSamplePositionsBufferStride;
-	}
-
-	// LODS
-	uint32 LODCount = Header.LODs.Num();
-	Ar << LODCount;
-	if (Ar.IsLoading())
-	{
-		Header.LODs.SetNum(LODCount);
-		Data.LODs.SetNum(LODCount);
-	}
-	for (uint32 LODIt = 0; LODIt < LODCount; ++LODIt)
-	{
-		if (Ar.IsLoading())
-		{
-			Header.LODs[LODIt].LODIndex = LODIt;
-		}
-
-		Ar << Header.LODs[LODIt].UniqueTriangleCount;
-		Ar << Header.LODs[LODIt].SampleCount;
-		Ar << Header.LODs[LODIt].UniqueSectionIndices;
-	}
-}
-
-const TArray<uint32>& FHairStrandsRootBulkData::GetValidSectionIndices(int32 LODIndex) const
-{
-	check(LODIndex >= 0 && LODIndex < Header.LODs.Num());
-	return Header.LODs[LODIndex].UniqueSectionIndices;
-}
-
-void FHairStrandsRootBulkData::SerializeData(FArchive& Ar, UObject* Owner)
-{
-	check(Header.LODs.Num() > 0);
-	check(Header.LODs.Num() == Data.LODs.Num());
-
-	for (uint32 LODIt = 0, LODCount=Header.LODs.Num(); LODIt < LODCount; ++LODIt)
-	{
-		const FHeader::FLOD& HeaderLOD = Header.LODs[LODIt];
-		FData::FLOD& DataLOD = Data.LODs[LODIt];
-
-		const int32 ChunkIndex = 0;
-		bool bAttemptFileMapping = false;
-	
-		if (Ar.IsSaving())
-		{
-			const uint32 BulkFlags = BULKDATA_Force_NOT_InlinePayload;
-	
-			DataLOD.RootBarycentricBuffer.SetBulkDataFlags(BulkFlags);
-			DataLOD.RootToUniqueTriangleIndexBuffer.SetBulkDataFlags(BulkFlags);
-			DataLOD.UniqueTriangleIndexBuffer.SetBulkDataFlags(BulkFlags);
-			DataLOD.RestUniqueTrianglePositionBuffer.SetBulkDataFlags(BulkFlags);
-	
-			DataLOD.MeshInterpolationWeightsBuffer.SetBulkDataFlags(BulkFlags);
-			DataLOD.MeshSampleIndicesBuffer.SetBulkDataFlags(BulkFlags);
-			DataLOD.RestSamplePositionsBuffer.SetBulkDataFlags(BulkFlags);
-		}
-
-		DataLOD.RootBarycentricBuffer.Serialize(Ar, Owner, ChunkIndex, bAttemptFileMapping);
-		DataLOD.RootToUniqueTriangleIndexBuffer.Serialize(Ar, Owner, ChunkIndex, bAttemptFileMapping);
-		DataLOD.UniqueTriangleIndexBuffer.Serialize(Ar, Owner, ChunkIndex, bAttemptFileMapping);
-		DataLOD.RestUniqueTrianglePositionBuffer.Serialize(Ar, Owner, ChunkIndex, bAttemptFileMapping);
-	
-		if (HeaderLOD.SampleCount)
-		{
-			DataLOD.MeshInterpolationWeightsBuffer.Serialize(Ar, Owner, ChunkIndex, bAttemptFileMapping);
-			DataLOD.MeshSampleIndicesBuffer.Serialize(Ar, Owner, ChunkIndex, bAttemptFileMapping);
-			DataLOD.RestSamplePositionsBuffer.Serialize(Ar, Owner, ChunkIndex, bAttemptFileMapping);
-		}
-	}
-}
-
-void FHairStrandsRootBulkData::Serialize(FArchive& Ar, UObject* Owner)
-{
-	// Root data don't have UObject references
-	if (!Ar.IsLoading() && !Ar.IsSaving() && Ar.IsObjectReferenceCollector())
-	{
-		return;
-	}
-
-	SerializeHeader(Ar, Owner);
-	SerializeData(Ar, Owner);
-}
-
-void FHairStrandsRootBulkData::Reset()
-{
-	// Header
-	Header.RootCount = 0;
-	Header.PointCount = 0;
-	Header.LODs.Empty();
-
-	// Data
-	for (FData::FLOD& LOD : Data.LODs)
-	{
-		// Binding
-		LOD.RootBarycentricBuffer.RemoveBulkData();
-		LOD.RootToUniqueTriangleIndexBuffer.RemoveBulkData();
-		LOD.UniqueTriangleIndexBuffer.RemoveBulkData();
-		LOD.RestUniqueTrianglePositionBuffer.RemoveBulkData();
-
-		// Reset the bulk byte buffer to ensure the (serialize) data size is reset to 0
-		LOD.RootBarycentricBuffer 			 = FByteBulkData();
-		LOD.RootToUniqueTriangleIndexBuffer  = FByteBulkData();
-		LOD.UniqueTriangleIndexBuffer		 = FByteBulkData();
-		LOD.RestUniqueTrianglePositionBuffer = FByteBulkData();
-
-		// RBF
-		LOD.MeshInterpolationWeightsBuffer.RemoveBulkData();
-		LOD.MeshSampleIndicesBuffer.RemoveBulkData();
-		LOD.RestSamplePositionsBuffer.RemoveBulkData();
-
-		LOD.MeshInterpolationWeightsBuffer 	 = FByteBulkData();
-		LOD.MeshSampleIndicesBuffer 		 = FByteBulkData();
-		LOD.RestSamplePositionsBuffer 		 = FByteBulkData();
-	}
-	Data.LODs.Empty();
-}
-
-/////////////////////////////////////////////////////////////////////////////////////////
-// Root data
-FHairStrandsRootData::FHairStrandsRootData()
-{
-
-}
-
-bool FHairStrandsRootData::HasProjectionData() const
-{
-	bool bIsValid = MeshProjectionLODs.Num() > 0;
-	for (const FMeshProjectionLOD& LOD : MeshProjectionLODs)
-	{
-		const bool bHasValidCPUData = LOD.RootBarycentricBuffer.Num() > 0;
-		if (bHasValidCPUData)
-		{
-			bIsValid = bIsValid && LOD.RootBarycentricBuffer.Num() > 0;
-			bIsValid = bIsValid && LOD.RootToUniqueTriangleIndexBuffer.Num() > 0;
-			bIsValid = bIsValid && LOD.UniqueTriangleIndexBuffer.Num() > 0;
-			bIsValid = bIsValid && LOD.RestUniqueTrianglePositionBuffer.Num() > 0;
-
-			if (!bIsValid) break;
-		}
-	}
-
-	return bIsValid;
-}
-
-void FHairStrandsRootData::Reset()
-{
-	RootCount = 0;
-	MeshProjectionLODs.Empty();
-}
-
 /////////////////////////////////////////////////////////////////////////////////////////
 
 FHairStrandsInterpolationResource::FHairStrandsInterpolationResource(FHairStrandsInterpolationBulkData& InBulkData, const FHairResourceName& InResourceName, const FName& InOwnerName) :
@@ -1762,24 +1458,6 @@ void FHairStrandsInterpolationResource::InternalRelease()
 }
 
 /////////////////////////////////////////////////////////////////////////////////////////
-void FHairCardsInterpolationDatas::SetNum(const uint32 NumPoints)
-{
-	PointsSimCurvesIndex.SetNum(NumPoints);
-	PointsSimCurvesVertexIndex.SetNum(NumPoints);
-	PointsSimCurvesVertexLerp.SetNum(NumPoints);
-}
-
-void FHairCardsInterpolationDatas::Reset()
-{
-	PointsSimCurvesIndex.Empty();
-	PointsSimCurvesVertexIndex.Empty();
-	PointsSimCurvesVertexLerp.Empty();
-}
-
-void FHairCardsInterpolationBulkData::Serialize(FArchive& Ar)
-{
-	Ar << Interpolation;
-}
 
 FHairCardsInterpolationResource::FHairCardsInterpolationResource(FHairCardsInterpolationBulkData& InBulkData, const FHairResourceName& InResourceName, const FName& InOwnerName) :
 	FHairCommonResource(EHairStrandsAllocationType::Deferred, InResourceName, InOwnerName),
diff --git a/Engine/Plugins/Runtime/HairStrands/Source/HairStrandsCore/Private/HairCardsDatas.cpp b/Engine/Plugins/Runtime/HairStrands/Source/HairStrandsCore/Private/HairCardsDatas.cpp
--- a/Engine/Plugins/Runtime/HairStrands/Source/HairStrandsCore/Private/HairCardsDatas.cpp
+++ b/Engine/Plugins/Runtime/HairStrands/Source/HairStrandsCore/Private/HairCardsDatas.cpp
@@ -44,3 +44,23 @@ void FHairMeshesBulkData::Serialize(FArchive& Ar)
 	Ar << BoundingBox;
 }
 
+
+/////////////////////////////////////////////////////////////////////////////////////////
+void FHairCardsInterpolationDatas::SetNum(const uint32 NumPoints)
+{
+	PointsSimCurvesIndex.SetNum(NumPoints);
+	PointsSimCurvesVertexIndex.SetNum(NumPoints);
+	PointsSimCurvesVertexLerp.SetNum(NumPoints);
+}
+
+void FHairCardsInterpolationDatas::Reset()
+{
+	PointsSimCurvesIndex.Empty();
+	PointsSimCurvesVertexIndex.Empty();
+	PointsSimCurvesVertexLerp.Empty();
+}
+
+void FHairCardsInterpolationBulkData::Serialize(FArchive& Ar)
+{
+	Ar << Interpolation;
+}
\ No newline at end of file
diff --git a/Engine/Plugins/Runtime/HairStrands/Source/HairStrandsCore/Private/HairStrandsDatas.cpp b/Engine/Plugins/Runtime/HairStrands/Source/HairStrandsCore/Private/HairStrandsDatas.cpp
--- a/Engine/Plugins/Runtime/HairStrands/Source/HairStrandsCore/Private/HairStrandsDatas.cpp
+++ b/Engine/Plugins/Runtime/HairStrands/Source/HairStrandsCore/Private/HairStrandsDatas.cpp
@@ -11,12 +11,41 @@
 #include "DerivedDataRequestOwner.h"
 #endif
 
-void FHairStrandsInterpolationDatas::SetNum(const uint32 NumCurves)
+///////////////////////////////////////////////////////////////////////////////////////////////////
+// Helper
+
+float GetHairStrandsMaxLength(const FHairStrandsDatas& In)
 {
-	PointsSimCurvesVertexWeights.SetNum(NumCurves);
-	PointsSimCurvesVertexLerp.SetNum(NumCurves);
-	PointsSimCurvesVertexIndex.SetNum(NumCurves);
-	PointsSimCurvesIndex.SetNum(NumCurves);
+	float MaxLength = 0;
+	for (float CurveLength : In.StrandsCurves.CurvesLength)
+	{
+		MaxLength = FMath::Max(MaxLength, CurveLength);
+	}
+	return MaxLength;
 }
 
+float GetHairStrandsMaxRadius(const FHairStrandsDatas& In)
+{
+	float MaxRadius = 0;
+	for (float PointRadius : In.StrandsPoints.PointsRadius)
+	{
+		MaxRadius = FMath::Max(MaxRadius, PointRadius);
+	}
+	return MaxRadius;
+}
+
+///////////////////////////////////////////////////////////////////////////////////////////////////
+// Curve/Point/Strands/Interpolation/Culling Data
+
+void FHairStrandsCurves::Reset()
+{
+	CurvesOffset.Reset();
+	CurvesCount.Reset();
+	CurvesLength.Reset();
+	CurvesRootUV.Reset();
+	ClumpIDs.Reset();
+	CurvesClosestGuideIDs.Reset();
+	CurvesClosestGuideWeights.Reset();
+}
+
 void FHairStrandsCurves::SetNum(const uint32 NumCurves, uint32 InAttributes)
@@ -36,6 +65,16 @@
 	}
 }
 
+void FHairStrandsPoints::Reset()
+{
+	PointsPosition.Reset();
+	PointsRadius.Reset();
+	PointsCoordU.Reset();
+	PointsBaseColor.Reset();
+	PointsRoughness.Reset();
+	PointsAO.Reset();
+}
+
 void FHairStrandsPoints::SetNum(const uint32 NumPoints, uint32 InAttributes)
 {
 	PointsPosition.SetNum(NumPoints);
@@ -64,50 +103,17 @@
 	PointsSimCurvesIndex.Reset();
 }
 
+void FHairStrandsInterpolationDatas::SetNum(const uint32 NumCurves)
+{
+	PointsSimCurvesVertexWeights.SetNum(NumCurves);
+	PointsSimCurvesVertexLerp.SetNum(NumCurves);
+	PointsSimCurvesVertexIndex.SetNum(NumCurves);
+	PointsSimCurvesIndex.SetNum(NumCurves);
+}
+
 void FHairStrandsClusterCullingData::Reset()
 {
 	*this = FHairStrandsClusterCullingData();
 }
 
-void FHairStrandsCurves::Reset()
-{
-	CurvesOffset.Reset();
-	CurvesCount.Reset();
-	CurvesLength.Reset();
-	CurvesRootUV.Reset();
-	ClumpIDs.Reset();
-	CurvesClosestGuideIDs.Reset();
-	CurvesClosestGuideWeights.Reset();
-}
-
-void FHairStrandsPoints::Reset()
-{
-	PointsPosition.Reset();
-	PointsRadius.Reset();
-	PointsCoordU.Reset();
-	PointsBaseColor.Reset();
-	PointsRoughness.Reset();
-	PointsAO.Reset();
-}
-
-float GetHairStrandsMaxLength(const FHairStrandsDatas& In)
-{
-	float MaxLength = 0;
-	for (float CurveLength : In.StrandsCurves.CurvesLength)
-	{
-		MaxLength = FMath::Max(MaxLength, CurveLength);
-	}
-	return MaxLength;
-}
-
-float GetHairStrandsMaxRadius(const FHairStrandsDatas& In)
-{
-	float MaxRadius = 0;
-	for (float PointRadius : In.StrandsPoints.PointsRadius)
-	{
-		MaxRadius = FMath::Max(MaxRadius, PointRadius);
-	}
-	return MaxRadius;
-}
-
 /////////////////////////////////////////////////////////////////////////////////////////
@@ -268,5 +274,143 @@
 }
 
 /////////////////////////////////////////////////////////////////////////////////////////
+// IO/DDC request
+
+bool FHairStreamingRequest::IsNone() const
+{ 
+#if !WITH_EDITORONLY_DATA
+	return IORequest.IsNone();
+#else
+	return Chunks.Num() == 0;
+#endif
+}
+bool FHairStreamingRequest::IsCompleted()
+{ 
+#if !WITH_EDITORONLY_DATA
+	if (IORequest.IsCompleted())
+	{
+		for (FHairStreamingRequest::FChunk& Chunk : Chunks)
+		{
+			Chunk.Status = FHairStreamingRequest::FChunk::Completed;
+			check(Chunk.Container);
+			Chunk.Container->LoadedSize = Chunk.TotalSize;
+		}
+	}
+	return Chunks.Num() == 0 || IORequest.IsCompleted(); 
+#else
+	for (const FHairStreamingRequest::FChunk& Chunk : Chunks)
+	{
+		if (Chunk.Status != FHairStreamingRequest::FChunk::Completed)
+		{
+			return false;
+		}
+	}
+	return true; 
+#endif
+}
+// Request fullfil 2 use cases:
+// * Load DDC data and upload them to GPU
+// * Load DDC data and store them into bulkdata for serialization
+void FHairStreamingRequest::Request(uint32 InRequestedCurveCount, uint32 InRequestedPointCount, FHairStrandsBulkCommon& In, bool bWait, bool bFillBulkdata, const FName& InOwnerName)
+{
+	if (In.GetResourceCount() == 0 || InRequestedCurveCount == 0)
+	{
+		CurveCount = 0;
+		PointCount = 0;
+		return;
+	}
+	CurveCount = InRequestedCurveCount;
+	PointCount = InRequestedPointCount;
+
+#if !WITH_EDITORONLY_DATA
+	{ 
+		In.Read_IO(this, IORequest);
+		if (bWait)
+		{
+			IORequest.Wait();
+		}
+	}
+#else
+	{
+		// When enabled, data can be loaded from an offset. Otherwisee, start from the beginning of the resource
+		// This is used when cooking data to force the loading of the entire resource (i.e., bSupportOffsetLoad=false)
+		bSupportOffsetLoad = !bFillBulkdata; 
+
+		using namespace UE::DerivedData;
+		TArray<FCacheGetChunkRequest> Requests;
+		In.Read_DDC(this, Requests);
+	
+		check(DDCRequestOwner == nullptr);
+		DDCRequestOwner = MakeUnique<FRequestOwner>(bWait ? UE::DerivedData::EPriority::Blocking : UE::DerivedData::EPriority::Normal); // <= Move this onto the resource (Buffer, in order to ensure no race condition)
+	
+		//FRequestBarrier Barrier(*DDCRequestOwner);	// This is a critical section on the owner. It does not constrain ordering
+		GetCache().GetChunks(Requests, *DDCRequestOwner,
+		                     [this, &In, bFillBulkdata, InOwnerName](FCacheGetChunkResponse && Response)
+		                     {
+		                     if (Response.Status == UE::DerivedData::EStatus::Ok)
+		                     {
+		                     FHairStreamingRequest::FChunk& Chunk = *(FHairStreamingRequest::FChunk*)Response.UserData;
+							 Chunk.Data_DDC 	= MoveTemp(Response.RawData);
+							 Chunk.Offset 	= Response.RawOffset;
+							 Chunk.Size 		= Response.RawSize;
+							 Chunk.TotalSize = Response.RawOffset + Response.RawSize;
+							 Chunk.Status 	= FHairStreamingRequest::FChunk::Completed;
+		
+							 // Upload the total amount of loaded data
+							 Chunk.Container->LoadedSize = Chunk.TotalSize;
+
+							 // Optional fill in of bytebulkdata container
+							 if (bFillBulkdata)
+							 {
+								 check(Chunk.Container);
+								 FByteBulkData& BulkData = Chunk.Container->Data;
+	
+								 // The buffer is then stored into bulk data
+								 BulkData.Lock(LOCK_READ_WRITE);
+								 void* DstData = BulkData.Realloc(Chunk.Size);
+								 FMemory::Memcpy(DstData, Chunk.GetData(), Chunk.Size);
+								 BulkData.Unlock();
+
+								 //Chunk.Release();
+							 }
+							 }
+							 else
+							 {
+								 FHairStreamingRequest::FChunk& Chunk = *(FHairStreamingRequest::FChunk*)Response.UserData;
+								 Chunk.Status = FHairStreamingRequest::FChunk::Failed;
+								 UE_LOG(LogHairStrands, Error, TEXT("[Groom] DDC request failed for '%s' (Key:%s) "), *InOwnerName.ToString(), *In.DerivedDataKey);
+							 }
+							 });
+
+		// Optional wait on DDC response
+		if (bWait || bFillBulkdata)
+		{
+			DDCRequestOwner->Wait();
+			check(IsCompleted());
+		}
+	}
+#endif
+}
+
+const uint8* FHairStreamingRequest::FChunk::GetData() const 
+{ 
+	check(Status == Completed); 
+#if !WITH_EDITORONLY_DATA
+	check(Container);
+	return (const uint8*)Data_IO.GetData();
+#else
+	return (const uint8*)Data_DDC.GetData(); 
+#endif
+}
+void FHairStreamingRequest::FChunk::Release() 
+{ 
+#if !WITH_EDITORONLY_DATA
+	TIoStatusOr<uint8*> Out = Data_IO.Release();
+#else
+	Data_DDC.Reset(); 
+#endif
+}
+
+/////////////////////////////////////////////////////////////////////////////////////////
 // Rest bulk data
 
@@ -429,7 +573,6 @@
 	}
 }
 
-
 /////////////////////////////////////////////////////////////////////////////////////////
 // Cluster culling bulk data
 
@@ -539,3 +682,167 @@
 
 	Data.PackedClusterInfos.Data.Unlock();
 }
+
+
+/////////////////////////////////////////////////////////////////////////////////////////
+// Root data
+
+bool FHairStrandsRootData::HasProjectionData() const
+{
+	bool bIsValid = MeshProjectionLODs.Num() > 0;
+	for (const FMeshProjectionLOD& LOD : MeshProjectionLODs)
+	{
+		const bool bHasValidCPUData = LOD.RootBarycentricBuffer.Num() > 0;
+		if (bHasValidCPUData)
+		{
+			bIsValid = bIsValid && LOD.RootBarycentricBuffer.Num() > 0;
+			bIsValid = bIsValid && LOD.RootToUniqueTriangleIndexBuffer.Num() > 0;
+			bIsValid = bIsValid && LOD.UniqueTriangleIndexBuffer.Num() > 0;
+			bIsValid = bIsValid && LOD.RestUniqueTrianglePositionBuffer.Num() > 0;
+
+			if (!bIsValid) break;
+		}
+	}
+
+	return bIsValid;
+}
+
+void FHairStrandsRootData::Reset()
+{
+	RootCount = 0;
+	MeshProjectionLODs.Empty();
+}
+
+/////////////////////////////////////////////////////////////////////////////////////////
+// Root bulk data
+
+void FHairStrandsRootBulkData::SerializeHeader(FArchive& Ar, UObject* Owner)
+{
+	Ar << Header.RootCount;
+	Ar << Header.PointCount;
+
+	// Strides
+	{
+		Ar << Header.Strides.RootToUniqueTriangleIndexBufferStride;
+		Ar << Header.Strides.RootBarycentricBufferStride;
+		Ar << Header.Strides.UniqueTriangleIndexBufferStride;
+		Ar << Header.Strides.RestUniqueTrianglePositionBufferStride;
+
+		Ar << Header.Strides.MeshInterpolationWeightsBufferStride;
+		Ar << Header.Strides.MeshSampleIndicesBufferStride;
+		Ar << Header.Strides.RestSamplePositionsBufferStride;
+	}
+
+	// LODS
+	uint32 LODCount = Header.LODs.Num();
+	Ar << LODCount;
+	if (Ar.IsLoading())
+	{
+		Header.LODs.SetNum(LODCount);
+		Data.LODs.SetNum(LODCount);
+	}
+	for (uint32 LODIt = 0; LODIt < LODCount; ++LODIt)
+	{
+		if (Ar.IsLoading())
+		{
+			Header.LODs[LODIt].LODIndex = LODIt;
+		}
+
+		Ar << Header.LODs[LODIt].UniqueTriangleCount;
+		Ar << Header.LODs[LODIt].SampleCount;
+		Ar << Header.LODs[LODIt].UniqueSectionIndices;
+	}
+}
+
+const TArray<uint32>& FHairStrandsRootBulkData::GetValidSectionIndices(int32 LODIndex) const
+{
+	check(LODIndex >= 0 && LODIndex < Header.LODs.Num());
+	return Header.LODs[LODIndex].UniqueSectionIndices;
+}
+
+void FHairStrandsRootBulkData::SerializeData(FArchive& Ar, UObject* Owner)
+{
+	check(Header.LODs.Num() > 0);
+	check(Header.LODs.Num() == Data.LODs.Num());
+
+	for (uint32 LODIt = 0, LODCount=Header.LODs.Num(); LODIt < LODCount; ++LODIt)
+	{
+		const FHeader::FLOD& HeaderLOD = Header.LODs[LODIt];
+		FData::FLOD& DataLOD = Data.LODs[LODIt];
+
+		const int32 ChunkIndex = 0;
+		bool bAttemptFileMapping = false;
+	
+		if (Ar.IsSaving())
+		{
+			const uint32 BulkFlags = BULKDATA_Force_NOT_InlinePayload;
+	
+			DataLOD.RootBarycentricBuffer.SetBulkDataFlags(BulkFlags);
+			DataLOD.RootToUniqueTriangleIndexBuffer.SetBulkDataFlags(BulkFlags);
+			DataLOD.UniqueTriangleIndexBuffer.SetBulkDataFlags(BulkFlags);
+			DataLOD.RestUniqueTrianglePositionBuffer.SetBulkDataFlags(BulkFlags);
+	
+			DataLOD.MeshInterpolationWeightsBuffer.SetBulkDataFlags(BulkFlags);
+			DataLOD.MeshSampleIndicesBuffer.SetBulkDataFlags(BulkFlags);
+			DataLOD.RestSamplePositionsBuffer.SetBulkDataFlags(BulkFlags);
+		}
+
+		DataLOD.RootBarycentricBuffer.Serialize(Ar, Owner, ChunkIndex, bAttemptFileMapping);
+		DataLOD.RootToUniqueTriangleIndexBuffer.Serialize(Ar, Owner, ChunkIndex, bAttemptFileMapping);
+		DataLOD.UniqueTriangleIndexBuffer.Serialize(Ar, Owner, ChunkIndex, bAttemptFileMapping);
+		DataLOD.RestUniqueTrianglePositionBuffer.Serialize(Ar, Owner, ChunkIndex, bAttemptFileMapping);
+	
+		if (HeaderLOD.SampleCount)
+		{
+			DataLOD.MeshInterpolationWeightsBuffer.Serialize(Ar, Owner, ChunkIndex, bAttemptFileMapping);
+			DataLOD.MeshSampleIndicesBuffer.Serialize(Ar, Owner, ChunkIndex, bAttemptFileMapping);
+			DataLOD.RestSamplePositionsBuffer.Serialize(Ar, Owner, ChunkIndex, bAttemptFileMapping);
+		}
+	}
+}
+
+void FHairStrandsRootBulkData::Serialize(FArchive& Ar, UObject* Owner)
+{
+	// Root data don't have UObject references
+	if (!Ar.IsLoading() && !Ar.IsSaving() && Ar.IsObjectReferenceCollector())
+	{
+		return;
+	}
+
+	SerializeHeader(Ar, Owner);
+	SerializeData(Ar, Owner);
+}
+
+void FHairStrandsRootBulkData::Reset()
+{
+	// Header
+	Header.RootCount = 0;
+	Header.PointCount = 0;
+	Header.LODs.Empty();
+
+	// Data
+	for (FData::FLOD& LOD : Data.LODs)
+	{
+		// Binding
+		LOD.RootBarycentricBuffer.RemoveBulkData();
+		LOD.RootToUniqueTriangleIndexBuffer.RemoveBulkData();
+		LOD.UniqueTriangleIndexBuffer.RemoveBulkData();
+		LOD.RestUniqueTrianglePositionBuffer.RemoveBulkData();
+
+		// Reset the bulk byte buffer to ensure the (serialize) data size is reset to 0
+		LOD.RootBarycentricBuffer 			 = FByteBulkData();
+		LOD.RootToUniqueTriangleIndexBuffer  = FByteBulkData();
+		LOD.UniqueTriangleIndexBuffer		 = FByteBulkData();
+		LOD.RestUniqueTrianglePositionBuffer = FByteBulkData();
+
+		// RBF
+		LOD.MeshInterpolationWeightsBuffer.RemoveBulkData();
+		LOD.MeshSampleIndicesBuffer.RemoveBulkData();
+		LOD.RestSamplePositionsBuffer.RemoveBulkData();
+
+		LOD.MeshInterpolationWeightsBuffer 	 = FByteBulkData();
+		LOD.MeshSampleIndicesBuffer 		 = FByteBulkData();
+		LOD.RestSamplePositionsBuffer 		 = FByteBulkData();
+	}
+	Data.LODs.Empty();
+}
diff --git a/Engine/Plugins/Runtime/HairStrands/Source/HairStrandsCore/Public/GroomResources.h b/Engine/Plugins/Runtime/HairStrands/Source/HairStrandsCore/Public/GroomResources.h
--- a/Engine/Plugins/Runtime/HairStrands/Source/HairStrandsCore/Public/GroomResources.h
+++ b/Engine/Plugins/Runtime/HairStrands/Source/HairStrandsCore/Public/GroomResources.h
@@ -14,7 +14,8 @@
 
 namespace UE::DerivedData { class FRequestOwner; }
 
-#define STRANDS_PROCEDURAL_INTERSECTOR_MAX_SPLITS 4
+////////////////////////////////////////////////////////////////////////////////////////////////////////////////////////
+// Helpers
 
 inline uint32 GetBufferTotalNumBytes(const FRDGExternalBuffer& In) 
 {
@@ -55,6 +56,9 @@
 #endif
 };
 
+////////////////////////////////////////////////////////////////////////////////////////////////////////////////////////
+// Loading Type/Status
+
 enum class EHairResourceLoadingType : uint8
 {
 	Async,
@@ -83,6 +87,9 @@
 
 EHairResourceLoadingType GetHairResourceLoadingType(EHairGeometryType InGeometryType, int32 InLODIndex);
 
+////////////////////////////////////////////////////////////////////////////////////////////////////////////////////////
+// Common resources
+
 /* Hair resouces which whom allocation can be deferred */
 struct FHairCommonResource : public FRenderResource
 {
@@ -122,177 +129,8 @@
 	uint32 MaxAvailableCurveCount = HAIR_MAX_NUM_CURVE_PER_GROUP;
 };
 
-/* Render buffers for root deformation for dynamic meshes */
-struct FHairStrandsRestRootResource : public FHairCommonResource
-{
-	/** Build the hair strands resource */
-	FHairStrandsRestRootResource(FHairStrandsRootBulkData& BulkData, EHairStrandsResourcesType CurveType, const FHairResourceName& ResourceName, const FName& OwnerName);
+////////////////////////////////////////////////////////////////////////////////////////////////////////////////////////
+// Strands resources
 
-	/* Init/Release buffers */
-	virtual void InternalAllocate(FRDGBuilder& GraphBuilder) override;
-	virtual void InternalAllocateLOD(FRDGBuilder& GraphBuilder, int32 LODIndex) override;
-	virtual void InternalRelease() override;
-	virtual bool InternalIsLODDataLoaded(int32 LODIndex) override;
-
-	/* Get the resource name */
-	virtual FString GetFriendlyName() const override { return TEXT("FHairStrandsRestRootResource"); }
-
-	/* Populate GPU LOD data from RootData (this function doesn't initialize resources) */
-	void PopulateFromRootData();
-
-	// Accessors
-	uint32 GetLODCount() const { return BulkData.GetLODCount(); }
-	uint32 GetRootCount()const { return BulkData.GetRootCount(); }
-
-	/* Return the memory size for GPU resources */
-	uint32 GetResourcesSize() const 
-	{
-		uint32 Total = 0;
-		for (const FLOD& LOD : LODs)
-		{
-			Total += GetBufferTotalNumBytes(LOD.UniqueTriangleIndexBuffer);
-			Total += GetBufferTotalNumBytes(LOD.RootToUniqueTriangleIndexBuffer);
-			Total += GetBufferTotalNumBytes(LOD.RootBarycentricBuffer);
-			Total += GetBufferTotalNumBytes(LOD.RestUniqueTrianglePositionBuffer);
-			Total += GetBufferTotalNumBytes(LOD.MeshInterpolationWeightsBuffer);
-			Total += GetBufferTotalNumBytes(LOD.MeshSampleIndicesBuffer);
-			Total += GetBufferTotalNumBytes(LOD.RestSamplePositionsBuffer);
-		}
-		return Total;
-	}
-
-	struct FLOD
-	{
-		enum class EStatus { Invalid, Initialized, Completed };
-
-		const bool IsValid() const { return Status == EStatus::Completed; }
-		EStatus Status = EStatus::Invalid;
-		int32 LODIndex = -1;
-
-		/* Triangle on which a root is attached */
-		/* When the projection is done with source to target mesh transfer, the projection indices does not match.
-		   In this case we need to separate index computation. The barycentric coords remain the same however. */
-		FRDGExternalBuffer UniqueTriangleIndexBuffer;
-		/* Strands hair root to unique triangle index */
-		FRDGExternalBuffer RootToUniqueTriangleIndexBuffer;
-		FRDGExternalBuffer RootBarycentricBuffer;
-
-		/* Strand hair roots translation and rotation in rest position relative to the bound triangle. Positions are relative to the rest root center */
-		FRDGExternalBuffer RestUniqueTrianglePositionBuffer;
-
-		/* Strand hair mesh interpolation matrix and sample indices */
-		uint32 SampleCount = 0;
-		FRDGExternalBuffer MeshInterpolationWeightsBuffer;
-		FRDGExternalBuffer MeshSampleIndicesBuffer;
-		FRDGExternalBuffer RestSamplePositionsBuffer;
-	};
-
-	/* Store the hair projection information for each mesh LOD */
-	TArray<FLOD> LODs;
-
-	/* LOD bulk data requests */
-	TArray<FBulkDataBatchRequest> LODRequests;
-
-	/* Store CPU data for root info & root binding */
-	FHairStrandsRootBulkData& BulkData;
-
-	/* Type of curves */
-	const EHairStrandsResourcesType CurveType;
-};
-
-/* Render buffers for root deformation for dynamic meshes */
-struct FHairStrandsDeformedRootResource : public FHairCommonResource
-{
-	/** Build the hair strands resource */
-	FHairStrandsDeformedRootResource(EHairStrandsResourcesType CurveType, const FHairResourceName& ResourceName, const FName& OwnerName);
-	FHairStrandsDeformedRootResource(const FHairStrandsRestRootResource* InRestResources, EHairStrandsResourcesType CurveType, const FHairResourceName& ResourceName, const FName& OwnerName);
-
-	/* Init/Release buffers */
-	virtual void InternalAllocateLOD(FRDGBuilder& GraphBuilder, int32 LODIndex) override;
-	virtual void InternalRelease() override;
-
-	/* Get the resource name */
-	virtual FString GetFriendlyName() const override { return TEXT("FHairStrandsDeformedRootResource"); }
-
-	/* Indirect if the current root resources are valid and up to date */
-	bool IsValid() const { return MeshLODIndex >= 0 && MeshLODIndex < LODs.Num() && LODs[MeshLODIndex].IsValid(); }
-	bool IsValid(int32 InMeshLODIndex) const { return InMeshLODIndex >= 0 && InMeshLODIndex < LODs.Num() && LODs[InMeshLODIndex].IsValid(); }
-
-	/* Return the memory size for GPU resources */
-	uint32 GetResourcesSize() const
-	{
-		uint32 Total = 0;
-		for (const FLOD& LOD : LODs)
-		{
-			Total += GetBufferTotalNumBytes(LOD.DeformedUniqueTrianglePositionBuffer[0]);
-			Total += GetBufferTotalNumBytes(LOD.DeformedSamplePositionsBuffer[0]);
-			Total += GetBufferTotalNumBytes(LOD.MeshSampleWeightsBuffer[0]);
-
-			// Double buffering is disabled by default unless the read-only cvar r.HairStrands.ContinuousDecimationReordering is set
-			if (IsHairStrandContinuousDecimationReorderingEnabled())
-			{
-				Total += GetBufferTotalNumBytes(LOD.DeformedUniqueTrianglePositionBuffer[1]);
-				Total += GetBufferTotalNumBytes(LOD.DeformedSamplePositionsBuffer[1]);
-				Total += GetBufferTotalNumBytes(LOD.MeshSampleWeightsBuffer[1]);
-			}
-		}
-		return Total;
-	}
-
-	void SwapBuffer()
-	{
-		for (FLOD& LOD : LODs)
-		{
-			LOD.SwapBuffer();
-		}
-	}
-
-	struct FLOD
-	{
-		enum class EStatus { Invalid, Initialized, Completed };
-
-		// A LOD is considered valid as long as its resources are initialized. 
-		// Its state will become completed once its triangle position will be 
-		// update, but in order to be update its status needs to be valid.
-		const bool IsValid() const { return Status == EStatus::Initialized || Status == EStatus::Completed; }
-		EStatus Status = EStatus::Invalid;
-		int32 LODIndex = -1;
-
-		/* Strand hair roots translation and rotation in triangle-deformed position relative to the bound triangle. Positions are relative the deformed root center*/
-		FRDGExternalBuffer DeformedUniqueTrianglePositionBuffer[2];
-
-		/* Strand hair mesh interpolation matrix and sample indices */
-		uint32 SampleCount = 0;
-		FRDGExternalBuffer DeformedSamplePositionsBuffer[2];
-		FRDGExternalBuffer MeshSampleWeightsBuffer[2];
-
-		/* Whether the GPU data should be initialized with the asset data or not */
-		uint32 CurrentIndex = 0;
-
-		enum EFrameType
-		{
-			Previous,
-			Current
-		};
-
-		// Double buffering is disabled by default unless the read-only cvar r.HairStrands.ContinuousDecimationReordering is set
-		inline uint32 GetIndex(EFrameType T) const { return T == EFrameType::Current ? CurrentIndex : 1u - CurrentIndex; }
-		inline const FRDGExternalBuffer& GetDeformedUniqueTrianglePositionBuffer(EFrameType T) const { return IsHairStrandContinuousDecimationReorderingEnabled() ? DeformedUniqueTrianglePositionBuffer[GetIndex(T)] : DeformedUniqueTrianglePositionBuffer[0]; }
-		inline const FRDGExternalBuffer& GetDeformedSamplePositionsBuffer(EFrameType T) const { return IsHairStrandContinuousDecimationReorderingEnabled() ? DeformedSamplePositionsBuffer[GetIndex(T)] : DeformedSamplePositionsBuffer[0]; }
-		inline const FRDGExternalBuffer& GetMeshSampleWeightsBuffer(EFrameType T) const { return IsHairStrandContinuousDecimationReorderingEnabled() ? MeshSampleWeightsBuffer[GetIndex(T)] : MeshSampleWeightsBuffer[0]; }
-		inline void SwapBuffer() { CurrentIndex = 1u - CurrentIndex; }
-	};
-
-	/* Store the hair projection information for each mesh LOD */
-	uint32 RootCount = 0;
-	TArray<FLOD> LODs;
-
-	/* Last update MeshLODIndex */
-	int32 MeshLODIndex = -1;
-
-	/* Type of curves */
-	const EHairStrandsResourcesType CurveType;
-};
-
 /* Render buffers that will be used for rendering */
 struct FHairStrandsRestResource : public FHairCommonResource
@@ -407,9 +245,9 @@
 	uint32 CurrentIndex = 0;
 
 	/* Track the view which has update the formed position. This is used when rendering the same 
-	   instance accross several editor viewport (not views of the same viewport), to prevent 
-	   incorrect motion vector, as editor viewport are refresh not on the same tick, but with a 
-	   throttling mechanism which make the updates not coherent. */
+		instance accross several editor viewport (not views of the same viewport), to prevent 
+		incorrect motion vector, as editor viewport are refresh not on the same tick, but with a 
+		throttling mechanism which make the updates not coherent. */
 	uint32 UniqueViewIDs[2] = { 0, 0 };
 
 	/* Type of curves */
@@ -439,6 +277,39 @@
 	FRDGExternalBuffer& GetDeformerPointAttributeBuffer(FRDGBuilder& GraphBuilder);
 };
 
+
+struct FHairStrandsInterpolationResource : public FHairCommonResource
+{
+	/** Build the hair strands resource */
+	FHairStrandsInterpolationResource(FHairStrandsInterpolationBulkData& InBulkData, const FHairResourceName& ResourceName, const FName& OwnerName);
+	
+	/* Init/Release buffers */
+	virtual void InternalAllocate(FRDGBuilder& GraphBuilder) override;
+	virtual void InternalRelease() override;
+	virtual bool InternalIsDataLoaded(uint32 InRequestedCurveCount, uint32 InRequestedPointCount) override;
+
+	/* Get the resource name */
+	virtual FString GetFriendlyName() const override { return TEXT("FHairStrandsInterplationResource"); }
+
+	/* Return the memory size for GPU resources */
+	uint32 GetResourcesSize() const
+	{
+		uint32 Total = 0;
+		Total += GetBufferTotalNumBytes(InterpolationBuffer);
+		Total += GetBufferTotalNumBytes(SimRootPointIndexBuffer);
+		return Total;
+	}
+
+	bool UseSingleGuide() const { return (BulkData.Header.Flags & FHairStrandsInterpolationBulkData::DataFlags_HasSingleGuideData) != 0; }
+
+	FRDGExternalBuffer InterpolationBuffer;
+	FRDGExternalBuffer SimRootPointIndexBuffer;
+
+	/* Reference to the hair strands interpolation render data */
+	FHairStrandsInterpolationBulkData& BulkData;
+};
+
+
 struct FHairStrandsClusterCullingResource : public FHairCommonResource
 {
 	FHairStrandsClusterCullingResource(FHairStrandsClusterCullingBulkData& Data, const FHairResourceName& ResourceName, const FName& OwnerName);
@@ -475,35 +346,176 @@
 	FHairStrandsClusterCullingBulkData& BulkData;
 };
 
-struct FHairStrandsInterpolationResource : public FHairCommonResource
+/* Render buffers for root deformation for dynamic meshes */
+struct FHairStrandsRestRootResource : public FHairCommonResource
 {
 	/** Build the hair strands resource */
-	FHairStrandsInterpolationResource(FHairStrandsInterpolationBulkData& InBulkData, const FHairResourceName& ResourceName, const FName& OwnerName);
-	
+	FHairStrandsRestRootResource(FHairStrandsRootBulkData& BulkData, EHairStrandsResourcesType CurveType, const FHairResourceName& ResourceName, const FName& OwnerName);
+
 	/* Init/Release buffers */
 	virtual void InternalAllocate(FRDGBuilder& GraphBuilder) override;
+	virtual void InternalAllocateLOD(FRDGBuilder& GraphBuilder, int32 LODIndex) override;
 	virtual void InternalRelease() override;
-	virtual bool InternalIsDataLoaded(uint32 InRequestedCurveCount, uint32 InRequestedPointCount) override;
+	virtual bool InternalIsLODDataLoaded(int32 LODIndex) override;
 
 	/* Get the resource name */
-	virtual FString GetFriendlyName() const override { return TEXT("FHairStrandsInterplationResource"); }
+	virtual FString GetFriendlyName() const override { return TEXT("FHairStrandsRestRootResource"); }
+
+	/* Populate GPU LOD data from RootData (this function doesn't initialize resources) */
+	void PopulateFromRootData();
+
+	// Accessors
+	uint32 GetLODCount() const { return BulkData.GetLODCount(); }
+	uint32 GetRootCount()const { return BulkData.GetRootCount(); }
 
 	/* Return the memory size for GPU resources */
+	uint32 GetResourcesSize() const 
+	{
+		uint32 Total = 0;
+		for (const FLOD& LOD : LODs)
+		{
+			Total += GetBufferTotalNumBytes(LOD.UniqueTriangleIndexBuffer);
+			Total += GetBufferTotalNumBytes(LOD.RootToUniqueTriangleIndexBuffer);
+			Total += GetBufferTotalNumBytes(LOD.RootBarycentricBuffer);
+			Total += GetBufferTotalNumBytes(LOD.RestUniqueTrianglePositionBuffer);
+			Total += GetBufferTotalNumBytes(LOD.MeshInterpolationWeightsBuffer);
+			Total += GetBufferTotalNumBytes(LOD.MeshSampleIndicesBuffer);
+			Total += GetBufferTotalNumBytes(LOD.RestSamplePositionsBuffer);
+		}
+		return Total;
+	}
+
+	struct FLOD
+	{
+		enum class EStatus { Invalid, Initialized, Completed };
+
+		const bool IsValid() const { return Status == EStatus::Completed; }
+		EStatus Status = EStatus::Invalid;
+		int32 LODIndex = -1;
+
+		/* Triangle on which a root is attached */
+		/* When the projection is done with source to target mesh transfer, the projection indices does not match.
+		   In this case we need to separate index computation. The barycentric coords remain the same however. */
+		FRDGExternalBuffer UniqueTriangleIndexBuffer;
+		/* Strands hair root to unique triangle index */
+		FRDGExternalBuffer RootToUniqueTriangleIndexBuffer;
+		FRDGExternalBuffer RootBarycentricBuffer;
+
+		/* Strand hair roots translation and rotation in rest position relative to the bound triangle. Positions are relative to the rest root center */
+		FRDGExternalBuffer RestUniqueTrianglePositionBuffer;
+
+		/* Strand hair mesh interpolation matrix and sample indices */
+		uint32 SampleCount = 0;
+		FRDGExternalBuffer MeshInterpolationWeightsBuffer;
+		FRDGExternalBuffer MeshSampleIndicesBuffer;
+		FRDGExternalBuffer RestSamplePositionsBuffer;
+	};
+
+	/* Store the hair projection information for each mesh LOD */
+	TArray<FLOD> LODs;
+
+	/* LOD bulk data requests */
+	TArray<FBulkDataBatchRequest> LODRequests;
+
+	/* Store CPU data for root info & root binding */
+	FHairStrandsRootBulkData& BulkData;
+
+	/* Type of curves */
+	const EHairStrandsResourcesType CurveType;
+};
+
+/* Render buffers for root deformation for dynamic meshes */
+struct FHairStrandsDeformedRootResource : public FHairCommonResource
+{
+	/** Build the hair strands resource */
+	FHairStrandsDeformedRootResource(EHairStrandsResourcesType CurveType, const FHairResourceName& ResourceName, const FName& OwnerName);
+	FHairStrandsDeformedRootResource(const FHairStrandsRestRootResource* InRestResources, EHairStrandsResourcesType CurveType, const FHairResourceName& ResourceName, const FName& OwnerName);
+
+	/* Init/Release buffers */
+	virtual void InternalAllocateLOD(FRDGBuilder& GraphBuilder, int32 LODIndex) override;
+	virtual void InternalRelease() override;
+
+	/* Get the resource name */
+	virtual FString GetFriendlyName() const override { return TEXT("FHairStrandsDeformedRootResource"); }
+
+	/* Indirect if the current root resources are valid and up to date */
+	bool IsValid() const { return MeshLODIndex >= 0 && MeshLODIndex < LODs.Num() && LODs[MeshLODIndex].IsValid(); }
+	bool IsValid(int32 InMeshLODIndex) const { return InMeshLODIndex >= 0 && InMeshLODIndex < LODs.Num() && LODs[InMeshLODIndex].IsValid(); }
+
+	/* Return the memory size for GPU resources */
 	uint32 GetResourcesSize() const
 	{
 		uint32 Total = 0;
-		Total += GetBufferTotalNumBytes(InterpolationBuffer);
-		Total += GetBufferTotalNumBytes(SimRootPointIndexBuffer);
+		for (const FLOD& LOD : LODs)
+		{
+			Total += GetBufferTotalNumBytes(LOD.DeformedUniqueTrianglePositionBuffer[0]);
+			Total += GetBufferTotalNumBytes(LOD.DeformedSamplePositionsBuffer[0]);
+			Total += GetBufferTotalNumBytes(LOD.MeshSampleWeightsBuffer[0]);
+
+			// Double buffering is disabled by default unless the read-only cvar r.HairStrands.ContinuousDecimationReordering is set
+			if (IsHairStrandContinuousDecimationReorderingEnabled())
+			{
+				Total += GetBufferTotalNumBytes(LOD.DeformedUniqueTrianglePositionBuffer[1]);
+				Total += GetBufferTotalNumBytes(LOD.DeformedSamplePositionsBuffer[1]);
+				Total += GetBufferTotalNumBytes(LOD.MeshSampleWeightsBuffer[1]);
+			}
+		}
 		return Total;
 	}
 
-	bool UseSingleGuide() const { return (BulkData.Header.Flags & FHairStrandsInterpolationBulkData::DataFlags_HasSingleGuideData) != 0; }
+	void SwapBuffer()
+	{
+		for (FLOD& LOD : LODs)
+		{
+			LOD.SwapBuffer();
+		}
+	}
 
-	FRDGExternalBuffer InterpolationBuffer;
-	FRDGExternalBuffer SimRootPointIndexBuffer;
+	struct FLOD
+	{
+		enum class EStatus { Invalid, Initialized, Completed };
 
-	/* Reference to the hair strands interpolation render data */
-	FHairStrandsInterpolationBulkData& BulkData;
+		// A LOD is considered valid as long as its resources are initialized. 
+		// Its state will become completed once its triangle position will be 
+		// update, but in order to be update its status needs to be valid.
+		const bool IsValid() const { return Status == EStatus::Initialized || Status == EStatus::Completed; }
+		EStatus Status = EStatus::Invalid;
+		int32 LODIndex = -1;
+
+		/* Strand hair roots translation and rotation in triangle-deformed position relative to the bound triangle. Positions are relative the deformed root center*/
+		FRDGExternalBuffer DeformedUniqueTrianglePositionBuffer[2];
+
+		/* Strand hair mesh interpolation matrix and sample indices */
+		uint32 SampleCount = 0;
+		FRDGExternalBuffer DeformedSamplePositionsBuffer[2];
+		FRDGExternalBuffer MeshSampleWeightsBuffer[2];
+
+		/* Whether the GPU data should be initialized with the asset data or not */
+		uint32 CurrentIndex = 0;
+
+		enum EFrameType
+		{
+			Previous,
+			Current
+		};
+
+		// Double buffering is disabled by default unless the read-only cvar r.HairStrands.ContinuousDecimationReordering is set
+		inline uint32 GetIndex(EFrameType T) const { return T == EFrameType::Current ? CurrentIndex : 1u - CurrentIndex; }
+		inline const FRDGExternalBuffer& GetDeformedUniqueTrianglePositionBuffer(EFrameType T) const { return IsHairStrandContinuousDecimationReorderingEnabled() ? DeformedUniqueTrianglePositionBuffer[GetIndex(T)] : DeformedUniqueTrianglePositionBuffer[0]; }
+		inline const FRDGExternalBuffer& GetDeformedSamplePositionsBuffer(EFrameType T) const { return IsHairStrandContinuousDecimationReorderingEnabled() ? DeformedSamplePositionsBuffer[GetIndex(T)] : DeformedSamplePositionsBuffer[0]; }
+		inline const FRDGExternalBuffer& GetMeshSampleWeightsBuffer(EFrameType T) const { return IsHairStrandContinuousDecimationReorderingEnabled() ? MeshSampleWeightsBuffer[GetIndex(T)] : MeshSampleWeightsBuffer[0]; }
+		inline void SwapBuffer() { CurrentIndex = 1u - CurrentIndex; }
+	};
+
+	/* Store the hair projection information for each mesh LOD */
+	uint32 RootCount = 0;
+	TArray<FLOD> LODs;
+
+	/* Last update MeshLODIndex */
+	int32 MeshLODIndex = -1;
+
+	/* Type of curves */
+	const EHairStrandsResourcesType CurveType;
 };
 
 #if RHI_RAYTRACING
@@ -558,8 +570,6 @@
 	virtual void InitRHI() override;
 };
 
-struct FHairCardsBulkData;
-
 /* Render buffers that will be used for rendering */
 struct FHairCardsRestResource : public FHairCommonResource
 {
@@ -700,37 +710,6 @@
 	inline FRDGExternalBuffer& GetBuffer(EFrameType T)	{ return DeformedPositionBuffer[GetIndex(T)];  }
 };
 
-struct FHairCardsInterpolationBulkData;
-
-/** Hair cards points interpolation attributes */
-struct HAIRSTRANDSCORE_API FHairCardsInterpolationDatas
-{
-	/** Set the number of interpolated points */
-	void SetNum(const uint32 NumPoints);
-
-	/** Reset the interpolated points to 0 */
-	void Reset();
-
-	/** Get the number of interpolated points */
-	uint32 Num() const { return PointsSimCurvesVertexIndex.Num(); }
-
-	/** Simulation curve indices */
-	TArray<int32> PointsSimCurvesIndex;
-
-	/** Closest vertex indices on simulation curve */
-	TArray<int32> PointsSimCurvesVertexIndex;
-
-	/** Lerp value between the closest vertex indices and the next one */
-	TArray<float> PointsSimCurvesVertexLerp;
-};
-
-struct HAIRSTRANDSCORE_API FHairCardsInterpolationBulkData
-{
-	TArray<FHairCardsInterpolationFormat::Type> Interpolation;
-
-	void Serialize(FArchive& Ar);
-};
-
 struct FHairCardsInterpolationResource : public FHairCommonResource
 {
 	/** Build the hair strands resource */
diff --git a/Engine/Plugins/Runtime/HairStrands/Source/HairStrandsCore/Public/HairCardsDatas.h b/Engine/Plugins/Runtime/HairStrands/Source/HairStrandsCore/Public/HairCardsDatas.h
--- a/Engine/Plugins/Runtime/HairStrands/Source/HairStrandsCore/Public/HairCardsDatas.h
+++ b/Engine/Plugins/Runtime/HairStrands/Source/HairStrandsCore/Public/HairCardsDatas.h
@@ -9,7 +9,7 @@
 #include "HairStrandsVertexFactory.h"
 
 ////////////////////////////////////////////////////////////////////////////////////////////////
-// Utils/Format data
+// Formats
 
 struct FHairCardsPositionFormat
 {
@@ -173,7 +173,7 @@ struct FHairOrientedBound
 };
 
 ////////////////////////////////////////////////////////////////////////////////////////////////
-// Hair Cards
+// Hair Cards Data/Bulk
 
 // Data structure holding cards geometry information
 struct FHairCardsGeometry
@@ -288,7 +288,7 @@ struct FHairCardsBulkData
 };
 
 ////////////////////////////////////////////////////////////////////////////////////////////////
-// Hair Cards (Procedural generation)
+// Hair Cards Data (Procedural generation)
 
 // Internal data structure for computing surface flow. 
 // This structured is exposed only for debug purpose
@@ -406,7 +406,39 @@ struct HAIRSTRANDSCORE_API FHairCardsSourceData
 };
 
 ////////////////////////////////////////////////////////////////////////////////////////////////
-// Hair meshes
+// Hair Card interpolation Data/Bulk
+
+/** Hair cards points interpolation attributes */
+struct HAIRSTRANDSCORE_API FHairCardsInterpolationDatas
+{
+	/** Set the number of interpolated points */
+	void SetNum(const uint32 NumPoints);
+
+	/** Reset the interpolated points to 0 */
+	void Reset();
+
+	/** Get the number of interpolated points */
+	uint32 Num() const { return PointsSimCurvesVertexIndex.Num(); }
+
+	/** Simulation curve indices */
+	TArray<int32> PointsSimCurvesIndex;
+
+	/** Closest vertex indices on simulation curve */
+	TArray<int32> PointsSimCurvesVertexIndex;
+
+	/** Lerp value between the closest vertex indices and the next one */
+	TArray<float> PointsSimCurvesVertexLerp;
+};
+
+struct HAIRSTRANDSCORE_API FHairCardsInterpolationBulkData
+{
+	TArray<FHairCardsInterpolationFormat::Type> Interpolation;
+
+	void Serialize(FArchive& Ar);
+};
+
+////////////////////////////////////////////////////////////////////////////////////////////////
+// Hair Meshes Data/Bulk
 
 struct FHairMeshes
 {
@@ -435,8 +467,6 @@ struct FHairMeshes
 	uint32 GetNumVertices() const { return Positions.Num(); }
 };
 
-struct FHairMeshesBulkData;
-
 struct FHairMeshesDatas
 {
 	bool IsValid() const { return Meshes.Positions.Num() > 0; }
diff --git a/Engine/Plugins/Runtime/HairStrands/Source/HairStrandsCore/Public/HairStrandsDatas.h b/Engine/Plugins/Runtime/HairStrands/Source/HairStrandsCore/Public/HairStrandsDatas.h
--- a/Engine/Plugins/Runtime/HairStrands/Source/HairStrandsCore/Public/HairStrandsDatas.h
+++ b/Engine/Plugins/Runtime/HairStrands/Source/HairStrandsCore/Public/HairStrandsDatas.h
@@ -14,6 +14,19 @@
 
 DECLARE_LOG_CATEGORY_EXTERN(LogHairStrands, Log, All);
 
+///////////////////////////////////////////////////////////////////////////////////////////////////
+// Forward declarations
+
+struct FHairStrandsBulkCommon;
+struct FHairBulkContainer;
+struct FHairStrandsDatas;
+
+float GetHairStrandsMaxLength(const FHairStrandsDatas& In);
+float GetHairStrandsMaxRadius(const FHairStrandsDatas& In);
+
+///////////////////////////////////////////////////////////////////////////////////////////////////
+// Formats
+
 enum class EHairAttribute : uint8;
 namespace UE::DerivedData
 {
@@ -235,8 +248,91 @@
 	static float	 PackUVsToFloat(const FVector2f& UV);
 };
 
-struct FHairStrandsBulkCommon;
-struct FHairBulkContainer;
+/*  Structure describing the LOD settings (Screen size, vertex info, ...) for each clusters.
+	The packed version of this structure corresponds to the GPU data layout (HairStrandsClusterCommon.ush)
+	This uses by the GPU LOD selection. */
+struct FHairClusterInfo
+{
+	static const uint32 MaxLOD = 8;
 
+	struct Packed
+	{
+		uint32 LODInfoOffset : 24;
+		uint32 LODCount : 8;
+
+		uint32 LOD_ScreenSize_0 : 10;
+		uint32 LOD_ScreenSize_1 : 10;
+		uint32 LOD_ScreenSize_2 : 10;
+		uint32 Pad0 : 2;
+
+		uint32 LOD_ScreenSize_3 : 10;
+		uint32 LOD_ScreenSize_4 : 10;
+		uint32 LOD_ScreenSize_5 : 10;
+		uint32 Pad1 : 2;
+
+		uint32 LOD_ScreenSize_6 : 10;
+		uint32 LOD_ScreenSize_7 : 10;
+		uint32 LOD_bIsVisible : 8;
+		uint32 Pad2 : 4;
+	};
+	typedef FUintVector4 BulkType;
+
+	FHairClusterInfo()
+	{
+		for (uint32 LODIt = 0; LODIt < MaxLOD; ++LODIt)
+		{
+			ScreenSize[LODIt] = 0;
+			bIsVisible[LODIt] = true;
+		}
+	}
+
+	uint32 LODCount = 0;
+	uint32 LODInfoOffset = 0;
+	TStaticArray<float,MaxLOD> ScreenSize;
+	TStaticArray<bool, MaxLOD> bIsVisible;
+};
+
+/*  Structure describing the LOD settings common to all clusters. The layout of this structure is
+	identical the GPU data layout (HairStrandsClusterCommon.ush). This uses by the GPU LOD selection. */
+struct FHairClusterLODInfo
+{
+	uint32 VertexOffset = 0;
+	uint32 VertexCount0 = 0;
+	uint32 VertexCount1 = 0;
+	FFloat16 RadiusScale0 = 0;
+	FFloat16 RadiusScale1 = 0;
+};
+
+struct FHairLODInfo
+{
+	uint32 CurveCount = 0;
+	uint32 PointCount = 0;
+};
+
+struct FHairClusterInfoFormat
+{
+	typedef FHairClusterInfo::Packed Type;
+	typedef FHairClusterInfo::Packed BulkType;
+	static const uint32 SizeInByte = sizeof(Type);
+};
+
+struct FHairClusterLODInfoFormat
+{
+	typedef FHairClusterLODInfo Type;
+	typedef FHairClusterLODInfo BulkType;
+	static const uint32 SizeInByte = sizeof(Type);
+};
+
+struct FHairClusterIndexFormat
+{
+	typedef uint32 Type;
+	typedef uint32 BulkType;
+	static const uint32 SizeInByte = sizeof(Type);
+	static const EPixelFormat Format = PF_R32_UINT;
+};
+
+///////////////////////////////////////////////////////////////////////////////////////////////////
+// Streaming 
+
 // Streaming request are used for reading hair strands data from DDC or IO
 // The request are translated later to FHairStrandsBulkCommon::FQuery for appropriate reading.
@@ -293,6 +389,9 @@
 	bool bSupportOffsetLoad = true; 
 };
 
+///////////////////////////////////////////////////////////////////////////////////////////////////
+// Common hair data container and hair data
+
 struct FHairBulkContainer
 {
 	// Streaming
@@ -350,70 +449,8 @@
 #endif
 };
 
-/** Hair strands points interpolation attributes */
-struct HAIRSTRANDSCORE_API FHairStrandsInterpolationDatas
-{
-	/** Set the number of interpolated points */
-	void SetNum(const uint32 NumPoints);
+///////////////////////////////////////////////////////////////////////////////////////////////////
+// Strands Data/Bulk
 
-	/** Reset the interpolated points to 0 */
-	void Reset();
-
-	/** Get the number of interpolated points */
-	uint32 Num() const { return PointsSimCurvesVertexIndex.Num(); }
-
-	bool IsValid() const { return PointsSimCurvesIndex.Num() > 0; }
-
-	/** Simulation curve indices, ordered by closest influence */
-	TArray<FIntVector> PointsSimCurvesIndex;
-
-	/** Closest vertex indices on simulation curve, ordered by closest influence */
-	TArray<FIntVector> PointsSimCurvesVertexIndex;
-
-	/** Lerp value between the closest vertex indices and the next one, ordered by closest influence */
-	TArray<FVector3f> PointsSimCurvesVertexLerp;
-
-	/** Weight of vertex indices on simulation curve, ordered by closest influence */
-	TArray<FVector3f>	PointsSimCurvesVertexWeights;
-
-	/** True, if interpolation data are built using a single guide */
-	bool bUseUniqueGuide = false;
-};
-
-struct HAIRSTRANDSCORE_API FHairStrandsInterpolationBulkData : FHairStrandsBulkCommon
-{
-	enum EDataFlags
-	{
-		DataFlags_HasData = 1,
-		DataFlags_HasSingleGuideData = 2,
-	};
-
-	void Reset();
-	virtual void SerializeHeader(FArchive& Ar, UObject* Owner) override;
-
-	virtual uint32 GetResourceCount() const override;
-	virtual void GetResources(FQuery& Out) override;
-	uint32 GetPointCount() const { return Header.PointCount; };
-
-	struct FHeader
-	{
-		uint32 Flags = 0;
-		uint32 PointCount = 0;
-		uint32 SimPointCount = 0;
-
-		struct FStrides
-		{
-			uint32 InterpolationStride = 0;
-			uint32 SimRootPointIndexStride = 0;
-		} Strides;
-	} Header;
-
-	struct FData
-	{
-		FHairBulkContainer Interpolation;		// FHairStrandsInterpolationFormat  - Per-rendering-vertex interpolation data (closest guides, weight factors, ...). Data for a 1 or 3 guide(s))
-		FHairBulkContainer SimRootPointIndex;	// FHairStrandsRootIndexFormat      - Per-rendering-vertex index of the sim-root vertex
-	} Data;
-};
-
 /** Hair strands points attribute */
 struct HAIRSTRANDSCORE_API FHairStrandsPoints
@@ -529,9 +566,6 @@
 	FBox BoundingBox = FBox(EForceInit::ForceInit);
 };
 
-float GetHairStrandsMaxLength(const FHairStrandsDatas& In);
-float GetHairStrandsMaxRadius(const FHairStrandsDatas& In);
-
 struct HAIRSTRANDSCORE_API FHairStrandsBulkData : FHairStrandsBulkCommon
 {
 	enum EDataFlags
@@ -599,126 +633,146 @@
 	} Data;
 };
 
-/** Hair strands debug data */
-struct HAIRSTRANDSCORE_API FHairStrandsDebugDatas
+///////////////////////////////////////////////////////////////////////////////////////////////////
+// Interpolation Data/Bulk
+
+/** Hair strands points interpolation attributes */
+struct HAIRSTRANDSCORE_API FHairStrandsInterpolationDatas
 {
-	bool IsValid() const { return VoxelData.Num() > 0;  }
+	/** Set the number of interpolated points */
+	void SetNum(const uint32 NumPoints);
 
-	static const uint32 InvalidIndex = ~0u;
-	struct FOffsetAndCount
+	/** Reset the interpolated points to 0 */
+	void Reset();
+
+	/** Get the number of interpolated points */
+	uint32 Num() const { return PointsSimCurvesVertexIndex.Num(); }
+
+	bool IsValid() const { return PointsSimCurvesIndex.Num() > 0; }
+
+	/** Simulation curve indices, ordered by closest influence */
+	TArray<FIntVector> PointsSimCurvesIndex;
+
+	/** Closest vertex indices on simulation curve, ordered by closest influence */
+	TArray<FIntVector> PointsSimCurvesVertexIndex;
+
+	/** Lerp value between the closest vertex indices and the next one, ordered by closest influence */
+	TArray<FVector3f> PointsSimCurvesVertexLerp;
+
+	/** Weight of vertex indices on simulation curve, ordered by closest influence */
+	TArray<FVector3f>	PointsSimCurvesVertexWeights;
+
+	/** True, if interpolation data are built using a single guide */
+	bool bUseUniqueGuide = false;
+};
+
+struct HAIRSTRANDSCORE_API FHairStrandsInterpolationBulkData : FHairStrandsBulkCommon
+{
+	enum EDataFlags
 	{
-		uint32 Offset = 0u;
-		uint32 Count = 0u;
+		DataFlags_HasData = 1,
+		DataFlags_HasSingleGuideData = 2,
 	};
 
-	struct FVoxel
-	{
-		uint32 Index0 = InvalidIndex;
-		uint32 Index1 = InvalidIndex;
-	};
+	void Reset();
+	virtual void SerializeHeader(FArchive& Ar, UObject* Owner) override;
 
-	struct FDesc
-	{
-		FVector3f VoxelMinBound = FVector3f::ZeroVector;
-		FVector3f VoxelMaxBound = FVector3f::ZeroVector;
-		FIntVector VoxelResolution = FIntVector::ZeroValue;
-		float VoxelSize = 0;
-		uint32 MaxSegmentPerVoxel = 0;
-	};
-
-	FDesc VoxelDescription;
-	TArray<FOffsetAndCount> VoxelOffsetAndCount;
-	TArray<FVoxel> VoxelData;
-
-	struct FResources
-	{
-		FDesc VoxelDescription;
-
-		TRefCountPtr<FRDGPooledBuffer> VoxelOffsetAndCount;
-		TRefCountPtr<FRDGPooledBuffer> VoxelData;
-	};
-};
-
-/* Bulk data for root resources (GPU resources are stored into FHairStrandsRootResources) */
-struct FHairStrandsRootBulkData
-{
-	void SerializeHeader(FArchive& Ar, UObject* Owner);
-	void SerializeData(FArchive& Ar, UObject* Owner);
-	void Serialize(FArchive& Ar, UObject* Owner);
-
-	void Reset();
-	bool IsValid() const { return Header.RootCount > 0; }
-	const TArray<uint32>& GetValidSectionIndices(int32 LODIndex) const;
-	uint32 GetLODCount() const { return Header.LODs.Num(); }
-	uint32 GetRootCount()const { return Header.RootCount; }
-
-	uint32 GetDataSize() const
-	{
-		uint32 Total = 0;
-		for (const FData::FLOD& LOD : Data.LODs)
-		{
-			Total += LOD.UniqueTriangleIndexBuffer.IsBulkDataLoaded() ?			LOD.UniqueTriangleIndexBuffer.GetBulkDataSize() : 0u;
-			Total += LOD.RootToUniqueTriangleIndexBuffer.IsBulkDataLoaded() ?	LOD.RootToUniqueTriangleIndexBuffer.GetBulkDataSize() : 0u;
-			Total += LOD.RootBarycentricBuffer.IsBulkDataLoaded() ?				LOD.RootBarycentricBuffer.GetBulkDataSize() : 0u;
-			Total += LOD.RestUniqueTrianglePositionBuffer.IsBulkDataLoaded() ?	LOD.RestUniqueTrianglePositionBuffer.GetBulkDataSize() : 0u;
-			Total += LOD.MeshInterpolationWeightsBuffer.IsBulkDataLoaded() ?	LOD.MeshInterpolationWeightsBuffer.GetBulkDataSize() : 0u;
-			Total += LOD.MeshSampleIndicesBuffer.IsBulkDataLoaded() ?			LOD.MeshSampleIndicesBuffer.GetBulkDataSize() : 0u;
-			Total += LOD.RestSamplePositionsBuffer.IsBulkDataLoaded() ?			LOD.RestSamplePositionsBuffer.GetBulkDataSize() : 0u;
-		}
-		return Total;
-	}
+	virtual uint32 GetResourceCount() const override;
+	virtual void GetResources(FQuery& Out) override;
+	uint32 GetPointCount() const { return Header.PointCount; };
 
 	struct FHeader
-	{		
-		struct FLOD
-		{
-			int32  LODIndex = -1;
-			uint32 UniqueTriangleCount = 0;			
-			uint32 SampleCount = 0; 				// Number of samples used for the mesh interpolation
-			TArray<uint32> UniqueSectionIndices; 	// Store the mesh section indices which are relevant for this root LOD data
-		};
+	{
+		uint32 Flags = 0;
+		uint32 PointCount = 0;
+		uint32 SimPointCount = 0;
 
 		struct FStrides
 		{
-			uint32 RootToUniqueTriangleIndexBufferStride = 0;
-			uint32 RootBarycentricBufferStride = 0;
-			uint32 UniqueTriangleIndexBufferStride = 0;
-			uint32 RestUniqueTrianglePositionBufferStride = 0;
-
-			uint32 MeshInterpolationWeightsBufferStride = 0;
-			uint32 MeshSampleIndicesBufferStride = 0;
-			uint32 RestSamplePositionsBufferStride = 0;
-		};
-
-		uint32 RootCount = 0;						// Number of roots
-		uint32 PointCount = 0;						// Number of control points
-		FStrides Strides;
-		TArray<FLOD> LODs;
+			uint32 InterpolationStride = 0;
+			uint32 SimRootPointIndexStride = 0;
+		} Strides;
 	} Header;
 
 	struct FData
 	{
-		struct FLOD
-		{
-			// Binding
-			FByteBulkData RootToUniqueTriangleIndexBuffer; 	// Map each root onto the unique triangle Id (per-root)
-			FByteBulkData RootBarycentricBuffer; 			// Root's barycentric (per-root)
-			FByteBulkData UniqueTriangleIndexBuffer; 		// Unique triangles list from skeleton mesh section IDs and triangle IDs (per-unique-triangle)
-			FByteBulkData RestUniqueTrianglePositionBuffer;	// Rest triangle positions (per-unique-triangle)
-
-			// RBF
-			FByteBulkData MeshInterpolationWeightsBuffer; 	// Store the hair interpolation weights | Size = SamplesCount * SamplesCount (per-sample
-			FByteBulkData MeshSampleIndicesBuffer; 			// Store the samples vertex indices (per-sample)
-			FByteBulkData RestSamplePositionsBuffer; 		// Store the samples rest positions (per-sample)
-		};
-		TArray<FLOD> LODs;
+		FHairBulkContainer Interpolation;		// FHairStrandsInterpolationFormat  - Per-rendering-vertex interpolation data (closest guides, weight factors, ...). Data for a 1 or 3 guide(s))
+		FHairBulkContainer SimRootPointIndex;	// FHairStrandsRootIndexFormat      - Per-rendering-vertex index of the sim-root vertex
 	} Data;
 };
 
+///////////////////////////////////////////////////////////////////////////////////////////////////
+// Cluster Data/Bulk
+
+struct HAIRSTRANDSCORE_API FHairStrandsClusterCullingData
+{
+	void Reset();
+	bool IsValid() const { return ClusterCount > 0 && VertexCount > 0; }
+
+	/* Set LOD visibility, allowing to remove the simulation/rendering of certain LOD */
+	TArray<bool>				LODVisibility;
+
+	/* Screen size at which LOD should switches on CPU */
+	TArray<float>				CPULODScreenSize;
+
+	/* LOD info for the various clusters for LOD management on GPU */
+	TArray<FHairClusterInfo>	ClusterInfos;
+	TArray<FHairClusterLODInfo> ClusterLODInfos;
+	TArray<uint32>				VertexToClusterIds;
+	TArray<uint32>				ClusterVertexIds;
+	TArray<FHairLODInfo>		LODInfos;
+
+	/* Number of cluster  */
+	uint32 ClusterCount = 0;
+
+	/* Number of vertex  */
+	uint32 VertexCount = 0;
+};
+
+struct HAIRSTRANDSCORE_API FHairStrandsClusterCullingBulkData : FHairStrandsBulkCommon
+{
+	void Reset();
+
+	virtual void SerializeHeader(FArchive& Ar, UObject* Owner) override;
+	virtual uint32 GetResourceCount() const override;
+	virtual void GetResources(FQuery& Out) override;
+
+	bool IsValid() const { return Header.ClusterCount > 0 && Header.VertexCount > 0; }
+	void Validate(bool bIsSaving);
+
+	struct FHeader
+	{
+		/* Set LOD visibility, allowing to remove the simulation/rendering of certain LOD */
+		TArray<bool> LODVisibility;
+	
+		/* Screen size at which LOD should switches on CPU */
+		TArray<float> CPULODScreenSize;
+	
+		/* Curve count and Point count per LOD */
+		TArray<FHairLODInfo> LODInfos;
+	
+		uint32 ClusterCount = 0;
+		uint32 ClusterLODCount = 0;
+		uint32 VertexCount = 0;
+		uint32 VertexLODCount = 0;
+	} Header;
+
+	struct FData
+	{
+		/* LOD info for the various clusters for LOD management on GPU */
+		FHairBulkContainer	PackedClusterInfos;		// Size - ClusterCount
+		FHairBulkContainer	ClusterLODInfos;		// Size - ClusterLODCount
+		FHairBulkContainer	VertexToClusterIds;		// Size - VertexCount
+		FHairBulkContainer	ClusterVertexIds;		// Size - VertexLODCount
+	} Data;
+};
+
+///////////////////////////////////////////////////////////////////////////////////////////////////
+// Root Data/Bulk
+
 /* Source data for building root bulk data */
 struct FHairStrandsRootData
 {
-	FHairStrandsRootData();
 	void Reset();
 	bool HasProjectionData() const;
 	bool IsValid() const { return RootCount > 0; }
@@ -729,7 +783,7 @@
 
 		/* Triangle on which a root is attached */
 		/* When the projection is done with source to target mesh transfer, the projection indices does not match.
-		   In this case we need to separate index computation. The barycentric coords remain the same however. */
+			In this case we need to separate index computation. The barycentric coords remain the same however. */
 		TArray<FHairStrandsRootToUniqueTriangleIndexFormat::Type> RootToUniqueTriangleIndexBuffer;
 		TArray<FHairStrandsRootBarycentricFormat::Type> RootBarycentricBuffer;
 
@@ -763,149 +817,121 @@
 	TArray<FMeshProjectionLOD> MeshProjectionLODs;
 };
 
-/* Structure describing the LOD settings (Screen size, vertex info, ...) for each clusters.
-	The packed version of this structure corresponds to the GPU data layout (HairStrandsClusterCommon.ush)
-	This uses by the GPU LOD selection. */
-struct FHairClusterInfo
+/* Bulk data for root resources (GPU resources are stored into FHairStrandsRootResources) */
+struct FHairStrandsRootBulkData
 {
-	static const uint32 MaxLOD = 8;
+	void SerializeHeader(FArchive& Ar, UObject* Owner);
+	void SerializeData(FArchive& Ar, UObject* Owner);
+	void Serialize(FArchive& Ar, UObject* Owner);
 
-	struct Packed
+	void Reset();
+	bool IsValid() const { return Header.RootCount > 0; }
+	const TArray<uint32>& GetValidSectionIndices(int32 LODIndex) const;
+	uint32 GetLODCount() const { return Header.LODs.Num(); }
+	uint32 GetRootCount()const { return Header.RootCount; }
+
+	uint32 GetDataSize() const
 	{
-		uint32 LODInfoOffset : 24;
-		uint32 LODCount : 8;
-
-		uint32 LOD_ScreenSize_0 : 10;
-		uint32 LOD_ScreenSize_1 : 10;
-		uint32 LOD_ScreenSize_2 : 10;
-		uint32 Pad0 : 2;
-
-		uint32 LOD_ScreenSize_3 : 10;
-		uint32 LOD_ScreenSize_4 : 10;
-		uint32 LOD_ScreenSize_5 : 10;
-		uint32 Pad1 : 2;
-
-		uint32 LOD_ScreenSize_6 : 10;
-		uint32 LOD_ScreenSize_7 : 10;
-		uint32 LOD_bIsVisible : 8;
-		uint32 Pad2 : 4;
-	};
-	typedef FUintVector4 BulkType;
-
-	FHairClusterInfo()
-	{
-		for (uint32 LODIt = 0; LODIt < MaxLOD; ++LODIt)
+		uint32 Total = 0;
+		for (const FData::FLOD& LOD : Data.LODs)
 		{
-			ScreenSize[LODIt] = 0;
-			bIsVisible[LODIt] = true;
+			Total += LOD.UniqueTriangleIndexBuffer.IsBulkDataLoaded() ?			LOD.UniqueTriangleIndexBuffer.GetBulkDataSize() : 0u;
+			Total += LOD.RootToUniqueTriangleIndexBuffer.IsBulkDataLoaded() ?	LOD.RootToUniqueTriangleIndexBuffer.GetBulkDataSize() : 0u;
+			Total += LOD.RootBarycentricBuffer.IsBulkDataLoaded() ?				LOD.RootBarycentricBuffer.GetBulkDataSize() : 0u;
+			Total += LOD.RestUniqueTrianglePositionBuffer.IsBulkDataLoaded() ?	LOD.RestUniqueTrianglePositionBuffer.GetBulkDataSize() : 0u;
+			Total += LOD.MeshInterpolationWeightsBuffer.IsBulkDataLoaded() ?	LOD.MeshInterpolationWeightsBuffer.GetBulkDataSize() : 0u;
+			Total += LOD.MeshSampleIndicesBuffer.IsBulkDataLoaded() ?			LOD.MeshSampleIndicesBuffer.GetBulkDataSize() : 0u;
+			Total += LOD.RestSamplePositionsBuffer.IsBulkDataLoaded() ?			LOD.RestSamplePositionsBuffer.GetBulkDataSize() : 0u;
 		}
+		return Total;
 	}
 
-	uint32 LODCount = 0;
-	uint32 LODInfoOffset = 0;
-	TStaticArray<float,MaxLOD> ScreenSize;
-	TStaticArray<bool, MaxLOD> bIsVisible;
-};
-
-/* Structure describing the LOD settings common to all clusters. The layout of this structure is
-	identical the GPU data layout (HairStrandsClusterCommon.ush). This uses by the GPU LOD selection. */
-struct FHairClusterLODInfo
-{
-	uint32 VertexOffset = 0;
-	uint32 VertexCount0 = 0;
-	uint32 VertexCount1 = 0;
-	FFloat16 RadiusScale0 = 0;
-	FFloat16 RadiusScale1 = 0;
-};
-
-struct FHairClusterInfoFormat
-{
-	typedef FHairClusterInfo::Packed Type;
-	typedef FHairClusterInfo::Packed BulkType;
-	static const uint32 SizeInByte = sizeof(Type);
-};
-
-struct FHairClusterLODInfoFormat
-{
-	typedef FHairClusterLODInfo Type;
-	typedef FHairClusterLODInfo BulkType;
-	static const uint32 SizeInByte = sizeof(Type);
-};
-
-struct FHairClusterIndexFormat
-{
-	typedef uint32 Type;
-	typedef uint32 BulkType;
-	static const uint32 SizeInByte = sizeof(Type);
-	static const EPixelFormat Format = PF_R32_UINT;
-};
-
-struct FHairLODInfo
-{
-	uint32 CurveCount = 0;
-	uint32 PointCount = 0;
-};
-
-struct HAIRSTRANDSCORE_API FHairStrandsClusterCullingData
-{
-	void Reset();
-	bool IsValid() const { return ClusterCount > 0 && VertexCount > 0; }
-
-	/* Set LOD visibility, allowing to remove the simulation/rendering of certain LOD */
-	TArray<bool>				LODVisibility;
-
-	/* Screen size at which LOD should switches on CPU */
-	TArray<float>				CPULODScreenSize;
-
-	/* LOD info for the various clusters for LOD management on GPU */
-	TArray<FHairClusterInfo>	ClusterInfos;
-	TArray<FHairClusterLODInfo> ClusterLODInfos;
-	TArray<uint32>				VertexToClusterIds;
-	TArray<uint32>				ClusterVertexIds;
-	TArray<FHairLODInfo>		LODInfos;
-
-	/* Number of cluster  */
-	uint32 ClusterCount = 0;
-
-	/* Number of vertex  */
-	uint32 VertexCount = 0;
-};
-
-struct HAIRSTRANDSCORE_API FHairStrandsClusterCullingBulkData : FHairStrandsBulkCommon
-{
-	void Reset();
-
-	virtual void SerializeHeader(FArchive& Ar, UObject* Owner) override;
-	virtual uint32 GetResourceCount() const override;
-	virtual void GetResources(FQuery& Out) override;
-
-	bool IsValid() const { return Header.ClusterCount > 0 && Header.VertexCount > 0; }
-	void Validate(bool bIsSaving);
-
 	struct FHeader
-	{
-		/* Set LOD visibility, allowing to remove the simulation/rendering of certain LOD */
-		TArray<bool> LODVisibility;
-	
-		/* Screen size at which LOD should switches on CPU */
-		TArray<float> CPULODScreenSize;
-	
-		/* Curve count and Point count per LOD */
-		TArray<FHairLODInfo> LODInfos;
-	
-		uint32 ClusterCount = 0;
-		uint32 ClusterLODCount = 0;
-		uint32 VertexCount = 0;
-		uint32 VertexLODCount = 0;
+	{		
+		struct FLOD
+		{
+			int32  LODIndex = -1;
+			uint32 UniqueTriangleCount = 0;			
+			uint32 SampleCount = 0; 				// Number of samples used for the mesh interpolation
+			TArray<uint32> UniqueSectionIndices; 	// Store the mesh section indices which are relevant for this root LOD data
+		};
+
+		struct FStrides
+		{
+			uint32 RootToUniqueTriangleIndexBufferStride = 0;
+			uint32 RootBarycentricBufferStride = 0;
+			uint32 UniqueTriangleIndexBufferStride = 0;
+			uint32 RestUniqueTrianglePositionBufferStride = 0;
+
+			uint32 MeshInterpolationWeightsBufferStride = 0;
+			uint32 MeshSampleIndicesBufferStride = 0;
+			uint32 RestSamplePositionsBufferStride = 0;
+		};
+
+		uint32 RootCount = 0;						// Number of roots
+		uint32 PointCount = 0;						// Number of control points
+		FStrides Strides;
+		TArray<FLOD> LODs;
 	} Header;
 
 	struct FData
 	{
-		/* LOD info for the various clusters for LOD management on GPU */
-		FHairBulkContainer	PackedClusterInfos;		// Size - ClusterCount
-		FHairBulkContainer	ClusterLODInfos;		// Size - ClusterLODCount
-		FHairBulkContainer	VertexToClusterIds;		// Size - VertexCount
-		FHairBulkContainer	ClusterVertexIds;		// Size - VertexLODCount
+		struct FLOD
+		{
+			// Binding
+			FByteBulkData RootToUniqueTriangleIndexBuffer; 	// Map each root onto the unique triangle Id (per-root)
+			FByteBulkData RootBarycentricBuffer; 			// Root's barycentric (per-root)
+			FByteBulkData UniqueTriangleIndexBuffer; 		// Unique triangles list from skeleton mesh section IDs and triangle IDs (per-unique-triangle)
+			FByteBulkData RestUniqueTrianglePositionBuffer;	// Rest triangle positions (per-unique-triangle)
+
+			// RBF
+			FByteBulkData MeshInterpolationWeightsBuffer; 	// Store the hair interpolation weights | Size = SamplesCount * SamplesCount (per-sample
+			FByteBulkData MeshSampleIndicesBuffer; 			// Store the samples vertex indices (per-sample)
+			FByteBulkData RestSamplePositionsBuffer; 		// Store the samples rest positions (per-sample)
+		};
+		TArray<FLOD> LODs;
 	} Data;
 };
 
+///////////////////////////////////////////////////////////////////////////////////////////////////
+// Hair debug data
+
+/** Hair strands debug data */
+struct HAIRSTRANDSCORE_API FHairStrandsDebugDatas
+{
+	bool IsValid() const { return VoxelData.Num() > 0;  }
+
+	static const uint32 InvalidIndex = ~0u;
+	struct FOffsetAndCount
+	{
+		uint32 Offset = 0u;
+		uint32 Count = 0u;
+	};
+
+	struct FVoxel
+	{
+		uint32 Index0 = InvalidIndex;
+		uint32 Index1 = InvalidIndex;
+	};
+
+	struct FDesc
+	{
+		FVector3f VoxelMinBound = FVector3f::ZeroVector;
+		FVector3f VoxelMaxBound = FVector3f::ZeroVector;
+		FIntVector VoxelResolution = FIntVector::ZeroValue;
+		float VoxelSize = 0;
+		uint32 MaxSegmentPerVoxel = 0;
+	};
+
+	FDesc VoxelDescription;
+	TArray<FOffsetAndCount> VoxelOffsetAndCount;
+	TArray<FVoxel> VoxelData;
+
+	struct FResources
+	{
+		FDesc VoxelDescription;
+
+		TRefCountPtr<FRDGPooledBuffer> VoxelOffsetAndCount;
+		TRefCountPtr<FRDGPooledBuffer> VoxelData;
+	};
+};
diff --git a/Engine/Shaders/Shared/HairStrandsDefinitions.h b/Engine/Shaders/Shared/HairStrandsDefinitions.h
--- a/Engine/Shaders/Shared/HairStrandsDefinitions.h
+++ b/Engine/Shaders/Shared/HairStrandsDefinitions.h
@@ -38,6 +38,9 @@
 // Max number of discrete LOD that a hair group can have
 #define MAX_HAIR_LOD 8
 
+// Max split for raytracing geometry
+#define STRANDS_PROCEDURAL_INTERSECTOR_MAX_SPLITS 4
+
 // Use triangle strip for HW raster path
 #define USE_HAIR_TRIANGLE_STRIP 0
 
