Change hair bulk data to support LOD index request.

diff --git a/Engine/Plugins/Runtime/HairStrands/Source/HairStrandsCore/Private/GroomAsset.cpp b/Engine/Plugins/Runtime/HairStrands/Source/HairStrandsCore/Private/GroomAsset.cpp
--- a/Engine/Plugins/Runtime/HairStrands/Source/HairStrandsCore/Private/GroomAsset.cpp
+++ b/Engine/Plugins/Runtime/HairStrands/Source/HairStrandsCore/Private/GroomAsset.cpp
@@ -1775,9 +1775,9 @@ static void InternalSerializeStrand(FArchive& Ar, UObject* Owner, FHairGroupPlat
 	// from the 'start' (i.e., without offset)
 	if (Ar.IsCooking() && Ar.IsSaving())
 	{
-		{ FHairStreamingRequest R; R.Request(HAIR_MAX_NUM_CURVE_PER_GROUP, HAIR_MAX_NUM_POINT_PER_GROUP, StrandData.BulkData,               true /*bWait*/, true /*bFillBulkdata*/, Owner->GetFName()); }
-		{ FHairStreamingRequest R; R.Request(HAIR_MAX_NUM_CURVE_PER_GROUP, HAIR_MAX_NUM_POINT_PER_GROUP, StrandData.InterpolationBulkData,  true /*bWait*/, true /*bFillBulkdata*/, Owner->GetFName()); }
-		{ FHairStreamingRequest R; R.Request(HAIR_MAX_NUM_CURVE_PER_GROUP, HAIR_MAX_NUM_POINT_PER_GROUP, StrandData.ClusterCullingBulkData, true /*bWait*/, true /*bFillBulkdata*/, Owner->GetFName()); }
+		{ FHairStreamingRequest R; R.Request(HAIR_MAX_NUM_CURVE_PER_GROUP, HAIR_MAX_NUM_POINT_PER_GROUP, -1/*LODIndex*/, StrandData.BulkData,               true /*bWait*/, true /*bFillBulkdata*/, Owner->GetFName()); }
+		{ FHairStreamingRequest R; R.Request(HAIR_MAX_NUM_CURVE_PER_GROUP, HAIR_MAX_NUM_POINT_PER_GROUP, -1/*LODIndex*/, StrandData.InterpolationBulkData,  true /*bWait*/, true /*bFillBulkdata*/, Owner->GetFName()); }
+		{ FHairStreamingRequest R; R.Request(HAIR_MAX_NUM_CURVE_PER_GROUP, HAIR_MAX_NUM_POINT_PER_GROUP, -1/*LODIndex*/, StrandData.ClusterCullingBulkData, true /*bWait*/, true /*bFillBulkdata*/, Owner->GetFName()); }
 	}
 
 	if (!Ar.IsCooking() || !StrandData.bIsCookedOut)
diff --git a/Engine/Plugins/Runtime/HairStrands/Source/HairStrandsCore/Private/GroomManager.cpp b/Engine/Plugins/Runtime/HairStrands/Source/HairStrandsCore/Private/GroomManager.cpp
--- a/Engine/Plugins/Runtime/HairStrands/Source/HairStrandsCore/Private/GroomManager.cpp
+++ b/Engine/Plugins/Runtime/HairStrands/Source/HairStrandsCore/Private/GroomManager.cpp
@@ -703,7 +703,7 @@ void AddHairStreamingRequest(FHairGroupInstance* Instance, int32 InLODIndex)
 		// Note: Allocation will only be done if the resources is not initialized yet. Guides deformed position are also initialized from the Rest position at creation time.
 		if (Instance->Guides.Data && bLODNeedsGuides)
 		{
-			if (Instance->Guides.RestRootResource)			{ Instance->Guides.RestRootResource->StreamInData(); Instance->Guides.RestRootResource->StreamInLODData(MeshLODIndex); }
+			if (Instance->Guides.RestRootResource)			{ Instance->Guides.RestRootResource->StreamInData(MeshLODIndex);}
 			if (Instance->Guides.RestResource)				{ Instance->Guides.RestResource->StreamInData(); }
 		}
 
@@ -717,13 +717,13 @@ void AddHairStreamingRequest(FHairGroupInstance* Instance, int32 InLODIndex)
 
 			if (InstanceLOD.InterpolationResource)			{ InstanceLOD.InterpolationResource->StreamInData(); }
 			if (InstanceLOD.DeformedResource)				{ InstanceLOD.DeformedResource->StreamInData(); }
-			if (InstanceLOD.Guides.RestRootResource)		{ InstanceLOD.Guides.RestRootResource->StreamInData(); InstanceLOD.Guides.RestRootResource->StreamInLODData(MeshLODIndex); }
+			if (InstanceLOD.Guides.RestRootResource)		{ InstanceLOD.Guides.RestRootResource->StreamInData(MeshLODIndex); }
 			if (InstanceLOD.Guides.RestResource)			{ InstanceLOD.Guides.RestResource->StreamInData(); }
 			if (InstanceLOD.Guides.InterpolationResource)	{ InstanceLOD.Guides.InterpolationResource->StreamInData(); }
 		}
 		else if (GeometryType == EHairGeometryType::Strands)
 		{
-			if (Instance->Strands.RestRootResource)			{ Instance->Strands.RestRootResource->StreamInData(); Instance->Strands.RestRootResource->StreamInLODData(MeshLODIndex); }
+			if (Instance->Strands.RestRootResource)			{ Instance->Strands.RestRootResource->StreamInData(MeshLODIndex); }
 			if (Instance->Strands.RestResource)				{ Instance->Strands.RestResource->StreamInData(); }
 			if (Instance->Strands.ClusterCullingResource)	{ Instance->Strands.ClusterCullingResource->StreamInData(); }
 			if (Instance->Strands.InterpolationResource)	{ Instance->Strands.InterpolationResource->StreamInData(); }
@@ -958,9 +958,9 @@ static void RunHairLODSelection(
 			// Note: Allocation will only be done if the resources is not initialized yet. Guides deformed position are also initialized from the Rest position at creation time.
 			if (Instance->Guides.Data && bLODNeedsGuides)
 			{
-				if (Instance->Guides.RestRootResource)			{ Instance->Guides.RestRootResource->Allocate(GraphBuilder, LoadingType, ResourceStatus); Instance->Guides.RestRootResource->AllocateLOD(GraphBuilder, MeshLODIndex, LoadingType, ResourceStatus); }
+				if (Instance->Guides.RestRootResource)			{ Instance->Guides.RestRootResource->Allocate(GraphBuilder, LoadingType, ResourceStatus, MeshLODIndex); }
 				if (Instance->Guides.RestResource)				{ Instance->Guides.RestResource->Allocate(GraphBuilder, LoadingType, ResourceStatus); }
-				if (Instance->Guides.DeformedRootResource)		{ Instance->Guides.DeformedRootResource->Allocate(GraphBuilder, LoadingType, ResourceStatus); Instance->Guides.DeformedRootResource->AllocateLOD(GraphBuilder, MeshLODIndex, LoadingType, ResourceStatus); }
+				if (Instance->Guides.DeformedRootResource)		{ Instance->Guides.DeformedRootResource->Allocate(GraphBuilder, LoadingType, ResourceStatus, MeshLODIndex); }
 				if (Instance->Guides.DeformedResource)			
 				{ 
 					// Ensure the rest resources are correctly loaded prior to initialized and copy the rest position into the deformed positions
@@ -993,9 +993,9 @@ static void RunHairLODSelection(
 
 				if (InstanceLOD.InterpolationResource)			{ InstanceLOD.InterpolationResource->Allocate(GraphBuilder, LoadingType, ResourceStatus); }
 				if (InstanceLOD.DeformedResource)				{ InstanceLOD.DeformedResource->Allocate(GraphBuilder, LoadingType, ResourceStatus); }
-				if (InstanceLOD.Guides.RestRootResource)		{ InstanceLOD.Guides.RestRootResource->Allocate(GraphBuilder, LoadingType, ResourceStatus); InstanceLOD.Guides.RestRootResource->AllocateLOD(GraphBuilder, MeshLODIndex, LoadingType, ResourceStatus); }
+				if (InstanceLOD.Guides.RestRootResource)		{ InstanceLOD.Guides.RestRootResource->Allocate(GraphBuilder, LoadingType, ResourceStatus, MeshLODIndex); }
 				if (InstanceLOD.Guides.RestResource)			{ InstanceLOD.Guides.RestResource->Allocate(GraphBuilder, LoadingType, ResourceStatus); }
-				if (InstanceLOD.Guides.DeformedRootResource)	{ InstanceLOD.Guides.DeformedRootResource->Allocate(GraphBuilder, LoadingType, ResourceStatus); InstanceLOD.Guides.DeformedRootResource->AllocateLOD(GraphBuilder, MeshLODIndex, LoadingType, ResourceStatus); }
+				if (InstanceLOD.Guides.DeformedRootResource)	{ InstanceLOD.Guides.DeformedRootResource->Allocate(GraphBuilder, LoadingType, ResourceStatus, MeshLODIndex); }
 				if (InstanceLOD.Guides.DeformedResource)		{ InstanceLOD.Guides.DeformedResource->Allocate(GraphBuilder, LoadingType, ResourceStatus); }
 				if (InstanceLOD.Guides.InterpolationResource)	{ InstanceLOD.Guides.InterpolationResource->Allocate(GraphBuilder, LoadingType, ResourceStatus); }
 				#if RHI_RAYTRACING
@@ -1008,12 +1008,12 @@ static void RunHairLODSelection(
 				check(Instance->HairGroupPublicData);
 				Instance->HairGroupPublicData->Allocate(GraphBuilder);
 
-				if (Instance->Strands.RestRootResource)			{ Instance->Strands.RestRootResource->Allocate(GraphBuilder, LoadingType, ResourceStatus); Instance->Strands.RestRootResource->AllocateLOD(GraphBuilder, MeshLODIndex, LoadingType, ResourceStatus); }
+				if (Instance->Strands.RestRootResource)			{ Instance->Strands.RestRootResource->Allocate(GraphBuilder, LoadingType, ResourceStatus, MeshLODIndex); }
 				if (Instance->Strands.RestResource)				{ Instance->Strands.RestResource->Allocate(GraphBuilder, EHairResourceLoadingType::Async, ResourceStatus, RequestedCurveCount, RequestedPointCount); }
 				if (Instance->Strands.ClusterCullingResource)	{ Instance->Strands.ClusterCullingResource->Allocate(GraphBuilder, EHairResourceLoadingType::Async, ResourceStatus, RequestedCurveCount, RequestedPointCount); }
 				if (Instance->Strands.InterpolationResource)	{ Instance->Strands.InterpolationResource->Allocate(GraphBuilder, EHairResourceLoadingType::Async, ResourceStatus, RequestedCurveCount, RequestedPointCount); }
 
-				if (Instance->Strands.DeformedRootResource)		{ Instance->Strands.DeformedRootResource->Allocate(GraphBuilder, LoadingType, ResourceStatus); Instance->Strands.DeformedRootResource->AllocateLOD(GraphBuilder, MeshLODIndex, LoadingType, ResourceStatus); }
+				if (Instance->Strands.DeformedRootResource)		{ Instance->Strands.DeformedRootResource->Allocate(GraphBuilder, LoadingType, ResourceStatus, MeshLODIndex); }
 				if (Instance->Strands.DeformedResource)			{ Instance->Strands.DeformedResource->Allocate(GraphBuilder, LoadingType, ResourceStatus); }
 				#if RHI_RAYTRACING
 				if (bHasPathTracingView)						{ AllocateRaytracingResources(Instance); }
diff --git a/Engine/Plugins/Runtime/HairStrands/Source/HairStrandsCore/Private/GroomResources.cpp b/Engine/Plugins/Runtime/HairStrands/Source/HairStrandsCore/Private/GroomResources.cpp
--- a/Engine/Plugins/Runtime/HairStrands/Source/HairStrandsCore/Private/GroomResources.cpp
+++ b/Engine/Plugins/Runtime/HairStrands/Source/HairStrandsCore/Private/GroomResources.cpp
@@ -615,13 +615,14 @@ void FHairCommonResource::InitRHI()
 {
 	if (bIsInitialized || AllocationType == EHairStrandsAllocationType::Deferred || GUsingNullRHI) { return; }
 
-	check(InternalIsDataLoaded(HAIR_MAX_NUM_CURVE_PER_GROUP, HAIR_MAX_NUM_POINT_PER_GROUP));
+	const int32 DummyLODIndex = -1;
+	check(InternalIsDataLoaded(HAIR_MAX_NUM_CURVE_PER_GROUP, HAIR_MAX_NUM_POINT_PER_GROUP, DummyLODIndex));
 
 	if (bUseRenderGraph)
 	{
 		FRHICommandListImmediate& RHICmdList = FRHICommandListExecutor::GetImmediateCommandList();
 		FRDGBuilder GraphBuilder(RHICmdList);
-		InternalAllocate(GraphBuilder);
+		InternalAllocate(GraphBuilder, DummyLODIndex);
 		GraphBuilder.Execute();
 	}
 	else
@@ -642,15 +643,15 @@ void FHairCommonResource::Allocate(FRDGBuilder& GraphBuilder, EHairResourceLoadi
 	EHairResourceStatus Status;
 	Status.Status = EHairResourceStatus::EStatus::None;
 	Status.AvailableCurveCount = HAIR_MAX_NUM_CURVE_PER_GROUP;
-	Allocate(GraphBuilder, LoadingType, Status, HAIR_MAX_NUM_CURVE_PER_GROUP, HAIR_MAX_NUM_POINT_PER_GROUP);
+	Allocate(GraphBuilder, LoadingType, Status, HAIR_MAX_NUM_CURVE_PER_GROUP, HAIR_MAX_NUM_POINT_PER_GROUP, -1/*LODIndex*/);
 }
 
-void FHairCommonResource::Allocate(FRDGBuilder& GraphBuilder, EHairResourceLoadingType LoadingType, EHairResourceStatus& Status)
+void FHairCommonResource::Allocate(FRDGBuilder& GraphBuilder, EHairResourceLoadingType LoadingType, EHairResourceStatus& Status, int32 InLODIndex)
 {
-	Allocate(GraphBuilder, LoadingType, Status, HAIR_MAX_NUM_CURVE_PER_GROUP, HAIR_MAX_NUM_POINT_PER_GROUP);
+	Allocate(GraphBuilder, LoadingType, Status, HAIR_MAX_NUM_CURVE_PER_GROUP, HAIR_MAX_NUM_POINT_PER_GROUP, InLODIndex);
 }
 
-void FHairCommonResource::Allocate(FRDGBuilder& GraphBuilder, EHairResourceLoadingType LoadingType, EHairResourceStatus& Status, uint32 InRequestedCurveCount, uint32 InRequestedPointCount)
+void FHairCommonResource::Allocate(FRDGBuilder& GraphBuilder, EHairResourceLoadingType LoadingType, EHairResourceStatus& Status, uint32 InRequestedCurveCount, uint32 InRequestedPointCount, int32 InLODIndex)
 {
 	check(AllocationType == EHairStrandsAllocationType::Deferred);
 
@@ -661,30 +662,34 @@ void FHairCommonResource::Allocate(FRDGBuilder& GraphBuilder, EHairResourceLoadi
 		if (!bIsInitialized)
 		{
 			FRenderResource::InitResource(); // Call RenderResource InitResource() so that the resources is marked as initialized
-			InternalAllocate(GraphBuilder);
-			bIsInitialized = true;
-
-			MaxAvailableCurveCount = InRequestedCurveCount;
 		}
-		Status |= EHairResourceStatus::EStatus::Valid;
+		if (!bIsInitialized || !InternalIsLODDataLoaded(InLODIndex))
+		{
+			InternalAllocate(GraphBuilder, InLODIndex);
+		}
+		bIsInitialized = true;
+
+		// Update the max curve count available
+		MaxAvailableCurveCount = InRequestedCurveCount;
 
+		Status |= EHairResourceStatus::EStatus::Valid;
 	}
 	else if (LoadingType == EHairResourceLoadingType::Async)
 	{
 		// 1. If all requested curve are already loaded, nothing to do
-		if (bIsInitialized && MaxAvailableCurveCount >= InRequestedCurveCount) 
+		if (bIsInitialized && MaxAvailableCurveCount >= InRequestedCurveCount && InternalIsLODDataLoaded(InLODIndex)) 
 		{ 
 			Status |= EHairResourceStatus::EStatus::Valid; 
 		}
 		// 2. If more curves are requested, issue a streaming request
-		else if (InternalIsDataLoaded(InRequestedCurveCount, InRequestedPointCount))
+		else if (InternalIsDataLoaded(InRequestedCurveCount, InRequestedPointCount, InLODIndex))
 		{ 
 			// 2.1 Curve data are availble, and update GPU resources
 			if (!bIsInitialized)
 			{
 				FRenderResource::InitResource(); // Call RenderResource InitResource() so that the resources is marked as initialized
 			}
-			InternalAllocate(GraphBuilder);
+			InternalAllocate(GraphBuilder, InLODIndex);
 			bIsInitialized = true;
 
 			// Update the max curve count available
@@ -705,33 +710,12 @@ void FHairCommonResource::Allocate(FRDGBuilder& GraphBuilder, EHairResourceLoadi
 	Status.AddAvailableCurve(MaxAvailableCurveCount);
 }
 
-void FHairCommonResource::AllocateLOD(FRDGBuilder& GraphBuilder, int32 LODIndex, EHairResourceLoadingType LoadingType, EHairResourceStatus& Status)
-{
-	// When using a async loading, sub-resources allocation requires the common/main resource to be already initialized.
-	if (LoadingType == EHairResourceLoadingType::Async && (!InternalIsLODDataLoaded(LODIndex) || !bIsInitialized)) { Status |= EHairResourceStatus::EStatus::Loading; return; }
-
-	// Sanity check. 
-	check(AllocationType == EHairStrandsAllocationType::Deferred);
-
-	InternalAllocateLOD(GraphBuilder, LODIndex);
-	Status |= EHairResourceStatus::EStatus::Valid;
-}
-
-void FHairCommonResource::StreamInData()
-{
-	if (!bIsInitialized)
-	{
-		// TODO
-		InternalIsDataLoaded(HAIR_MAX_NUM_CURVE_PER_GROUP, HAIR_MAX_NUM_POINT_PER_GROUP);
-	}
-}
-
-void FHairCommonResource::StreamInLODData(int32 LODIndex)
+void FHairCommonResource::StreamInData(int32 InLODIndex)
 {
 	if (!bIsInitialized)
 	{
 		// TODO
-		InternalIsLODDataLoaded(LODIndex);
+		InternalIsDataLoaded(HAIR_MAX_NUM_CURVE_PER_GROUP, HAIR_MAX_NUM_POINT_PER_GROUP, InLODIndex);
 	}
 }
 
@@ -982,11 +966,11 @@ FHairStrandsRestResource::FHairStrandsRestResource(FHairStrandsBulkData& InBulkD
 	check(!!(BulkData.Header.Flags & FHairStrandsBulkData::DataFlags_HasData));
 }
 
-bool FHairStrandsRestResource::InternalIsDataLoaded(uint32 InRequestedCurveCount, uint32 InRequestedPointCount)
+bool FHairStrandsRestResource::InternalIsDataLoaded(uint32 InRequestedCurveCount, uint32 InRequestedPointCount, int32 InLODIndex)
 {
 	if (StreamingRequest.IsNone())
 	{
-		StreamingRequest.Request(InRequestedCurveCount, InRequestedPointCount, BulkData, false, false, OwnerName);
+		StreamingRequest.Request(InRequestedCurveCount, InRequestedPointCount, InLODIndex, BulkData, false, false, OwnerName);
 	}
 	return StreamingRequest.IsCompleted();
 }
@@ -1154,11 +1138,11 @@ FHairStrandsClusterCullingResource::FHairStrandsClusterCullingResource(FHairStra
 	MaxAvailableCurveCount = 0;
 }
 
-bool FHairStrandsClusterCullingResource::InternalIsDataLoaded(uint32 InRequestedCurveCount, uint32 InRequestedPointCount)
+bool FHairStrandsClusterCullingResource::InternalIsDataLoaded(uint32 InRequestedCurveCount, uint32 InRequestedPointCount, int32 InLODIndex)
 {
 	if (StreamingRequest.IsNone())
 	{
-		StreamingRequest.Request(InRequestedCurveCount, InRequestedPointCount, BulkData, true, true, OwnerName);
+		StreamingRequest.Request(InRequestedCurveCount, InRequestedPointCount, InLODIndex, BulkData, true, true, OwnerName);
 	}
 	return StreamingRequest.IsCompleted();
 }
@@ -1199,7 +1183,6 @@ FHairStrandsRestRootResource::FHairStrandsRestRootResource(FHairStrandsRootBulkD
 
 void FHairStrandsRestRootResource::PopulateFromRootData()
 {
-	uint32 LODIndex = 0;
 	LODs.Reserve(BulkData.Header.LODs.Num());
 	for (const FHairStrandsRootBulkData::FHeader::FLOD& InLOD : BulkData.Header.LODs)
 	{
@@ -1212,18 +1195,28 @@ void FHairStrandsRestRootResource::PopulateFromRootData()
 	LODRequests.SetNum(BulkData.Header.LODs.Num());
 }
 
-bool FHairStrandsRestRootResource::InternalIsLODDataLoaded(int32 LODIndex)
+bool FHairStrandsRestRootResource::InternalIsLODDataLoaded(int32 InLODIndex) const
+{
+	if (InLODIndex >= 0 && InLODIndex < LODs.Num())
+	{
+		const FLOD& GPUData = LODs[InLODIndex];
+		return GPUData.Status == FLOD::EStatus::Completed || GPUData.Status == FLOD::EStatus::Initialized;
+	}
+	return true;
+}
+
+bool FHairStrandsRestRootResource::InternalIsDataLoaded(uint32 InRequestedCurveCount, uint32 InRequestedPointCount, int32 InLODIndex)
 {
 	bool bIsLoading = false;
 	
 	check(LODs.Num() == BulkData.Header.LODs.Num());
-	if (LODIndex >= 0 && LODIndex < LODs.Num())
+	if (InLODIndex >= 0 && InLODIndex < LODs.Num())
 	{
-		FBulkDataBatchRequest& LODRequest = LODRequests[LODIndex];
+		FBulkDataBatchRequest& LODRequest = LODRequests[InLODIndex];
 		if (LODRequest.IsNone())
 		{
 			FBulkDataBatchRequest::FBatchBuilder Batch = FBulkDataBatchRequest::NewBatch(9);
-			FHairStrandsRootBulkData::FData::FLOD& CPUData = BulkData.Data.LODs[LODIndex];
+			FHairStrandsRootBulkData::FData::FLOD& CPUData = BulkData.Data.LODs[InLODIndex];
 			const bool bHasValidCPUData = CPUData.RootBarycentricBuffer.GetBulkDataSize() > 0;
 			if (bHasValidCPUData)
 			{
@@ -1253,7 +1246,7 @@ bool FHairStrandsRestRootResource::InternalIsLODDataLoaded(int32 LODIndex)
 	return !bIsLoading;
 }
 
-void FHairStrandsRestRootResource::InternalAllocate(FRDGBuilder& GraphBuilder)
+void FHairStrandsRestRootResource::InternalAllocate(FRDGBuilder& GraphBuilder, int32 InLODIndex)
 {
 	// Once empty, the MeshProjectionLODsneeds to be repopulate as it might be re-initialized. 
 	// E.g., when a resource is updated, it is first released, then re-init. 
@@ -1261,27 +1254,22 @@ void FHairStrandsRestRootResource::InternalAllocate(FRDGBuilder& GraphBuilder)
 	{
 		PopulateFromRootData();
 	}
-}
 
-void FHairStrandsRestRootResource::InternalAllocateLOD(FRDGBuilder& GraphBuilder, int32 LODIndex)
-{
-	// Sanity check ti insure that the 'common' part of FHairStrandsRestRootResource is already initialized
-	check(bIsInitialized);
 	check(BulkData.Header.PointCount > 0);
 	check(LODs.Num() == BulkData.Header.LODs.Num());
-	if (LODIndex >= 0 && LODIndex < LODs.Num())
+	if (InLODIndex >= 0 && InLODIndex < LODs.Num())
 	{
-		FLOD& GPUData = LODs[LODIndex];
+		FLOD& GPUData = LODs[InLODIndex];
 		const bool bIsLODInitialized = GPUData.Status == FLOD::EStatus::Completed || GPUData.Status == FLOD::EStatus::Initialized;
 		if (bIsLODInitialized)
 		{
 			return;
 		}
 		
-		LODRequests[LODIndex] = FBulkDataBatchRequest();
+		LODRequests[InLODIndex] = FBulkDataBatchRequest();
 
-		const FHairStrandsRootBulkData::FHeader::FLOD& LODHeader = BulkData.Header.LODs[LODIndex];
-		FHairStrandsRootBulkData::FData::FLOD& CPUData = BulkData.Data.LODs[LODIndex];
+		const FHairStrandsRootBulkData::FHeader::FLOD& LODHeader = BulkData.Header.LODs[InLODIndex];
+		FHairStrandsRootBulkData::FData::FLOD& CPUData = BulkData.Data.LODs[InLODIndex];
 		const bool bHasValidCPUData = CPUData.RootBarycentricBuffer.GetBulkDataSize() > 0;
 		if (bHasValidCPUData)
 		{
@@ -1355,7 +1343,6 @@ FHairStrandsDeformedRootResource::FHairStrandsDeformedRootResource(const FHairSt
 	CurveType(InCurveType)
 {
 	check(InRestResources);
-	uint32 LODIndex = 0;
 	RootCount = InRestResources->BulkData.Header.RootCount;
 	LODs.Reserve(InRestResources->LODs.Num());
 	for (const FHairStrandsRestRootResource::FLOD& InLOD : InRestResources->LODs)
@@ -1368,11 +1355,21 @@ FHairStrandsDeformedRootResource::FHairStrandsDeformedRootResource(const FHairSt
 	}
 }
 
-void FHairStrandsDeformedRootResource::InternalAllocateLOD(FRDGBuilder& GraphBuilder, int32 LODIndex)
+bool FHairStrandsDeformedRootResource::InternalIsLODDataLoaded(int32 InLODIndex) const
+{
+	if (InLODIndex >= 0 && InLODIndex < LODs.Num())
+	{
+		const FLOD& GPUData = LODs[InLODIndex];
+		return GPUData.Status == FLOD::EStatus::Completed || GPUData.Status == FLOD::EStatus::Initialized;
+	}
+	return true;
+}
+
+void FHairStrandsDeformedRootResource::InternalAllocate(FRDGBuilder& GraphBuilder, int32 InLODIndex)
 {
-	if (RootCount > 0 && LODIndex >= 0 && LODIndex < LODs.Num())
+	if (RootCount > 0 && InLODIndex >= 0 && InLODIndex < LODs.Num())
 	{
-		FLOD& LOD = LODs[LODIndex];
+		FLOD& LOD = LODs[InLODIndex];
 		if (LOD.Status == FLOD::EStatus::Invalid)
 		{		
 			LOD.Status = FLOD::EStatus::Initialized;
@@ -1432,11 +1429,11 @@ FHairStrandsInterpolationResource::FHairStrandsInterpolationResource(FHairStrand
 	check(!!(BulkData.Header.Flags & FHairStrandsInterpolationBulkData::DataFlags_HasData));
 }
 
-bool FHairStrandsInterpolationResource::InternalIsDataLoaded(uint32 InRequestedCurveCount, uint32 InRequestedPointCount)
+bool FHairStrandsInterpolationResource::InternalIsDataLoaded(uint32 InRequestedCurveCount, uint32 InRequestedPointCount, int32 InLODIndex)
 {
 	if (StreamingRequest.IsNone())
 	{
-		StreamingRequest.Request(InRequestedCurveCount, InRequestedPointCount, BulkData, false, false, OwnerName);
+		StreamingRequest.Request(InRequestedCurveCount, InRequestedPointCount, InLODIndex, BulkData, false, false, OwnerName);
 	}
 	return StreamingRequest.IsCompleted();
 }
diff --git a/Engine/Plugins/Runtime/HairStrands/Source/HairStrandsCore/Private/HairStrandsDatas.cpp b/Engine/Plugins/Runtime/HairStrands/Source/HairStrandsCore/Private/HairStrandsDatas.cpp
--- a/Engine/Plugins/Runtime/HairStrands/Source/HairStrandsCore/Private/HairStrandsDatas.cpp
+++ b/Engine/Plugins/Runtime/HairStrands/Source/HairStrandsCore/Private/HairStrandsDatas.cpp
@@ -311,7 +311,7 @@ bool FHairStreamingRequest::IsCompleted()
 // Request fullfil 2 use cases:
 // * Load DDC data and upload them to GPU
 // * Load DDC data and store them into bulkdata for serialization
-void FHairStreamingRequest::Request(uint32 InRequestedCurveCount, uint32 InRequestedPointCount, FHairStrandsBulkCommon& In, bool bWait, bool bFillBulkdata, const FName& InOwnerName)
+void FHairStreamingRequest::Request(uint32 InRequestedCurveCount, uint32 InRequestedPointCount, int32 InLODIndex, FHairStrandsBulkCommon& In, bool bWait, bool bFillBulkdata, const FName& InOwnerName)
 {
 	if (In.GetResourceCount() == 0 || InRequestedCurveCount == 0)
 	{
@@ -345,42 +345,42 @@ void FHairStreamingRequest::Request(uint32 InRequestedCurveCount, uint32 InReque
 	
 		//FRequestBarrier Barrier(*DDCRequestOwner);	// This is a critical section on the owner. It does not constrain ordering
 		GetCache().GetChunks(Requests, *DDCRequestOwner,
-		                     [this, &In, bFillBulkdata, InOwnerName](FCacheGetChunkResponse && Response)
-		                     {
-		                     if (Response.Status == UE::DerivedData::EStatus::Ok)
-		                     {
-		                     FHairStreamingRequest::FChunk& Chunk = *(FHairStreamingRequest::FChunk*)Response.UserData;
-							 Chunk.Data_DDC 	= MoveTemp(Response.RawData);
-							 Chunk.Offset 	= Response.RawOffset;
-							 Chunk.Size 		= Response.RawSize;
-							 Chunk.TotalSize = Response.RawOffset + Response.RawSize;
-							 Chunk.Status 	= FHairStreamingRequest::FChunk::Completed;
+		[this, &In, bFillBulkdata, InOwnerName](FCacheGetChunkResponse && Response)
+		{
+			if (Response.Status == UE::DerivedData::EStatus::Ok)
+			{
+				FHairStreamingRequest::FChunk& Chunk = *(FHairStreamingRequest::FChunk*)Response.UserData;
+				Chunk.Data_DDC 	= MoveTemp(Response.RawData);
+				Chunk.Offset 	= Response.RawOffset;
+				Chunk.Size 		= Response.RawSize;
+				Chunk.TotalSize = Response.RawOffset + Response.RawSize;
+				Chunk.Status 	= FHairStreamingRequest::FChunk::Completed;
 		
-							 // Upload the total amount of loaded data
-							 Chunk.Container->LoadedSize = Chunk.TotalSize;
-
-							 // Optional fill in of bytebulkdata container
-							 if (bFillBulkdata)
-							 {
-								 check(Chunk.Container);
-								 FByteBulkData& BulkData = Chunk.Container->Data;
-	
-								 // The buffer is then stored into bulk data
-								 BulkData.Lock(LOCK_READ_WRITE);
-								 void* DstData = BulkData.Realloc(Chunk.Size);
-								 FMemory::Memcpy(DstData, Chunk.GetData(), Chunk.Size);
-								 BulkData.Unlock();
-
-								 //Chunk.Release();
-							 }
-							 }
-							 else
-							 {
-								 FHairStreamingRequest::FChunk& Chunk = *(FHairStreamingRequest::FChunk*)Response.UserData;
-								 Chunk.Status = FHairStreamingRequest::FChunk::Failed;
-								 UE_LOG(LogHairStrands, Error, TEXT("[Groom] DDC request failed for '%s' (Key:%s) "), *InOwnerName.ToString(), *In.DerivedDataKey);
-							 }
-							 });
+				// Upload the total amount of loaded data
+				Chunk.Container->LoadedSize = Chunk.TotalSize;
+		
+				// Optional fill in of bytebulkdata container
+				if (bFillBulkdata)
+				{
+					check(Chunk.Container);
+					FByteBulkData& BulkData = Chunk.Container->Data;
+		
+					// The buffer is then stored into bulk data
+					BulkData.Lock(LOCK_READ_WRITE);
+					void* DstData = BulkData.Realloc(Chunk.Size);
+					FMemory::Memcpy(DstData, Chunk.GetData(), Chunk.Size);
+					BulkData.Unlock();
+		
+					//Chunk.Release();
+				}
+			}
+			else
+			{
+				FHairStreamingRequest::FChunk& Chunk = *(FHairStreamingRequest::FChunk*)Response.UserData;
+				Chunk.Status = FHairStreamingRequest::FChunk::Failed;
+				UE_LOG(LogHairStrands, Error, TEXT("[Groom] DDC request failed for '%s' (Key:%s) "), *InOwnerName.ToString(), *In.DerivedDataKey);
+			}
+		});
 
 		// Optional wait on DDC response
 		if (bWait || bFillBulkdata)
@@ -762,9 +762,7 @@ const TArray<uint32>& FHairStrandsRootBulkData::GetValidSectionIndices(int32 LOD
 
 void FHairStrandsRootBulkData::SerializeData(FArchive& Ar, UObject* Owner)
 {
-	check(Header.LODs.Num() > 0);
 	check(Header.LODs.Num() == Data.LODs.Num());
-
 	for (uint32 LODIt = 0, LODCount=Header.LODs.Num(); LODIt < LODCount; ++LODIt)
 	{
 		const FHeader::FLOD& HeaderLOD = Header.LODs[LODIt];
diff --git a/Engine/Plugins/Runtime/HairStrands/Source/HairStrandsCore/Public/GroomResources.h b/Engine/Plugins/Runtime/HairStrands/Source/HairStrandsCore/Public/GroomResources.h
--- a/Engine/Plugins/Runtime/HairStrands/Source/HairStrandsCore/Public/GroomResources.h
+++ b/Engine/Plugins/Runtime/HairStrands/Source/HairStrandsCore/Public/GroomResources.h
@@ -43,7 +43,7 @@ struct FHairResourceName
 #if HAIR_RESOURCE_DEBUG_NAME
 	FHairResourceName(const FName& In) : AssetName(In) {}
 	FHairResourceName(const FName& In, int32 InGroupIndex) : AssetName(In), GroupIndex(InGroupIndex) {}
-	FHairResourceName(const FName& In, int32 InGroupIndex, int32 InLODIndex) : AssetName(In), GroupIndex(InGroupIndex), LODIndex(InLODIndex) {}
+	FHairResourceName(const FName& In, int32 InGroupIndex, int32 InLODIndex) : AssetName(In), GroupIndex(InGroupIndex), InLODIndex(InLODIndex) {}
 
 	FName AssetName;
 	int32 GroupIndex = -1;
@@ -102,19 +102,17 @@ struct FHairCommonResource : public FRenderResource
 
 	/* Init/Release buffers (FHairCommonResource) */
 	void Allocate(FRDGBuilder& GraphBuilder, EHairResourceLoadingType LoadingType);
-	void Allocate(FRDGBuilder& GraphBuilder, EHairResourceLoadingType LoadingType, EHairResourceStatus& Status);
-	void Allocate(FRDGBuilder& GraphBuilder, EHairResourceLoadingType LoadingType, EHairResourceStatus& Status, uint32 InRequestedCurveCount, uint32 InRequestedPointCount);
-	void AllocateLOD(FRDGBuilder& GraphBuilder, int32 LODIndex, EHairResourceLoadingType LoadingType, EHairResourceStatus& Status);
+	void Allocate(FRDGBuilder& GraphBuilder, EHairResourceLoadingType LoadingType, EHairResourceStatus& Status, int32 InLODIndex=-1);
+	void Allocate(FRDGBuilder& GraphBuilder, EHairResourceLoadingType LoadingType, EHairResourceStatus& Status, uint32 InRequestedCurveCount, uint32 InRequestedPointCount, int32 InLODIndex=-1);
 
-	void StreamInData();
-	void StreamInLODData(int32 LODIndex);
+	void StreamInData(int32 InLODIndex=-1);
 
 	virtual void InternalAllocate() {}
 	virtual void InternalAllocate(FRDGBuilder& GraphBuilder) {}
-	virtual void InternalAllocateLOD(FRDGBuilder& GraphBuilder, int32 LODIndex) {}
+	virtual void InternalAllocate(FRDGBuilder& GraphBuilder, int32 InLODIndex) { InternalAllocate(GraphBuilder); }
 	virtual void InternalRelease() {}
-	virtual bool InternalIsDataLoaded(uint32 InRequestedCurveCount, uint32 InRequestedPointCount) { return true; }
-	virtual bool InternalIsLODDataLoaded(int32 LODIndex) { return true; }
+	virtual bool InternalIsDataLoaded(uint32 InRequestedCurveCount, uint32 InRequestedPointCount, int32 InLODIndex) { return true; }
+	virtual bool InternalIsLODDataLoaded(int32 InLODIndex) const { return true; }
 
 	bool bUseRenderGraph = true;
 	bool bIsInitialized = false;
@@ -141,7 +139,7 @@ struct FHairStrandsRestResource : public FHairCommonResource
 	/* Init/Release buffers */
 	virtual void InternalAllocate(FRDGBuilder& GraphBuilder) override;
 	virtual void InternalRelease() override;
-	virtual bool InternalIsDataLoaded(uint32 InRequestedCurveCount, uint32 InRequestedPointCount) override;
+	virtual bool InternalIsDataLoaded(uint32 InRequestedCurveCount, uint32 InRequestedPointCount, int32 InLODIndex) override;
 
 	/* Get the resource name */
 	virtual FString GetFriendlyName() const override { return TEXT("FHairStrandsResource"); }
@@ -286,7 +284,7 @@ struct FHairStrandsInterpolationResource : public FHairCommonResource
 	/* Init/Release buffers */
 	virtual void InternalAllocate(FRDGBuilder& GraphBuilder) override;
 	virtual void InternalRelease() override;
-	virtual bool InternalIsDataLoaded(uint32 InRequestedCurveCount, uint32 InRequestedPointCount) override;
+	virtual bool InternalIsDataLoaded(uint32 InRequestedCurveCount, uint32 InRequestedPointCount, int32 InLODIndex) override;
 
 	/* Get the resource name */
 	virtual FString GetFriendlyName() const override { return TEXT("FHairStrandsInterplationResource"); }
@@ -317,7 +315,7 @@ struct FHairStrandsClusterCullingResource : public FHairCommonResource
 	/* Init/Release buffers */
 	virtual void InternalAllocate(FRDGBuilder& GraphBuilder) override;
 	virtual void InternalRelease() override;
-	virtual bool InternalIsDataLoaded(uint32 InRequestedCurveCount, uint32 InRequestedPointCount) override;
+	virtual bool InternalIsDataLoaded(uint32 InRequestedCurveCount, uint32 InRequestedPointCount, int32 InLODIndex) override;
 
 	/* Get the resource name */
 	virtual FString GetFriendlyName() const override { return TEXT("FHairStrandsClusterResource"); }
@@ -353,10 +351,10 @@ struct FHairStrandsRestRootResource : public FHairCommonResource
 	FHairStrandsRestRootResource(FHairStrandsRootBulkData& BulkData, EHairStrandsResourcesType CurveType, const FHairResourceName& ResourceName, const FName& OwnerName);
 
 	/* Init/Release buffers */
-	virtual void InternalAllocate(FRDGBuilder& GraphBuilder) override;
-	virtual void InternalAllocateLOD(FRDGBuilder& GraphBuilder, int32 LODIndex) override;
+	virtual void InternalAllocate(FRDGBuilder& GraphBuilder, int32 InLODIndex) override;
 	virtual void InternalRelease() override;
-	virtual bool InternalIsLODDataLoaded(int32 LODIndex) override;
+	virtual bool InternalIsDataLoaded(uint32 InRequestedCurveCount, uint32 InRequestedPointCount, int32 InLODIndex) override;
+	virtual bool InternalIsLODDataLoaded(int32 InLODIndex) const override;
 
 	/* Get the resource name */
 	virtual FString GetFriendlyName() const override { return TEXT("FHairStrandsRestRootResource"); }
@@ -432,8 +430,9 @@ struct FHairStrandsDeformedRootResource : public FHairCommonResource
 	FHairStrandsDeformedRootResource(const FHairStrandsRestRootResource* InRestResources, EHairStrandsResourcesType CurveType, const FHairResourceName& ResourceName, const FName& OwnerName);
 
 	/* Init/Release buffers */
-	virtual void InternalAllocateLOD(FRDGBuilder& GraphBuilder, int32 LODIndex) override;
+	virtual void InternalAllocate(FRDGBuilder& GraphBuilder, int32 InLODIndex) override;
 	virtual void InternalRelease() override;
+	virtual bool InternalIsLODDataLoaded(int32 InLODIndex) const override;
 
 	/* Get the resource name */
 	virtual FString GetFriendlyName() const override { return TEXT("FHairStrandsDeformedRootResource"); }
diff --git a/Engine/Plugins/Runtime/HairStrands/Source/HairStrandsCore/Public/HairStrandsDatas.h b/Engine/Plugins/Runtime/HairStrands/Source/HairStrandsCore/Public/HairStrandsDatas.h
--- a/Engine/Plugins/Runtime/HairStrands/Source/HairStrandsCore/Public/HairStrandsDatas.h
+++ b/Engine/Plugins/Runtime/HairStrands/Source/HairStrandsCore/Public/HairStrandsDatas.h
@@ -368,7 +368,7 @@ struct FHairStreamingRequest
 		void Release();
 	};
 
-	void Request(uint32 InRequestedCurveCount, uint32 InRequestedPointCount, FHairStrandsBulkCommon& In, bool bWait=false, bool bFillBulkData=false, const FName& InOwnerName = NAME_None);
+	void Request(uint32 InRequestedCurveCount, uint32 InRequestedPointCount, int32 InLODIndex, FHairStrandsBulkCommon& In, bool bWait=false, bool bFillBulkData=false, const FName& InOwnerName = NAME_None);
 	bool IsNone() const;
 	bool IsCompleted();
 
