* Change binding asset to be cached per group (builder & serialization)
* Add DDC/IO streaming support for binding asset
* Add compatible with fine grain streaming

diff --git a/Engine/Plugins/Importers/USDImporter/Source/USDSchemas/Private/USDGroomTranslatorUtils.cpp b/Engine/Plugins/Importers/USDImporter/Source/USDSchemas/Private/USDGroomTranslatorUtils.cpp
--- a/Engine/Plugins/Importers/USDImporter/Source/USDSchemas/Private/USDGroomTranslatorUtils.cpp
+++ b/Engine/Plugins/Importers/USDImporter/Source/USDSchemas/Private/USDGroomTranslatorUtils.cpp
@@ -119,7 +119,7 @@ namespace UE::UsdGroomTranslatorUtils::Private
 				GroomBinding->SourceGeometryCache = Cast<UGeometryCache>(Settings.SourceMesh);
 				GroomBinding->TargetGeometryCache = Cast<UGeometryCache>(Settings.TargetMesh);
 			}
-			GroomBinding->HairGroupBulkDatas.Reserve(Settings.Groom->HairGroupsPlatformData.Num());
+			GroomBinding->HairGroupsPlatformData.Reserve(Settings.Groom->HairGroupsPlatformData.Num());
 			GroomBinding->NumInterpolationPoints = Settings.NumInterpolationPoints;
 			GroomBinding->MatchingSection = Settings.MatchingSection;
 
diff --git a/Engine/Plugins/Runtime/HairStrands/Source/HairStrandsCore/Private/GroomAsset.cpp b/Engine/Plugins/Runtime/HairStrands/Source/HairStrandsCore/Private/GroomAsset.cpp
--- a/Engine/Plugins/Runtime/HairStrands/Source/HairStrandsCore/Private/GroomAsset.cpp
+++ b/Engine/Plugins/Runtime/HairStrands/Source/HairStrandsCore/Private/GroomAsset.cpp
@@ -282,11 +282,11 @@ void DumpLoadedGroomBindingAssets(IConsoleVariable* InCVarPakTesterEnabled)
 	{
 		if (AssetIt)
 		{			
-			const uint32 GroupCount = AssetIt->HairGroupBulkDatas.Num();
+			const uint32 GroupCount = AssetIt->HairGroupsPlatformData.Num();
 			for (uint32 GroupIt = 0; GroupIt < GroupCount; ++GroupIt)
 			{
-				const uint32 CPUMemorySize_Guides  = AssetIt->HairGroupBulkDatas[GroupIt].SimRootBulkData.GetDataSize();
-				const uint32 CPUMemorySize_Strands = AssetIt->HairGroupBulkDatas[GroupIt].RenRootBulkData.GetDataSize();
+				const uint32 CPUMemorySize_Guides  = AssetIt->HairGroupsPlatformData[GroupIt].SimRootBulkData.GetDataSize();
+				const uint32 CPUMemorySize_Strands = AssetIt->HairGroupsPlatformData[GroupIt].RenRootBulkData.GetDataSize();
 
 				uint32 GPUMemorySize_Guides = 0;
 				if (const FHairStrandsRestRootResource* RootResource = AssetIt->HairGroupResources[GroupIt].SimRootResources)
@@ -301,10 +301,10 @@ void DumpLoadedGroomBindingAssets(IConsoleVariable* InCVarPakTesterEnabled)
 
 				uint32 CPUMemorySize_Cards = 0;
 				uint32 GPUMemorySize_Cards = 0;
-				const uint32 CardCount = AssetIt->HairGroupBulkDatas[GroupIt].CardsRootBulkData.Num();
+				const uint32 CardCount = AssetIt->HairGroupsPlatformData[GroupIt].CardsRootBulkData.Num();
 				for (uint32 CardIt = 0; CardIt < CardCount; ++CardIt)
 				{
-					CPUMemorySize_Cards += AssetIt->HairGroupBulkDatas[GroupIt].CardsRootBulkData[CardIt].GetDataSize();
+					CPUMemorySize_Cards += AssetIt->HairGroupsPlatformData[GroupIt].CardsRootBulkData[CardIt].GetDataSize();
 					if (const FHairStrandsRestRootResource* RootResource = AssetIt->HairGroupResources[GroupIt].CardsRootResources[CardIt])
 					{
 						GPUMemorySize_Cards += RootResource->GetResourcesSize();
@@ -322,7 +322,7 @@ void DumpLoadedGroomBindingAssets(IConsoleVariable* InCVarPakTesterEnabled)
 				Total_GPUMemorySize_Strands+= GPUMemorySize_Strands;
 				Total_GPUMemorySize_Cards  += GPUMemorySize_Cards;
 
-				const uint32 SkelLODCount = AssetIt->HairGroupBulkDatas[GroupIt].RenRootBulkData.GetLODCount();
+				const uint32 SkelLODCount = AssetIt->HairGroupsPlatformData[GroupIt].RenRootBulkData.GetLODCount();
 				if (bDetails)
 				{
 //					UE_LOG(LogHairStrands, Log, TEXT("--  No.  - LOD -    CPU Total (     Guides|    Strands|      Cards) -    GPU Total (     Guides|    Strands|      Cards) - Asset Name "));
diff --git a/Engine/Plugins/Runtime/HairStrands/Source/HairStrandsCore/Private/GroomBindingAsset.cpp b/Engine/Plugins/Runtime/HairStrands/Source/HairStrandsCore/Private/GroomBindingAsset.cpp
--- a/Engine/Plugins/Runtime/HairStrands/Source/HairStrandsCore/Private/GroomBindingAsset.cpp
+++ b/Engine/Plugins/Runtime/HairStrands/Source/HairStrandsCore/Private/GroomBindingAsset.cpp
@@ -32,35 +32,77 @@ bool IsHairStrandsDDCLogEnable();
 
 ////////////////////////////////////////////////////////////////////////////////////////////////////////////
 
-static void InternalSerialize(FArchive& Ar, UObject* Owner, TArray<FHairStrandsRootBulkData>& GroupDatas)
+static void InternalSerializeGuides(FArchive& Ar, UObject* Owner, FHairStrandsRootBulkData& Data)
 {
-	uint32 GroupCount = GroupDatas.Num();
-	Ar << GroupCount;
-	if (Ar.IsLoading())
+	Data.SerializeHeader(Ar, Owner);
+	for (int32 LODIndex = 0, LODCount = Data.GetLODCount(); LODIndex < LODCount; ++LODIndex)
 	{
-		GroupDatas.SetNum(GroupCount);
+		Data.SerializeData(Ar, Owner, LODIndex);
 	}
-	for (uint32 GroupIt = 0; GroupIt < GroupCount; ++GroupIt)
+}
+
+static void InternalSerializeStrands(FArchive& Ar, UObject* Owner, FHairStrandsRootBulkData& Data, uint32 Flags, bool bHeader, bool bData)
+{
+	Ar.UsingCustomVersion(FAnimObjectVersion::GUID);
+
+	const bool bStripped = (Flags & UGroomAsset::CDSF_StrandsStripped);
+	if (!bStripped)
 	{
-		GroupDatas[GroupIt].Serialize(Ar, Owner);
+		// When cooking data, force loading of *all* bulk data prior to saving them
+		// Note: bFillBulkdata is true for filling in the bulkdata container prior to serialization. This also forces the resources loading 
+		// from the 'start' (i.e., without offset)
+		if (Ar.IsCooking() && Ar.IsSaving())
+		{
+			for (int32 LODIndex = 0, LODCount = Data.GetLODCount(); LODIndex < LODCount; ++LODIndex)
+			{
+				FHairStreamingRequest R; R.Request(HAIR_MAX_NUM_CURVE_PER_GROUP, HAIR_MAX_NUM_POINT_PER_GROUP, LODIndex, Data, true /*bWait*/, true /*bFillBulkdata*/, Owner->GetFName());
+			}
+		}
+
+		if (bHeader){ Data.SerializeHeader(Ar, Owner); }
+		if (bData)
+		{
+			for (int32 LODIndex = 0, LODCount = Data.GetLODCount(); LODIndex < LODCount; ++LODIndex)
+			{
+				Data.SerializeData(Ar, Owner, LODIndex);
+			}
+		}
 	}
 }
 
-static void InternalSerialize(FArchive& Ar, UObject* Owner, FHairGroupBulkData& GroupData, uint32 Flags)
+static void InternalSerializeCards(FArchive& Ar, UObject* Owner, TArray<FHairStrandsRootBulkData>& Datas)
 {
-	Ar.UsingCustomVersion(FAnimObjectVersion::GUID);
-	GroupData.SimRootBulkData.Serialize(Ar, Owner);
-	if (!(Flags & UGroomAsset::CDSF_StrandsStripped))
+	uint32 CardLODCount = Datas.Num();
+	Ar << CardLODCount;
+	if (Ar.IsLoading())
 	{
-		GroupData.RenRootBulkData.Serialize(Ar, Owner);
+		Datas.SetNum(CardLODCount);
 	}
-	if (Ar.CustomVer(FAnimObjectVersion::GUID) >= FAnimObjectVersion::SerializeHairBindingAsset)
+	for (FHairStrandsRootBulkData& Data : Datas)
 	{
-		InternalSerialize(Ar, Owner, GroupData.CardsRootBulkData);
+		Data.SerializeHeader(Ar, Owner);
+		for (int32 LODIndex = 0, LODCount = Data.GetLODCount(); LODIndex < LODCount; ++LODIndex)
+		{
+			Data.SerializeData(Ar, Owner, LODIndex);
+		}
 	}
 }
 
-static void InternalSerialize(FArchive& Ar, UObject* Owner, TArray<FHairGroupBulkData>& GroupDatas, uint32 Flags)
+static void InternalSerializePlatformData(FArchive& Ar, UObject* Owner, UGroomBindingAsset::FHairGroupPlatformData& GroupData, uint32 Flags, bool bHeader, bool bData)
+{
+	Ar.UsingCustomVersion(FAnimObjectVersion::GUID);
+
+	// Guides
+	InternalSerializeGuides(Ar, Owner, GroupData.SimRootBulkData);
+
+	// Strands
+	InternalSerializeStrands(Ar, Owner, GroupData.RenRootBulkData, Flags, bHeader, bData);
+
+	// Cards
+	InternalSerializeCards(Ar, Owner, GroupData.CardsRootBulkData);
+}
+
+static void InternalSerializePlatformDatas(FArchive& Ar, UObject* Owner, TArray<UGroomBindingAsset::FHairGroupPlatformData>& GroupDatas, uint32 Flags)
 {
 	uint32 GroupCount = GroupDatas.Num();
 	Ar << Flags;
@@ -71,7 +113,7 @@ static void InternalSerialize(FArchive& Ar, UObject* Owner, TArray<FHairGroupBul
 	}
 	for (uint32 GroupIt = 0; GroupIt < GroupCount; ++GroupIt)
 	{
-		InternalSerialize(Ar, Owner, GroupDatas[GroupIt], Flags);
+		InternalSerializePlatformData(Ar, Owner, GroupDatas[GroupIt], Flags, true /*bHeader*/, true /*bData*/);
 	}
 }
 
@@ -90,10 +132,10 @@ void UGroomBindingAsset::Serialize(FArchive& Ar)
 	Super::Serialize(Ar);
 	Ar.UsingCustomVersion(FAnimObjectVersion::GUID);
 #if WITH_EDITOR
-	if (Ar.CustomVer(FAnimObjectVersion::GUID) < FAnimObjectVersion::GroomBindingSerialization || Ar.IsCooking())
+	if (Ar.IsCooking())
 #endif
 	{
-		InternalSerialize(Ar, this, HairGroupBulkDatas, Flags);
+		InternalSerializePlatformDatas(Ar, this, HairGroupsPlatformData, Flags);
 		bIsValid = true;
 	}
 }
@@ -105,7 +147,7 @@ void UGroomBindingAsset::InitResource()
 	// Ensure we are releasing binding resources before reallocating them
 	ReleaseResource();
 
-	for (FHairGroupBulkData& BulkData : HairGroupBulkDatas)
+	for (UGroomBindingAsset::FHairGroupPlatformData& BulkData : HairGroupsPlatformData)
 	{
 		const int32 GroupIndex = HairGroupResources.Num();
 		FHairGroupResource& Resource = HairGroupResources.AddDefaulted_GetRef();
@@ -241,7 +283,7 @@ void UGroomBindingAsset::ReleaseResource()
 void UGroomBindingAsset::Reset()
 {
 	ReleaseResource();
-	for (FHairGroupBulkData& Data : HairGroupBulkDatas)
+	for (UGroomBindingAsset::FHairGroupPlatformData& Data : HairGroupsPlatformData)
 	{
 		Data.SimRootBulkData.Reset();
 		Data.RenRootBulkData.Reset();
@@ -661,6 +703,25 @@ void UGroomBindingAsset::PostEditChangeProperty(FPropertyChangedEvent& PropertyC
 }
 #endif // WITH_EDITOR
 
+void UpdateGroomBindingAssetInfos(UGroomBindingAsset* In)
+{
+	if (In)
+	{
+		const uint32 GroupCount = In->HairGroupsPlatformData.Num();
+		In->GroupInfos.SetNum(GroupCount);
+		for (uint32 GroupIt=0; GroupIt< GroupCount; ++GroupIt)
+		{
+			FGoomBindingGroupInfo& Info = In->GroupInfos[GroupIt];
+			const UGroomBindingAsset::FHairGroupPlatformData& BulkData = In->HairGroupsPlatformData[GroupIt];
+			Info.SimRootCount = BulkData.SimRootBulkData.GetRootCount();
+			Info.SimLODCount  = BulkData.SimRootBulkData.GetLODCount();
+
+			Info.RenRootCount = BulkData.RenRootBulkData.GetRootCount();
+			Info.RenLODCount  = BulkData.RenRootBulkData.GetLODCount();
+		}
+	}
+}
+
 #if WITH_EDITORONLY_DATA
 
 // If groom binding derived data needs to be rebuilt (new format, serialization
@@ -724,77 +785,122 @@ void UGroomBindingAsset::CacheDerivedDatas()
 		return;
 	}
 
-	// List all the components which will need to be recreated to get the new binding information
+	// 1. Set group count to the groom target
+	const uint32 GroupCount = Groom->GetNumHairGroups();
+	HairGroupsPlatformData.SetNum(GroupCount);
+	CachedDerivedDataKey.SetNum(GroupCount);
+	GroupInfos.SetNum(GroupCount);
+
+	// 2. Prepare main cache key
 	const FString KeySuffix = BuildDerivedDataKeySuffix(*this);
-	const FString DerivedDataKey = GroomBindingDerivedDataCacheUtils::BuildGroomBindingDerivedDataKey(KeySuffix);
 
-	if (DerivedDataKey != CachedDerivedDataKey)
+	// 3. Build or retrieve from cache, binding data for each group
+	bIsValid = true;
+	bool bReloadResource = false;
+	for (uint32 GroupIndex = 0; GroupIndex < GroupCount; ++GroupIndex)
+	{
+		bool bGroupReloadResource = false;
+		bool bGroupValid = true;
+
+		CacheDerivedDatas(GroupIndex, KeySuffix, bGroupValid, bGroupReloadResource);
+
+		bIsValid = bIsValid && bGroupValid;
+		bReloadResource = bReloadResource || bGroupReloadResource;
+	}
+
+	// 4. Reload resources if needed
+	if (bReloadResource)
+	{
+		InitResource();
+	}
+
+	// 3. Patch hair group info if it does not match the DDC-read/deserialized data
+	UpdateGroomBindingAssetInfos(this);
+}
+
+void UGroomBindingAsset::CacheDerivedDatas(uint32 InGroupIndex, const FString KeySuffix, bool& bOutValid, bool& bOutReloadResource)
+{
+	const FString DerivedDataKey = GroomBindingDerivedDataCacheUtils::BuildGroomBindingDerivedDataKey(KeySuffix + FString(TEXT("_Group")) + FString::FromInt(InGroupIndex));
+
+	bOutValid = false;
+	bOutReloadResource = false;
+	if (DerivedDataKey != CachedDerivedDataKey[InGroupIndex])
 	{
 		using namespace UE::DerivedData;
 
-		const FCacheKey Key = ConvertLegacyCacheKey(DerivedDataKey);
+		const FCacheKey HeaderKey = ConvertLegacyCacheKey(DerivedDataKey + FString(TEXT("_Header")));
 		const FSharedString Name = MakeStringView(GetPathName());
 		FSharedBuffer Data;
 		{
 			FRequestOwner Owner(EPriority::Blocking);
-			GetCache().GetValue({ {Name, Key} }, Owner, [&Data](FCacheGetValueResponse&& Response)
+			GetCache().GetValue({ {Name, HeaderKey} }, Owner, [&Data](FCacheGetValueResponse&& Response)
 			{
 				Data = Response.Value.GetData().Decompress();
 			});
 			Owner.Wait();
 		}
 
+		UGroomBindingAsset::FHairGroupPlatformData& PlatformData = HairGroupsPlatformData[InGroupIndex];
+
+		// Populate key/name for streaming data request
+		auto FillDrivedDataKey = [&DerivedDataKey, &Name](UGroomBindingAsset::FHairGroupPlatformData& In)
+		{
+			In.RenRootBulkData.DerivedDataKey = DerivedDataKey + FString(TEXT("_Data"));
+		};
+
 		FGroomComponentRecreateRenderStateContext RecreateRenderContext(Groom);
 
 		if (Data)
 		{
 			UE_CLOG(IsHairStrandsDDCLogEnable(), LogHairStrands, Log, TEXT("[GroomBinding/DDC] Found (GroomBinding:%s)."), *GetName());
 
+			FillDrivedDataKey(PlatformData);
+
+			// Header
 			FMemoryReaderView Ar(Data, /*bIsPersistent*/ true);
-			InternalSerialize(Ar, this, HairGroupBulkDatas, 0);
+			InternalSerializePlatformData(Ar, this, PlatformData, 0 /*Flags*/, true /*bHeader*/, false /*bData*/);
 
-			bIsValid = true;
+			bOutValid = true;
 		}
 		else
 		{
 			UE_CLOG(IsHairStrandsDDCLogEnable(), LogHairStrands, Log, TEXT("[GroomBinding/DDC] Not found (GroomBinding:%s)."), *GetName());
 
 			// Build groom binding data
-			bIsValid = FGroomBindingBuilder::BuildBinding(this, false);
-			if (bIsValid)
+			bOutValid = FGroomBindingBuilder::BuildBinding(this, InGroupIndex);
+
+			if (bOutValid)
 			{
-				TArray<uint8> WriteData;
-				FMemoryWriter Ar(WriteData, /*bIsPersistent*/ true);
-				InternalSerialize(Ar, this, HairGroupBulkDatas, 0);
+				FillDrivedDataKey(PlatformData);
 
-				FRequestOwner AsyncOwner(EPriority::Normal);
-				GetCache().PutValue({ {Name, Key, FValue::Compress(MakeSharedBufferFromArray(MoveTemp(WriteData)))} }, AsyncOwner);
-				AsyncOwner.KeepAlive();
-			}
-		}
+				// Header
+				{
+					TArray<uint8> WriteData;
+					FMemoryWriter Ar(WriteData, /*bIsPersistent*/ true);
+					InternalSerializePlatformData(Ar, this, PlatformData, 0 /*Flags*/, true /*bHeader*/, false /*bData*/);
+	
+					FRequestOwner AsyncOwner(EPriority::Normal);
+					GetCache().PutValue({ {Name, HeaderKey, FValue::Compress(MakeSharedBufferFromArray(MoveTemp(WriteData)))} }, AsyncOwner);
+					AsyncOwner.KeepAlive();
+				}
+	
+				// Data
+				for (uint32 LODIt=0, LODCount = PlatformData.RenRootBulkData.GetLODCount(); LODIt<LODCount;++LODIt)
+				{
+					TArray<FCachePutValueRequest> Out;
+					PlatformData.RenRootBulkData.Write_DDC(this, Out, LODIt);
 
-		// Patch hair group info if it does not match the DDC-read/deserialized data
-		const uint32 GroupCount = HairGroupBulkDatas.Num();
-		if (GroupInfos.Num() != GroupCount)
-		{
-			GroupInfos.SetNum(GroupCount);
-		}
-		for (uint32 GroupIt=0; GroupIt< GroupCount; ++GroupIt)
-		{
-			FGoomBindingGroupInfo& Info = GroupInfos[GroupIt];
-			const FHairGroupBulkData& BulkData = HairGroupBulkDatas[GroupIt];
-			{
-				Info.SimRootCount = BulkData.SimRootBulkData.GetRootCount();
-				Info.SimLODCount  = BulkData.SimRootBulkData.GetLODCount();
-				Info.RenRootCount = BulkData.RenRootBulkData.GetRootCount();
-				Info.RenLODCount  = BulkData.RenRootBulkData.GetLODCount();
+					FRequestOwner AsyncOwner(EPriority::Normal);
+					GetCache().PutValue(Out, AsyncOwner);
+					AsyncOwner.KeepAlive();
+				}
 			}
 		}
 
-		if (bIsValid)
+		if (bOutValid)
 		{
-			CachedDerivedDataKey = DerivedDataKey;
-			InitResource();
+			bOutReloadResource = true;
+			CachedDerivedDataKey[InGroupIndex] = DerivedDataKey;
 		}
 		else
 		{
@@ -806,7 +912,7 @@ void UGroomBindingAsset::CacheDerivedDatas()
 
 void UGroomBindingAsset::GetResourceSizeEx(FResourceSizeEx& CumulativeResourceSize)
 {
-	CumulativeResourceSize.AddDedicatedSystemMemoryBytes(HairGroupBulkDatas.GetAllocatedSize());
+	CumulativeResourceSize.AddDedicatedSystemMemoryBytes(HairGroupsPlatformData.GetAllocatedSize());
 
 	for (const FHairGroupResource& Group : HairGroupResources)
 	{
diff --git a/Engine/Plugins/Runtime/HairStrands/Source/HairStrandsCore/Private/GroomBindingBuilder.cpp b/Engine/Plugins/Runtime/HairStrands/Source/HairStrandsCore/Private/GroomBindingBuilder.cpp
--- a/Engine/Plugins/Runtime/HairStrands/Source/HairStrandsCore/Private/GroomBindingBuilder.cpp
+++ b/Engine/Plugins/Runtime/HairStrands/Source/HairStrandsCore/Private/GroomBindingBuilder.cpp
@@ -57,7 +57,7 @@ static FAutoConsoleVariableRef CVarHairStrandsBindingBuilderWarningEnable(TEXT("
 FString FGroomBindingBuilder::GetVersion()
 {
 	// Important to update the version when groom building changes
-	return TEXT("3a");
+	return TEXT("3c");
 }
 
 ///////////////////////////////////////////////////////////////////////////////////////////////////
@@ -699,13 +699,12 @@ namespace GroomBinding_RBFWeighting
 	}
 
 	void ComputeInterpolationWeights(
-		TArray<FHairRootGroupData>& Out, 
+		FHairRootGroupData& Out, 
 		const uint32 NumInterpolationPoints, 
 		const int32 MatchingSection, 
 		const IMeshData* MeshData, 
 		const TArray<TArray<FVector3f>>& TransferedPositions)
 	{
-		const uint32 GroupCount  = Out.Num();
 		const uint32 MeshLODCount= MeshData->GetNumLODs();
 		const uint32 MaxSamples  = NumInterpolationPoints;
 
@@ -730,9 +729,8 @@ namespace GroomBinding_RBFWeighting
 			if (!GlobalSamples)
 			{
 				TArray<bool> ValidPoints;
-				for (uint32 GroupIt = 0; GroupIt < GroupCount; ++GroupIt)
 				{
-					FillLocalValidPoints(MeshLODData, TargetSection, Out[GroupIt].SimRootData.MeshProjectionLODs[LODIndex], ValidPoints);
+					FillLocalValidPoints(MeshLODData, TargetSection, Out.SimRootData.MeshProjectionLODs[LODIndex], ValidPoints);
 
 					FPointsSampler PointsSampler(ValidPoints, PositionsPointer, MaxSamples);
 					const uint32 SampleCount = PointsSampler.SamplePositions.Num();
@@ -740,8 +738,8 @@ namespace GroomBinding_RBFWeighting
 					FWeightsBuilder InterpolationWeights(SampleCount, SampleCount,
 						PointsSampler.SamplePositions.GetData(), PointsSampler.SamplePositions.GetData());
 
-					UpdateInterpolationWeights(InterpolationWeights, PointsSampler, LODIndex, Out[GroupIt].SimRootData);
-					UpdateInterpolationWeights(InterpolationWeights, PointsSampler, LODIndex, Out[GroupIt].RenRootData);
+					UpdateInterpolationWeights(InterpolationWeights, PointsSampler, LODIndex, Out.SimRootData);
+					UpdateInterpolationWeights(InterpolationWeights, PointsSampler, LODIndex, Out.RenRootData);
 				}
 			}
 			else
@@ -756,10 +754,9 @@ namespace GroomBinding_RBFWeighting
 				FWeightsBuilder InterpolationWeights(SampleCount, SampleCount,
 					PointsSampler.SamplePositions.GetData(), PointsSampler.SamplePositions.GetData());
 
-				for (uint32 GroupIt = 0; GroupIt < GroupCount; ++GroupIt)
 				{
-					UpdateInterpolationWeights(InterpolationWeights, PointsSampler, LODIndex, Out[GroupIt].SimRootData);
-					UpdateInterpolationWeights(InterpolationWeights, PointsSampler, LODIndex, Out[GroupIt].RenRootData);
+					UpdateInterpolationWeights(InterpolationWeights, PointsSampler, LODIndex, Out.SimRootData);
+					UpdateInterpolationWeights(InterpolationWeights, PointsSampler, LODIndex, Out.RenRootData);
 				}
 			}
 		}
@@ -1916,7 +1913,7 @@ static void BuildRootBulkData(
 	Out.Header.Strides.RootToUniqueTriangleIndexBufferStride 	= FHairStrandsRootToUniqueTriangleIndexFormat::SizeInByte;
 	Out.Header.Strides.RootBarycentricBufferStride 				= FHairStrandsRootBarycentricFormat::SizeInByte;
 	Out.Header.Strides.UniqueTriangleIndexBufferStride 			= FHairStrandsUniqueTriangleIndexFormat::SizeInByte;
-	Out.Header.Strides.RestUniqueTrianglePositionBufferStride 	= FHairStrandsMeshTrianglePositionFormat::SizeInByte;
+	Out.Header.Strides.RestUniqueTrianglePositionBufferStride 	= FHairStrandsMeshTrianglePositionFormat::SizeInByte * 3; // 3 vertices per triangle
 
 	Out.Header.Strides.MeshInterpolationWeightsBufferStride 	= FHairStrandsWeightFormat::SizeInByte;
 	Out.Header.Strides.MeshSampleIndicesBufferStride 			= FHairStrandsIndexFormat::SizeInByte;
@@ -2007,21 +2004,18 @@ static void BuildRootData(
 
 // Convert the root data into root bulk data
 static void BuildRootBulkData(
-	TArray<FHairGroupBulkData>& Out,
-	const TArray<FHairRootGroupData>& In)
+	UGroomBindingAsset::FHairGroupPlatformData& Out,
+	const FHairRootGroupData& In)
 {
-	const uint32 GroupCount = In.Num();
-	Out.SetNum(GroupCount);
-	for (uint32 GroupIt = 0; GroupIt < GroupCount; ++GroupIt)
 	{
-		BuildRootBulkData(Out[GroupIt].SimRootBulkData, In[GroupIt].SimRootData);
-		BuildRootBulkData(Out[GroupIt].RenRootBulkData, In[GroupIt].RenRootData);
+		BuildRootBulkData(Out.SimRootBulkData, In.SimRootData);
+		BuildRootBulkData(Out.RenRootBulkData, In.RenRootData);
 
-		const uint32 LODCount = In[GroupIt].CardsRootData.Num();
-		Out[GroupIt].CardsRootBulkData.SetNum(LODCount);
+		const uint32 LODCount = In.CardsRootData.Num();
+		Out.CardsRootBulkData.SetNum(LODCount);
 		for (uint32 LODIt = 0; LODIt < LODCount; ++LODIt)
 		{
-			BuildRootBulkData(Out[GroupIt].CardsRootBulkData[LODIt], In[GroupIt].CardsRootData[LODIt]);
+			BuildRootBulkData(Out.CardsRootBulkData[LODIt], In.CardsRootData[LODIt]);
 		}
 	}
 }
@@ -2030,7 +2024,7 @@ static void BuildRootBulkData(
 
   ///////////////////////////////////////////////////////////////////////////////////////////////////
 // Main entry (CPU path)
-static bool InternalBuildBinding_CPU(UGroomBindingAsset* BindingAsset, bool bInitResources)
+static bool InternalBuildBinding_CPU(UGroomBindingAsset* BindingAsset, uint32 InGroupIndex)
 {
 #if WITH_EDITORONLY_DATA
 	if (!BindingAsset ||
@@ -2043,7 +2037,7 @@ static bool InternalBuildBinding_CPU(UGroomBindingAsset* BindingAsset, bool bIni
 	}
 
 	// 1. Build groom root data
-	TArray<FHairRootGroupData> OutHairGroupDatas;
+	FHairRootGroupData OutData;
 	{
 
 		BindingAsset->Groom->ConditionalPostLoad();
@@ -2084,14 +2078,13 @@ static bool InternalBuildBinding_CPU(UGroomBindingAsset* BindingAsset, bool bIni
 		const uint32 GroupCount = GroomAsset->GetNumHairGroups();
 		const uint32 MeshLODCount = TargetMeshData->GetNumLODs();
 
-		uint32 GroupIndex = 0;
-		for (const FHairGroupPlatformData& GroupData : GroomAsset->HairGroupsPlatformData)
+		check(InGroupIndex < uint32(GroomAsset->HairGroupsPlatformData.Num()));
+		const FHairGroupPlatformData& GroupData = GroomAsset->HairGroupsPlatformData[InGroupIndex];
 		{
 			FHairStrandsDatas StrandsData;
 			FHairStrandsDatas GuidesData;
-			GroomAsset->GetHairStrandsDatas(GroupIndex, StrandsData, GuidesData);
+			GroomAsset->GetHairStrandsDatas(InGroupIndex, StrandsData, GuidesData);
 
-			FHairRootGroupData& OutData = OutHairGroupDatas.AddDefaulted_GetRef();
 			InitHairStrandsRootData(OutData.RenRootData, &StrandsData, MeshLODCount, NumInterpolationPoints);
 			InitHairStrandsRootData(OutData.SimRootData, &GuidesData, MeshLODCount, NumInterpolationPoints);
 
@@ -2102,14 +2095,13 @@ static bool InternalBuildBinding_CPU(UGroomBindingAsset* BindingAsset, bool bIni
 				if (GroupData.Cards.IsValid(CardsLODIt))
 				{
 					FHairStrandsDatas LODGuidesData;
-					const bool bIsValid = GroomAsset->GetHairCardsGuidesDatas(GroupIndex, CardsLODIt, LODGuidesData);
+					const bool bIsValid = GroomAsset->GetHairCardsGuidesDatas(InGroupIndex, CardsLODIt, LODGuidesData);
 					if (bIsValid)
 					{
 						InitHairStrandsRootData(OutData.CardsRootData[CardsLODIt], &LODGuidesData, MeshLODCount, NumInterpolationPoints);
 					}
 				}
 			}
-			++GroupIndex;
 		}
 
 		// Transfer requires root UV embedded into the groom asset. It is not possible to read safely hair description here to extra this data.
@@ -2117,10 +2109,9 @@ static bool InternalBuildBinding_CPU(UGroomBindingAsset* BindingAsset, bool bIni
 
 		// Create mapping between the source & target using their UV
 		uint32 WorkItemCount = 1 + (bNeedTransferPosition ? 1 : 0); //RBF + optional position transfer
-		for (uint32 GroupIt = 0; GroupIt < GroupCount; ++GroupIt)
 		{
 			WorkItemCount += 2; // Sim & Render
-			const uint32 CardsLODCount = OutHairGroupDatas[GroupIt].CardsRootData.Num();
+			const uint32 CardsLODCount = OutData.CardsRootData.Num();
 			WorkItemCount += CardsLODCount;
 		}
 	
@@ -2145,17 +2136,16 @@ static bool InternalBuildBinding_CPU(UGroomBindingAsset* BindingAsset, bool bIni
 		}
 
 		bool bSucceed = false;
-		for (uint32 GroupIt=0; GroupIt < GroupCount; ++GroupIt)
 		{
 			FHairStrandsDatas StrandsData;
 			FHairStrandsDatas GuidesData;
-			GroomAsset->GetHairStrandsDatas(GroupIt, StrandsData, GuidesData);
+			GroomAsset->GetHairStrandsDatas(InGroupIndex, StrandsData, GuidesData);
 
 			bSucceed = GroomBinding_RootProjection::Project(
 				StrandsData,
 				TargetMeshData.Get(),
 				TransferredPositions,
-				OutHairGroupDatas[GroupIt].RenRootData);
+				OutData.RenRootData);
 			if (!bSucceed)
 			{
 				UE_LOG(LogHairStrands, Error, TEXT("[Groom] Binding asset could not be built. Some strand roots are not close enough to the target mesh to be projected onto it."));
@@ -2168,7 +2158,7 @@ static bool InternalBuildBinding_CPU(UGroomBindingAsset* BindingAsset, bool bIni
 				GuidesData,
 				TargetMeshData.Get(),
 				TransferredPositions,
-				OutHairGroupDatas[GroupIt].SimRootData);
+				OutData.SimRootData);
 			if (!bSucceed) 
 			{
 				UE_LOG(LogHairStrands, Error, TEXT("[Groom] Binding asset could not be built. Some guide roots are not close enough to the target mesh to be projected onto it."));
@@ -2177,20 +2167,20 @@ static bool InternalBuildBinding_CPU(UGroomBindingAsset* BindingAsset, bool bIni
 
 			SlowTask.EnterProgressFrame();
 
-			const uint32 CardsLODCount = OutHairGroupDatas[GroupIt].CardsRootData.Num();
+			const uint32 CardsLODCount = OutData.CardsRootData.Num();
 			for (uint32 CardsLODIt = 0; CardsLODIt < CardsLODCount; ++CardsLODIt)
 			{
-				if (BindingAsset->Groom->HairGroupsPlatformData[GroupIt].Cards.IsValid(CardsLODIt))
+				if (BindingAsset->Groom->HairGroupsPlatformData[InGroupIndex].Cards.IsValid(CardsLODIt))
 				{
 					FHairStrandsDatas LODGuidesData;
-					const bool bIsValid = GroomAsset->GetHairCardsGuidesDatas(GroupIt, CardsLODIt, LODGuidesData);
+					const bool bIsValid = GroomAsset->GetHairCardsGuidesDatas(InGroupIndex, CardsLODIt, LODGuidesData);
 					if (bIsValid)
 					{
 						bSucceed = GroomBinding_RootProjection::Project(
 							LODGuidesData,
 							TargetMeshData.Get(),
 							TransferredPositions,
-							OutHairGroupDatas[GroupIt].CardsRootData[CardsLODIt]);
+							OutData.CardsRootData[CardsLODIt]);
 
 						if (!bSucceed) 
 						{
@@ -2204,7 +2194,7 @@ static bool InternalBuildBinding_CPU(UGroomBindingAsset* BindingAsset, bool bIni
 			}
 		}
 		
-		GroomBinding_RBFWeighting::ComputeInterpolationWeights(OutHairGroupDatas, BindingAsset->NumInterpolationPoints, BindingAsset->MatchingSection, TargetMeshData.Get(), TransferredPositions);
+		GroomBinding_RBFWeighting::ComputeInterpolationWeights(OutData, BindingAsset->NumInterpolationPoints, BindingAsset->MatchingSection, TargetMeshData.Get(), TransferredPositions);
 		SlowTask.EnterProgressFrame();
 	}
 
@@ -2221,36 +2211,41 @@ static bool InternalBuildBinding_CPU(UGroomBindingAsset* BindingAsset, bool bIni
 	check(OutHairGroupResources.Num() == 0);
 
 	// 2. Convert data to bulk data
-	GroomBinding_BulkCopy::BuildRootBulkData(BindingAsset->HairGroupBulkDatas, OutHairGroupDatas);
-
-	// 3. Update GroomBindingAsset infos
-	{
-		TArray<FGoomBindingGroupInfo>& OutGroupInfos = BindingAsset->GroupInfos;
-		OutGroupInfos.Empty();
-		for (const FHairRootGroupData& Data : OutHairGroupDatas)
-		{
-			FGoomBindingGroupInfo& Info = OutGroupInfos.AddDefaulted_GetRef();
-			Info.SimRootCount = Data.SimRootData.RootCount;
-			Info.SimLODCount  = Data.SimRootData.MeshProjectionLODs.Num();
-			Info.RenRootCount = Data.RenRootData.RootCount;
-			Info.RenLODCount  = Data.RenRootData.MeshProjectionLODs.Num();
-		}
-	}
+	GroomBinding_BulkCopy::BuildRootBulkData(BindingAsset->HairGroupsPlatformData[InGroupIndex], OutData);
 
 	BindingAsset->QueryStatus = UGroomBindingAsset::EQueryStatus::Completed;
-
-	// 4. Optionnally update resources
-	if (bInitResources)
-	{
-		BindingAsset->InitResource();
-	}
 #endif
 	return true;
 }
+void UpdateGroomBindingAssetInfos(UGroomBindingAsset* In);
+
+bool FGroomBindingBuilder::BuildBinding(UGroomBindingAsset* BindingAsset, uint32 InGroupIndex)
+{
+	return InternalBuildBinding_CPU(BindingAsset, InGroupIndex);
+}
 
 bool FGroomBindingBuilder::BuildBinding(UGroomBindingAsset* BindingAsset, bool bInitResources)
 {
-	return InternalBuildBinding_CPU(BindingAsset, bInitResources);
+	bool bOutValid = true;
+
+	// 1. Build binding asset
+	const uint32 GroupCount = BindingAsset->HairGroupsPlatformData.Num();
+	BindingAsset->GroupInfos.SetNum(GroupCount);
+	for (uint32 InGroupIndex = 0; InGroupIndex < GroupCount; ++InGroupIndex)
+	{
+		bOutValid = InternalBuildBinding_CPU(BindingAsset, InGroupIndex) && bOutValid;
+	}
+
+	// 2. Optionnally update resources
+	if (bOutValid && bInitResources)
+	{
+		BindingAsset->InitResource();
+	}
+
+	// 3. Update GroomBindingAsset infos
+	UpdateGroomBindingAssetInfos(BindingAsset);
+
+	return bOutValid;
 }
 
 void FGroomBindingBuilder::GetRootData(
diff --git a/Engine/Plugins/Runtime/HairStrands/Source/HairStrandsCore/Private/GroomManager.cpp b/Engine/Plugins/Runtime/HairStrands/Source/HairStrandsCore/Private/GroomManager.cpp
--- a/Engine/Plugins/Runtime/HairStrands/Source/HairStrandsCore/Private/GroomManager.cpp
+++ b/Engine/Plugins/Runtime/HairStrands/Source/HairStrandsCore/Private/GroomManager.cpp
@@ -1008,7 +1008,7 @@ static void RunHairLODSelection(
 				check(Instance->HairGroupPublicData);
 				Instance->HairGroupPublicData->Allocate(GraphBuilder);
 
-				if (Instance->Strands.RestRootResource)			{ Instance->Strands.RestRootResource->Allocate(GraphBuilder, LoadingType, ResourceStatus, MeshLODIndex); }
+				if (Instance->Strands.RestRootResource)			{ Instance->Strands.RestRootResource->Allocate(GraphBuilder, EHairResourceLoadingType::Async, ResourceStatus, RequestedCurveCount, RequestedPointCount, MeshLODIndex); }
 				if (Instance->Strands.RestResource)				{ Instance->Strands.RestResource->Allocate(GraphBuilder, EHairResourceLoadingType::Async, ResourceStatus, RequestedCurveCount, RequestedPointCount); }
 				if (Instance->Strands.ClusterCullingResource)	{ Instance->Strands.ClusterCullingResource->Allocate(GraphBuilder, EHairResourceLoadingType::Async, ResourceStatus, RequestedCurveCount, RequestedPointCount); }
 				if (Instance->Strands.InterpolationResource)	{ Instance->Strands.InterpolationResource->Allocate(GraphBuilder, EHairResourceLoadingType::Async, ResourceStatus, RequestedCurveCount, RequestedPointCount); }
diff --git a/Engine/Plugins/Runtime/HairStrands/Source/HairStrandsCore/Private/GroomRBFDeformer.cpp b/Engine/Plugins/Runtime/HairStrands/Source/HairStrandsCore/Private/GroomRBFDeformer.cpp
--- a/Engine/Plugins/Runtime/HairStrands/Source/HairStrandsCore/Private/GroomRBFDeformer.cpp
+++ b/Engine/Plugins/Runtime/HairStrands/Source/HairStrandsCore/Private/GroomRBFDeformer.cpp
@@ -479,7 +479,7 @@ void FGroomRBFDeformer::GetRBFDeformedGroomAsset(const UGroomAsset* InGroomAsset
 		ExtractSkeletalVertexPosition(SkeletalMeshData_Target, MeshLODIndex, MeshVertexPositionsBuffer_Target);
 
 		// Apply RBF deformation to each group of guides and render strands
-		const int32 NumGroups = BindingAsset->HairGroupBulkDatas.Num();
+		const int32 NumGroups = BindingAsset->HairGroupsPlatformData.Num();
 
 		// Use the vertices positions from the HairDescription instead of the GroomAsset since the latter
 		// may contain decimated or auto-generated guides depending on the import settings
@@ -541,7 +541,7 @@ void FGroomRBFDeformer::GetRBFDeformedGroomAsset(const UGroomAsset* InGroomAsset
 			else
 			{
 				FHairStrandsRootData SimRootData;
-				FGroomBindingBuilder::GetRootData(SimRootData, BindingAsset->HairGroupBulkDatas[GroupIndex].SimRootBulkData);
+				FGroomBindingBuilder::GetRootData(SimRootData, BindingAsset->HairGroupsPlatformData[GroupIndex].SimRootBulkData);
 
 				DeformedPositions[GroupIndex].GuideStrands = GetDeformedHairStrandsPositions(
 					MeshVertexPositionsBuffer_Target,
@@ -555,7 +555,7 @@ void FGroomRBFDeformer::GetRBFDeformedGroomAsset(const UGroomAsset* InGroomAsset
 			// Get deformed render strands
 			{
 				FHairStrandsRootData RenRootData;
-				FGroomBindingBuilder::GetRootData(RenRootData, BindingAsset->HairGroupBulkDatas[GroupIndex].RenRootBulkData);
+				FGroomBindingBuilder::GetRootData(RenRootData, BindingAsset->HairGroupsPlatformData[GroupIndex].RenRootBulkData);
 
 				DeformedPositions[GroupIndex].RenderStrands = GetDeformedHairStrandsPositions(
 					MeshVertexPositionsBuffer_Target,
@@ -589,7 +589,7 @@ void FGroomRBFDeformer::GetRBFDeformedGroomAsset(const UGroomAsset* InGroomAsset
 				Mesh->ConditionalPostLoad();
 
 				FHairStrandsRootData RenRootData;
-				FGroomBindingBuilder::GetRootData(RenRootData, BindingAsset->HairGroupBulkDatas[Desc.GroupIndex].RenRootBulkData);
+				FGroomBindingBuilder::GetRootData(RenRootData, BindingAsset->HairGroupsPlatformData[Desc.GroupIndex].RenRootBulkData);
 
 				DeformStaticMeshPositions(Mesh, MeshVertexPositionsBuffer_Target, RenRootData.MeshProjectionLODs[MeshLODIndex]);
 			}
@@ -608,7 +608,7 @@ void FGroomRBFDeformer::GetRBFDeformedGroomAsset(const UGroomAsset* InGroomAsset
 				Mesh->ConditionalPostLoad();
 				
 				FHairStrandsRootData RenRootData;
-				FGroomBindingBuilder::GetRootData(RenRootData, BindingAsset->HairGroupBulkDatas[Desc.GroupIndex].RenRootBulkData);
+				FGroomBindingBuilder::GetRootData(RenRootData, BindingAsset->HairGroupsPlatformData[Desc.GroupIndex].RenRootBulkData);
 
 				DeformStaticMeshPositions(Mesh, MeshVertexPositionsBuffer_Target, RenRootData.MeshProjectionLODs[MeshLODIndex]);
 			}
diff --git a/Engine/Plugins/Runtime/HairStrands/Source/HairStrandsCore/Private/GroomResources.cpp b/Engine/Plugins/Runtime/HairStrands/Source/HairStrandsCore/Private/GroomResources.cpp
--- a/Engine/Plugins/Runtime/HairStrands/Source/HairStrandsCore/Private/GroomResources.cpp
+++ b/Engine/Plugins/Runtime/HairStrands/Source/HairStrandsCore/Private/GroomResources.cpp
@@ -622,7 +622,7 @@ void FHairCommonResource::InitRHI()
 	{
 		FRHICommandListImmediate& RHICmdList = FRHICommandListExecutor::GetImmediateCommandList();
 		FRDGBuilder GraphBuilder(RHICmdList);
-		InternalAllocate(GraphBuilder, DummyLODIndex);
+		InternalAllocate(GraphBuilder, HAIR_MAX_NUM_CURVE_PER_GROUP, DummyLODIndex);
 		GraphBuilder.Execute();
 	}
 	else
@@ -663,9 +663,9 @@ void FHairCommonResource::Allocate(FRDGBuilder& GraphBuilder, EHairResourceLoadi
 		{
 			FRenderResource::InitResource(); // Call RenderResource InitResource() so that the resources is marked as initialized
 		}
-		if (!bIsInitialized || !InternalIsLODDataLoaded(InLODIndex))
+		if (!bIsInitialized || !InternalIsLODDataLoaded(InRequestedCurveCount, InRequestedPointCount, InLODIndex))
 		{
-			InternalAllocate(GraphBuilder, InLODIndex);
+			InternalAllocate(GraphBuilder, InRequestedCurveCount, InLODIndex);
 		}
 		bIsInitialized = true;
 
@@ -677,7 +677,7 @@ void FHairCommonResource::Allocate(FRDGBuilder& GraphBuilder, EHairResourceLoadi
 	else if (LoadingType == EHairResourceLoadingType::Async)
 	{
 		// 1. If all requested curve are already loaded, nothing to do
-		if (bIsInitialized && MaxAvailableCurveCount >= InRequestedCurveCount && InternalIsLODDataLoaded(InLODIndex)) 
+		if (bIsInitialized && MaxAvailableCurveCount >= InRequestedCurveCount && InternalIsLODDataLoaded(InRequestedCurveCount, InRequestedPointCount, InLODIndex)) 
 		{ 
 			Status |= EHairResourceStatus::EStatus::Valid; 
 		}
@@ -689,7 +689,7 @@ void FHairCommonResource::Allocate(FRDGBuilder& GraphBuilder, EHairResourceLoadi
 			{
 				FRenderResource::InitResource(); // Call RenderResource InitResource() so that the resources is marked as initialized
 			}
-			InternalAllocate(GraphBuilder, InLODIndex);
+			InternalAllocate(GraphBuilder, StreamingRequest.CurveCount, InLODIndex);
 			bIsInitialized = true;
 
 			// Update the max curve count available
@@ -1178,6 +1178,10 @@ FHairStrandsRestRootResource::FHairStrandsRestRootResource(FHairStrandsRootBulkD
 	FHairCommonResource(EHairStrandsAllocationType::Deferred, InResourceName, InOwnerName),
 	BulkData(InBulkData), CurveType(InCurveType)
 {
+	// MaxAvailableCurveCount is not set to 0, as we need to track the loaded curve data *per LOD*. 
+	// This is done through callback to InternalIsLODDataLoaded()
+	// MaxAvailableCurveCount = 0;
+
 	PopulateFromRootData();
 }
 
@@ -1192,61 +1196,31 @@ void FHairStrandsRestRootResource::PopulateFromRootData()
 		OutLOD.Status 		= FLOD::EStatus::Invalid;
 		OutLOD.SampleCount 	= InLOD.SampleCount;
 	}
-	LODRequests.SetNum(BulkData.Header.LODs.Num());
 }
 
-bool FHairStrandsRestRootResource::InternalIsLODDataLoaded(int32 InLODIndex) const
+bool FHairStrandsRestRootResource::InternalIsLODDataLoaded(uint32 InRequestedCurveCount, uint32 InRequestedPointCount, int32 InLODIndex) const
 {
 	if (InLODIndex >= 0 && InLODIndex < LODs.Num())
 	{
+		// Clamp request to the max supported
+		InRequestedCurveCount = FMath::Min(InRequestedCurveCount, BulkData.Header.RootCount);
+
 		const FLOD& GPUData = LODs[InLODIndex];
-		return GPUData.Status == FLOD::EStatus::Completed || GPUData.Status == FLOD::EStatus::Initialized;
+		return (GPUData.Status == FLOD::EStatus::Completed || GPUData.Status == FLOD::EStatus::Initialized) && InRequestedCurveCount <= GPUData.AvailableCurveCount;
 	}
 	return true;
 }
 
 bool FHairStrandsRestRootResource::InternalIsDataLoaded(uint32 InRequestedCurveCount, uint32 InRequestedPointCount, int32 InLODIndex)
 {
-	bool bIsLoading = false;
-	
-	check(LODs.Num() == BulkData.Header.LODs.Num());
-	if (InLODIndex >= 0 && InLODIndex < LODs.Num())
+	if (StreamingRequest.IsNone())
 	{
-		FBulkDataBatchRequest& LODRequest = LODRequests[InLODIndex];
-		if (LODRequest.IsNone())
-		{
-			FBulkDataBatchRequest::FBatchBuilder Batch = FBulkDataBatchRequest::NewBatch(9);
-			FHairStrandsRootBulkData::FData::FLOD& CPUData = BulkData.Data.LODs[InLODIndex];
-			const bool bHasValidCPUData = CPUData.RootBarycentricBuffer.GetBulkDataSize() > 0;
-			if (bHasValidCPUData)
-			{
-				Batch.Read(CPUData.RootBarycentricBuffer);
-				Batch.Read(CPUData.RootToUniqueTriangleIndexBuffer);
-				Batch.Read(CPUData.UniqueTriangleIndexBuffer);
-				Batch.Read(CPUData.RestUniqueTrianglePositionBuffer);
-			}
-			
-			const bool bHasValidCPUWeights = CPUData.MeshSampleIndicesBuffer.GetBulkDataSize() > 0;
-			if (bHasValidCPUWeights)
-			{
-				Batch.Read(CPUData.MeshInterpolationWeightsBuffer);
-				Batch.Read(CPUData.MeshSampleIndicesBuffer);
-				Batch.Read(CPUData.RestSamplePositionsBuffer);
-			}
-			
-			if (bHasValidCPUData || bHasValidCPUWeights)
-			{
-				Batch.Issue(LODRequest);
-			}
-		}
-
-		bIsLoading = LODRequest.IsCompleted() == false;
+		StreamingRequest.Request(InRequestedCurveCount, InRequestedPointCount, InLODIndex, BulkData, false, false, OwnerName);
 	}
-
-	return !bIsLoading;
+	return StreamingRequest.IsCompleted();
 }
 
-void FHairStrandsRestRootResource::InternalAllocate(FRDGBuilder& GraphBuilder, int32 InLODIndex)
+void FHairStrandsRestRootResource::InternalAllocate(FRDGBuilder& GraphBuilder, uint32 InCurveCount, int32 InLODIndex)
 {
 	// Once empty, the MeshProjectionLODsneeds to be repopulate as it might be re-initialized. 
 	// E.g., when a resource is updated, it is first released, then re-init. 
@@ -1259,55 +1233,37 @@ void FHairStrandsRestRootResource::InternalAllocate(FRDGBuilder& GraphBuilder, i
 	check(LODs.Num() == BulkData.Header.LODs.Num());
 	if (InLODIndex >= 0 && InLODIndex < LODs.Num())
 	{
+		// Clamp to the max supported
+		InCurveCount = FMath::Min(InCurveCount, BulkData.Header.RootCount);
+
 		FLOD& GPUData = LODs[InLODIndex];
-		const bool bIsLODInitialized = GPUData.Status == FLOD::EStatus::Completed || GPUData.Status == FLOD::EStatus::Initialized;
+		const bool bIsLODInitialized = (GPUData.Status == FLOD::EStatus::Completed || GPUData.Status == FLOD::EStatus::Initialized) && InCurveCount <= GPUData.AvailableCurveCount;
 		if (bIsLODInitialized)
 		{
 			return;
 		}
-		
-		LODRequests[InLODIndex] = FBulkDataBatchRequest();
 
 		const FHairStrandsRootBulkData::FHeader::FLOD& LODHeader = BulkData.Header.LODs[InLODIndex];
 		FHairStrandsRootBulkData::FData::FLOD& CPUData = BulkData.Data.LODs[InLODIndex];
-		const bool bHasValidCPUData = CPUData.RootBarycentricBuffer.GetBulkDataSize() > 0;
-		if (bHasValidCPUData)
 		{
-			GPUData.Status = FLOD::EStatus::Completed;
 			
-			InternalCreateVertexBufferRDG_FromBulkData<FHairStrandsRootBarycentricFormat>(GraphBuilder, CPUData.RootBarycentricBuffer, BulkData.Header.RootCount, GPUData.RootBarycentricBuffer, ToHairResourceDebugName(HAIRSTRANDS_RESOUCE_NAME(CurveType, Hair.StrandsRestRoot_RootTriangleBarycentricBuffer), ResourceName), OwnerName, EHairResourceUsageType::Static);
-			InternalCreateVertexBufferRDG_FromBulkData<FHairStrandsRootToUniqueTriangleIndexFormat>(GraphBuilder, CPUData.RootToUniqueTriangleIndexBuffer, BulkData.Header.RootCount, GPUData.RootToUniqueTriangleIndexBuffer, ToHairResourceDebugName(HAIRSTRANDS_RESOUCE_NAME(CurveType, Hair.StrandsRestRoot_RootToUniqueTriangleIndexBuffer), ResourceName), OwnerName, EHairResourceUsageType::Static);
-			InternalCreateVertexBufferRDG_FromBulkData<FHairStrandsUniqueTriangleIndexFormat>(GraphBuilder, CPUData.UniqueTriangleIndexBuffer, LODHeader.UniqueTriangleCount, GPUData.UniqueTriangleIndexBuffer, ToHairResourceDebugName(HAIRSTRANDS_RESOUCE_NAME(CurveType, Hair.StrandsRestRoot_UniqueTriangleIndexBuffer), ResourceName), OwnerName, EHairResourceUsageType::Static);
-			InternalCreateVertexBufferRDG_FromBulkData<FHairStrandsMeshTrianglePositionFormat>(GraphBuilder, CPUData.RestUniqueTrianglePositionBuffer, LODHeader.UniqueTriangleCount * 3, GPUData.RestUniqueTrianglePositionBuffer, ToHairResourceDebugName(HAIRSTRANDS_RESOUCE_NAME(CurveType, Hair.StrandsRestRoot_RestUniqueTrianglePosition0Buffer), ResourceName), OwnerName, EHairResourceUsageType::Static);
-		}
-		else
-		{
-			GPUData.Status = FLOD::EStatus::Initialized;
-			
-			InternalCreateVertexBufferRDG<FHairStrandsRootBarycentricFormat>(GraphBuilder, BulkData.Header.RootCount, GPUData.RootBarycentricBuffer, ToHairResourceDebugName(HAIRSTRANDS_RESOUCE_NAME(CurveType, Hair.StrandsRestRoot_RootBarycentricBuffer), ResourceName), OwnerName, EHairResourceUsageType::Dynamic);
-			InternalCreateVertexBufferRDG<FHairStrandsRootToUniqueTriangleIndexFormat>(GraphBuilder, BulkData.Header.RootCount, GPUData.RootToUniqueTriangleIndexBuffer, ToHairResourceDebugName(HAIRSTRANDS_RESOUCE_NAME(CurveType, Hair.StrandsRestRoot_RootToUniqueTriangleIndexBuffer), ResourceName), OwnerName, EHairResourceUsageType::Dynamic);
-			InternalCreateVertexBufferRDG<FHairStrandsUniqueTriangleIndexFormat>(GraphBuilder, LODHeader.UniqueTriangleCount, GPUData.UniqueTriangleIndexBuffer, ToHairResourceDebugName(HAIRSTRANDS_RESOUCE_NAME(CurveType, Hair.StrandsRestRoot_UniqueTriangleIndexBuffer), ResourceName), OwnerName, EHairResourceUsageType::Dynamic);
-			
-			// Create buffers. Initialization will be done by render passes
-			InternalCreateVertexBufferRDG<FHairStrandsMeshTrianglePositionFormat>(GraphBuilder, LODHeader.UniqueTriangleCount * 3, GPUData.RestUniqueTrianglePositionBuffer, ToHairResourceDebugName(HAIRSTRANDS_RESOUCE_NAME(CurveType, Hair.StrandsRestRoot_RestUniqueTrianglePosition0Buffer), ResourceName), OwnerName, EHairResourceUsageType::Dynamic);
+			InternalCreateVertexBufferRDG_FromHairBulkData<FHairStrandsRootBarycentricFormat>(GraphBuilder, CPUData.RootBarycentricBuffer, BulkData.Header.RootCount, GPUData.RootBarycentricBuffer, ToHairResourceDebugName(HAIRSTRANDS_RESOUCE_NAME(CurveType, Hair.StrandsRestRoot_RootTriangleBarycentricBuffer), ResourceName), OwnerName, EHairResourceUsageType::Static);
+			InternalCreateVertexBufferRDG_FromHairBulkData<FHairStrandsRootToUniqueTriangleIndexFormat>(GraphBuilder, CPUData.RootToUniqueTriangleIndexBuffer, BulkData.Header.RootCount, GPUData.RootToUniqueTriangleIndexBuffer, ToHairResourceDebugName(HAIRSTRANDS_RESOUCE_NAME(CurveType, Hair.StrandsRestRoot_RootToUniqueTriangleIndexBuffer), ResourceName), OwnerName, EHairResourceUsageType::Static);
+			InternalCreateVertexBufferRDG_FromHairBulkData<FHairStrandsUniqueTriangleIndexFormat>(GraphBuilder, CPUData.UniqueTriangleIndexBuffer, LODHeader.UniqueTriangleCount, GPUData.UniqueTriangleIndexBuffer, ToHairResourceDebugName(HAIRSTRANDS_RESOUCE_NAME(CurveType, Hair.StrandsRestRoot_UniqueTriangleIndexBuffer), ResourceName), OwnerName, EHairResourceUsageType::Static);
+			InternalCreateVertexBufferRDG_FromHairBulkData<FHairStrandsMeshTrianglePositionFormat>(GraphBuilder, CPUData.RestUniqueTrianglePositionBuffer, LODHeader.UniqueTriangleCount * 3, GPUData.RestUniqueTrianglePositionBuffer, ToHairResourceDebugName(HAIRSTRANDS_RESOUCE_NAME(CurveType, Hair.StrandsRestRoot_RestUniqueTrianglePosition0Buffer), ResourceName), OwnerName, EHairResourceUsageType::Static);
 		}
 
-		GPUData.SampleCount = LODHeader.SampleCount;
-		const bool bHasValidCPUWeights = CPUData.MeshSampleIndicesBuffer.GetBulkDataSize() > 0;
-		if (bHasValidCPUWeights)
+		if (LODHeader.SampleCount > 0)
 		{
 			const uint32 InteroplationWeightCount = CPUData.MeshInterpolationWeightsBuffer.GetBulkDataSize() / sizeof(FHairStrandsWeightFormat::Type);
-			InternalCreateVertexBufferRDG_FromBulkData<FHairStrandsWeightFormat>(GraphBuilder, CPUData.MeshInterpolationWeightsBuffer, InteroplationWeightCount, GPUData.MeshInterpolationWeightsBuffer, ToHairResourceDebugName(HAIRSTRANDS_RESOUCE_NAME(CurveType, Hair.StrandsRestRoot_MeshInterpolationWeightsBuffer), ResourceName), OwnerName, EHairResourceUsageType::Static);
-			InternalCreateVertexBufferRDG_FromBulkData<FHairStrandsIndexFormat>(GraphBuilder, CPUData.MeshSampleIndicesBuffer, LODHeader.SampleCount, GPUData.MeshSampleIndicesBuffer, ToHairResourceDebugName(HAIRSTRANDS_RESOUCE_NAME(CurveType, Hair.StrandsRestRoot_MeshSampleIndicesBuffer), ResourceName), OwnerName, EHairResourceUsageType::Static);
-			InternalCreateVertexBufferRDG_FromBulkData<FHairStrandsMeshTrianglePositionFormat>(GraphBuilder, CPUData.RestSamplePositionsBuffer, LODHeader.SampleCount, GPUData.RestSamplePositionsBuffer, ToHairResourceDebugName(HAIRSTRANDS_RESOUCE_NAME(CurveType, Hair.StrandsRestRoot_RestSamplePositionsBuffer), ResourceName), OwnerName, EHairResourceUsageType::Static);
-		}
-		else
-		{
-			// TODO: do not allocate these resources, since they won't be used
-			InternalCreateVertexBufferRDG<FHairStrandsWeightFormat>(GraphBuilder, (LODHeader.SampleCount+4) * (LODHeader.SampleCount+4), GPUData.MeshInterpolationWeightsBuffer, ToHairResourceDebugName(HAIRSTRANDS_RESOUCE_NAME(CurveType, Hair.StrandsRestRoot_MeshInterpolationWeightsBuffer), ResourceName), OwnerName, EHairResourceUsageType::Dynamic);
-			InternalCreateVertexBufferRDG<FHairStrandsIndexFormat>(GraphBuilder, LODHeader.SampleCount, GPUData.MeshSampleIndicesBuffer, ToHairResourceDebugName(HAIRSTRANDS_RESOUCE_NAME(CurveType, Hair.StrandsRestRoot_MeshSampleIndicesBuffer), ResourceName), OwnerName, EHairResourceUsageType::Dynamic);
-			InternalCreateVertexBufferRDG<FHairStrandsMeshTrianglePositionFormat>(GraphBuilder, LODHeader.SampleCount, GPUData.RestSamplePositionsBuffer, ToHairResourceDebugName(HAIRSTRANDS_RESOUCE_NAME(CurveType, Hair.StrandsRestRoot_RestSamplePositionsBuffer), ResourceName), OwnerName, EHairResourceUsageType::Dynamic);
+			InternalCreateVertexBufferRDG_FromHairBulkData<FHairStrandsWeightFormat>(GraphBuilder, CPUData.MeshInterpolationWeightsBuffer, InteroplationWeightCount, GPUData.MeshInterpolationWeightsBuffer, ToHairResourceDebugName(HAIRSTRANDS_RESOUCE_NAME(CurveType, Hair.StrandsRestRoot_MeshInterpolationWeightsBuffer), ResourceName), OwnerName, EHairResourceUsageType::Static);
+			InternalCreateVertexBufferRDG_FromHairBulkData<FHairStrandsIndexFormat>(GraphBuilder, CPUData.MeshSampleIndicesBuffer, LODHeader.SampleCount, GPUData.MeshSampleIndicesBuffer, ToHairResourceDebugName(HAIRSTRANDS_RESOUCE_NAME(CurveType, Hair.StrandsRestRoot_MeshSampleIndicesBuffer), ResourceName), OwnerName, EHairResourceUsageType::Static);
+			InternalCreateVertexBufferRDG_FromHairBulkData<FHairStrandsMeshTrianglePositionFormat>(GraphBuilder, CPUData.RestSamplePositionsBuffer, LODHeader.SampleCount, GPUData.RestSamplePositionsBuffer, ToHairResourceDebugName(HAIRSTRANDS_RESOUCE_NAME(CurveType, Hair.StrandsRestRoot_RestSamplePositionsBuffer), ResourceName), OwnerName, EHairResourceUsageType::Static);
 		}
+
+		GPUData.Status = FLOD::EStatus::Completed;
+		GPUData.AvailableCurveCount = InCurveCount;
+		GPUData.SampleCount = LODHeader.SampleCount;
 	}
 }
 
@@ -1324,9 +1280,9 @@ void FHairStrandsRestRootResource::InternalRelease()
 		GPUData.MeshInterpolationWeightsBuffer.Release();
 		GPUData.MeshSampleIndicesBuffer.Release();
 		GPUData.RestSamplePositionsBuffer.Release();
+		GPUData.AvailableCurveCount = 0;
 	}
 	LODs.Empty();
-	LODRequests.Empty();
 }
 
 /////////////////////////////////////////////////////////////////////////////////////////
@@ -1355,17 +1311,20 @@ FHairStrandsDeformedRootResource::FHairStrandsDeformedRootResource(const FHairSt
 	}
 }
 
-bool FHairStrandsDeformedRootResource::InternalIsLODDataLoaded(int32 InLODIndex) const
+bool FHairStrandsDeformedRootResource::InternalIsLODDataLoaded(uint32 InRequestedCurveCount, uint32 InRequestedPointCount, int32 InLODIndex) const
 {
 	if (InLODIndex >= 0 && InLODIndex < LODs.Num())
 	{
+		// Clamp request to the max supported
+		InRequestedCurveCount = FMath::Min(InRequestedCurveCount, RootCount);
+
 		const FLOD& GPUData = LODs[InLODIndex];
-		return GPUData.Status == FLOD::EStatus::Completed || GPUData.Status == FLOD::EStatus::Initialized;
+		return (GPUData.Status == FLOD::EStatus::Completed || GPUData.Status == FLOD::EStatus::Initialized) && InRequestedCurveCount <= GPUData.AvailableCurveCount;
 	}
 	return true;
 }
 
-void FHairStrandsDeformedRootResource::InternalAllocate(FRDGBuilder& GraphBuilder, int32 InLODIndex)
+void FHairStrandsDeformedRootResource::InternalAllocate(FRDGBuilder& GraphBuilder, uint32 InCurveCount, int32 InLODIndex)
 {
 	if (RootCount > 0 && InLODIndex >= 0 && InLODIndex < LODs.Num())
 	{
@@ -1373,6 +1332,7 @@ void FHairStrandsDeformedRootResource::InternalAllocate(FRDGBuilder& GraphBuilde
 		if (LOD.Status == FLOD::EStatus::Invalid)
 		{		
 			LOD.Status = FLOD::EStatus::Initialized;
+			LOD.AvailableCurveCount = FMath::Min(InCurveCount, RootCount);
 			if (LOD.SampleCount > 0)
 			{
 				InternalCreateVertexBufferRDG<FHairStrandsMeshTrianglePositionFormat>(GraphBuilder, LOD.SampleCount, LOD.DeformedSamplePositionsBuffer[0], ToHairResourceDebugName(HAIRSTRANDS_RESOUCE_NAME(CurveType, Hair.StrandsRootDeformed_DeformedSamplePositionsBuffer0), ResourceName), OwnerName, EHairResourceUsageType::Dynamic);
@@ -1402,6 +1362,7 @@ void FHairStrandsDeformedRootResource::InternalRelease()
 	for (FLOD& GPUData : LODs)
 	{
 		GPUData.Status = FLOD::EStatus::Invalid;
+		GPUData.AvailableCurveCount = 0;
 		GPUData.DeformedUniqueTrianglePositionBuffer[0].Release();
 		GPUData.DeformedSamplePositionsBuffer[0].Release();
 		GPUData.MeshSampleWeightsBuffer[0].Release();
diff --git a/Engine/Plugins/Runtime/HairStrands/Source/HairStrandsCore/Private/HairStrandsCore.cpp b/Engine/Plugins/Runtime/HairStrands/Source/HairStrandsCore/Private/HairStrandsCore.cpp
--- a/Engine/Plugins/Runtime/HairStrands/Source/HairStrandsCore/Private/HairStrandsCore.cpp
+++ b/Engine/Plugins/Runtime/HairStrands/Source/HairStrandsCore/Private/HairStrandsCore.cpp
@@ -231,7 +231,7 @@ UGroomBindingAsset* FHairStrandsCore::CreateGroomBindingAsset(EGroomBindingMeshT
 			Out->SourceGeometryCache = Cast<UGeometryCache>(Source);
 			Out->TargetGeometryCache = Cast<UGeometryCache>(Target);
 		}
-		Out->HairGroupBulkDatas.Reserve(GroomAsset->HairGroupsPlatformData.Num());
+		Out->HairGroupsPlatformData.Reserve(GroomAsset->HairGroupsPlatformData.Num());
 		Out->NumInterpolationPoints = NumInterpolationPoints;
 		Out->MatchingSection = MatchingSection;
 		Out->MarkPackageDirty();
diff --git a/Engine/Plugins/Runtime/HairStrands/Source/HairStrandsCore/Private/HairStrandsDatas.cpp b/Engine/Plugins/Runtime/HairStrands/Source/HairStrandsCore/Private/HairStrandsDatas.cpp
--- a/Engine/Plugins/Runtime/HairStrands/Source/HairStrandsCore/Private/HairStrandsDatas.cpp
+++ b/Engine/Plugins/Runtime/HairStrands/Source/HairStrandsCore/Private/HairStrandsDatas.cpp
@@ -126,7 +126,7 @@ namespace HairStrands
 #endif
 }
 
-void FHairStrandsBulkCommon::Write_DDC(UObject* Owner, TArray<UE::DerivedData::FCachePutValueRequest>& Out)
+void FHairStrandsBulkCommon::Write_DDC(UObject* Owner, TArray<UE::DerivedData::FCachePutValueRequest>& Out, int32 InLODIndex)
 {
 #if WITH_EDITORONLY_DATA
 	FHairStrandsBulkCommon::FQuery Q;
@@ -134,6 +134,7 @@ void FHairStrandsBulkCommon::Write_DDC(UObject* Owner, TArray<UE::DerivedData::F
 	Q.OutWriteDDC = &Out;
 	Q.DerivedDataKey = &DerivedDataKey;
 	Q.Owner = Owner;
+	Q.LODIndex = InLODIndex;
 	GetResources(Q);
 #endif
 }
@@ -163,7 +164,7 @@ void FHairStrandsBulkCommon::Read_IO(FHairStreamingRequest* In, FBulkDataBatchRe
 	GetResources(Q);
 	Q.OutReadIO->Issue(Out);
 }
-void FHairStrandsBulkCommon::Write_IO(UObject* Owner, FArchive& Out)
+void FHairStrandsBulkCommon::Write_IO(UObject* Owner, FArchive& Out, int32 InLODIndex)
 {
 	GetResourceVersion(Out);
 
@@ -171,18 +172,19 @@ void FHairStrandsBulkCommon::Write_IO(UObject* Owner, FArchive& Out)
 	Q.Type = FHairStrandsBulkCommon::FQuery::ReadWriteIO;
 	Q.OutWriteIO = &Out;
 	Q.Owner = Owner;
+	Q.LODIndex = InLODIndex;
 	GetResources(Q);
 }
 
 void FHairStrandsBulkCommon::Serialize(FArchive& Ar, UObject* Owner)
 {
 	SerializeHeader(Ar, Owner);
-	SerializeData(Ar, Owner);
+	SerializeData(Ar, Owner, -1 /*InLODIndex*/);
 }
 
-void FHairStrandsBulkCommon::SerializeData(FArchive& Ar, UObject* Owner)
+void FHairStrandsBulkCommon::SerializeData(FArchive& Ar, UObject* Owner, int32 InLODIndex)
 {
-	Write_IO(Owner, Ar);
+	Write_IO(Owner, Ar, InLODIndex);
 }
 
 void FHairStrandsBulkCommon::FQuery::Add(FHairBulkContainer& In, const TCHAR* InSuffix, uint32 InOffset, uint32 InSize) 
@@ -197,10 +199,13 @@ void FHairStrandsBulkCommon::FQuery::Add(FHairBulkContainer& In, const TCHAR* In
 		FMemory::Memcpy(WriteData.GetData(), In.Data.Lock(LOCK_READ_ONLY), DataSizeInByte);
 		In.Data.Unlock();
 
+		// Optional LOD suffix
+		const FString InLODSuffix = GetLODIndex() >= 0 ? FString(TEXT("_LOD")) + FString::FromInt(GetLODIndex()) : FString();
+
 		using namespace UE::DerivedData;
 		FCachePutValueRequest& Out = OutWriteDDC->AddDefaulted_GetRef();
 		if (Owner) { Out.Name = Owner->GetPathName(); }
-		Out.Key 	= ConvertLegacyCacheKey(*DerivedDataKey + InSuffix);
+		Out.Key 	= ConvertLegacyCacheKey(*DerivedDataKey + InSuffix + InLODSuffix);
 		Out.Value 	= FValue::Compress(MakeSharedBufferFromArray(MoveTemp(WriteData)));
 		Out.Policy 	= ECachePolicy::Default;
 		Out.UserData= 0;
@@ -221,12 +226,15 @@ void FHairStrandsBulkCommon::FQuery::Add(FHairBulkContainer& In, const TCHAR* In
 		Chunk.TotalSize = InSize;
 		In.ChunkRequest = &Chunk;
 
+		// Optional LOD suffix
+		const FString InLODSuffix = GetLODIndex() >= 0 ? FString(TEXT("_LOD")) + FString::FromInt(GetLODIndex()) : FString();
+
 		// 2. Fill in actual DDC request
 		check(OutReadDDC);
 		using namespace UE::DerivedData;
 		FCacheGetChunkRequest& Out = OutReadDDC->AddDefaulted_GetRef();
 		Out.Id			= FValueId::Null; 	// HairStrands::HairStrandsValueId : This is only needed for cache record, not cache value.
-		Out.Key			= ConvertLegacyCacheKey(*DerivedDataKey + InSuffix);
+		Out.Key			= ConvertLegacyCacheKey(*DerivedDataKey + InSuffix + InLODSuffix);
 		Out.RawOffset	= InOffset;
 		Out.RawSize		= InSize != 0 ? InSize-InOffset : MAX_uint64;
 		Out.RawHash		= FIoHash();
@@ -317,10 +325,12 @@ void FHairStreamingRequest::Request(uint32 InRequestedCurveCount, uint32 InReque
 	{
 		CurveCount = 0;
 		PointCount = 0;
+		LODIndex = -1;
 		return;
 	}
 	CurveCount = InRequestedCurveCount;
 	PointCount = InRequestedPointCount;
+	LODIndex = InLODIndex;
 
 #if !WITH_EDITORONLY_DATA
 	{ 
@@ -760,55 +770,42 @@ const TArray<uint32>& FHairStrandsRootBulkData::GetValidSectionIndices(int32 LOD
 	return Header.LODs[LODIndex].UniqueSectionIndices;
 }
 
-void FHairStrandsRootBulkData::SerializeData(FArchive& Ar, UObject* Owner)
+uint32 FHairStrandsRootBulkData::GetResourceCount() const
+{
+	return 7;
+}
+
+void FHairStrandsRootBulkData::GetResources(FQuery& Out)
 {
+	const int32 LODIndex = Out.GetLODIndex();
+	check(LODIndex >= 0);
 	check(Header.LODs.Num() == Data.LODs.Num());
-	for (uint32 LODIt = 0, LODCount=Header.LODs.Num(); LODIt < LODCount; ++LODIt)
-	{
-		const FHeader::FLOD& HeaderLOD = Header.LODs[LODIt];
-		FData::FLOD& DataLOD = Data.LODs[LODIt];
+	check(LODIndex < Data.LODs.Num());
 
-		const int32 ChunkIndex = 0;
-		bool bAttemptFileMapping = false;
-	
-		if (Ar.IsSaving())
-		{
-			const uint32 BulkFlags = BULKDATA_Force_NOT_InlinePayload;
-	
-			DataLOD.RootBarycentricBuffer.SetBulkDataFlags(BulkFlags);
-			DataLOD.RootToUniqueTriangleIndexBuffer.SetBulkDataFlags(BulkFlags);
-			DataLOD.UniqueTriangleIndexBuffer.SetBulkDataFlags(BulkFlags);
-			DataLOD.RestUniqueTrianglePositionBuffer.SetBulkDataFlags(BulkFlags);
-	
-			DataLOD.MeshInterpolationWeightsBuffer.SetBulkDataFlags(BulkFlags);
-			DataLOD.MeshSampleIndicesBuffer.SetBulkDataFlags(BulkFlags);
-			DataLOD.RestSamplePositionsBuffer.SetBulkDataFlags(BulkFlags);
-		}
+	const FHeader::FLOD& HeaderLOD = Header.LODs[LODIndex];
+	FData::FLOD& DataLOD = Data.LODs[LODIndex];
 
-		DataLOD.RootBarycentricBuffer.Serialize(Ar, Owner, ChunkIndex, bAttemptFileMapping);
-		DataLOD.RootToUniqueTriangleIndexBuffer.Serialize(Ar, Owner, ChunkIndex, bAttemptFileMapping);
-		DataLOD.UniqueTriangleIndexBuffer.Serialize(Ar, Owner, ChunkIndex, bAttemptFileMapping);
-		DataLOD.RestUniqueTrianglePositionBuffer.Serialize(Ar, Owner, ChunkIndex, bAttemptFileMapping);
-	
-		if (HeaderLOD.SampleCount)
-		{
-			DataLOD.MeshInterpolationWeightsBuffer.Serialize(Ar, Owner, ChunkIndex, bAttemptFileMapping);
-			DataLOD.MeshSampleIndicesBuffer.Serialize(Ar, Owner, ChunkIndex, bAttemptFileMapping);
-			DataLOD.RestSamplePositionsBuffer.Serialize(Ar, Owner, ChunkIndex, bAttemptFileMapping);
-		}
+	// Translate requested curve count into chunk/offset/size to be read
+	// * RBF points are loaded all at once
+	// * Unique triangle data are loaded all at one
+	// * Root data are loaded based on the curve request
+	uint32 RootCount = Header.RootCount;
+	if (Out.Type == FHairStrandsBulkCommon::FQuery::ReadIO || Out.Type == FHairStrandsBulkCommon::FQuery::ReadDDC)
+	{
+		RootCount = FMath::Min(Header.RootCount, Out.GetCurveCount());
 	}
-}
 
-void FHairStrandsRootBulkData::Serialize(FArchive& Ar, UObject* Owner)
-{
-	// Root data don't have UObject references
-	if (!Ar.IsLoading() && !Ar.IsSaving() && Ar.IsObjectReferenceCollector())
+	Out.Add(DataLOD.RootToUniqueTriangleIndexBuffer, 	TEXT("_RootToUniqueTriangleIndexBuffer"), 		DataLOD.RootToUniqueTriangleIndexBuffer.LoadedSize, 	RootCount * Header.Strides.RootToUniqueTriangleIndexBufferStride);
+	Out.Add(DataLOD.RootBarycentricBuffer, 				TEXT("_RootBarycentricBuffer"), 				DataLOD.RootBarycentricBuffer.LoadedSize, 				RootCount * Header.Strides.RootBarycentricBufferStride);
+	Out.Add(DataLOD.UniqueTriangleIndexBuffer, 			TEXT("_UniqueTriangleIndexBuffer"), 			DataLOD.UniqueTriangleIndexBuffer.LoadedSize, 			HeaderLOD.UniqueTriangleCount * Header.Strides.UniqueTriangleIndexBufferStride);
+	Out.Add(DataLOD.RestUniqueTrianglePositionBuffer, 	TEXT("_RestUniqueTrianglePositionBuffer"), 		DataLOD.RestUniqueTrianglePositionBuffer.LoadedSize, 	HeaderLOD.UniqueTriangleCount * Header.Strides.RestUniqueTrianglePositionBufferStride);
+
+	if (HeaderLOD.SampleCount > 0)
 	{
-		return;
+		Out.Add(DataLOD.MeshInterpolationWeightsBuffer, TEXT("_MeshInterpolationWeightsBuffer"), 		DataLOD.MeshInterpolationWeightsBuffer.LoadedSize, 		HeaderLOD.SampleCount * Header.Strides.MeshInterpolationWeightsBufferStride);
+		Out.Add(DataLOD.MeshSampleIndicesBuffer, 		TEXT("_MeshSampleIndicesBuffer"), 				DataLOD.MeshSampleIndicesBuffer.LoadedSize, 			HeaderLOD.SampleCount * Header.Strides.MeshSampleIndicesBufferStride);
+		Out.Add(DataLOD.RestSamplePositionsBuffer, 		TEXT("_RestSamplePositionsBuffer"), 			DataLOD.RestSamplePositionsBuffer.LoadedSize, 			HeaderLOD.SampleCount * Header.Strides.RestSamplePositionsBufferStride);
 	}
-
-	SerializeHeader(Ar, Owner);
-	SerializeData(Ar, Owner);
 }
 
 void FHairStrandsRootBulkData::Reset()
@@ -828,19 +825,19 @@ void FHairStrandsRootBulkData::Reset()
 		LOD.RestUniqueTrianglePositionBuffer.RemoveBulkData();
 
 		// Reset the bulk byte buffer to ensure the (serialize) data size is reset to 0
-		LOD.RootBarycentricBuffer 			 = FByteBulkData();
-		LOD.RootToUniqueTriangleIndexBuffer  = FByteBulkData();
-		LOD.UniqueTriangleIndexBuffer		 = FByteBulkData();
-		LOD.RestUniqueTrianglePositionBuffer = FByteBulkData();
+		LOD.RootBarycentricBuffer 			 = FHairBulkContainer();
+		LOD.RootToUniqueTriangleIndexBuffer  = FHairBulkContainer();
+		LOD.UniqueTriangleIndexBuffer		 = FHairBulkContainer();
+		LOD.RestUniqueTrianglePositionBuffer = FHairBulkContainer();
 
 		// RBF
 		LOD.MeshInterpolationWeightsBuffer.RemoveBulkData();
 		LOD.MeshSampleIndicesBuffer.RemoveBulkData();
 		LOD.RestSamplePositionsBuffer.RemoveBulkData();
 
-		LOD.MeshInterpolationWeightsBuffer 	 = FByteBulkData();
-		LOD.MeshSampleIndicesBuffer 		 = FByteBulkData();
-		LOD.RestSamplePositionsBuffer 		 = FByteBulkData();
+		LOD.MeshInterpolationWeightsBuffer 	 = FHairBulkContainer();
+		LOD.MeshSampleIndicesBuffer 		 = FHairBulkContainer();
+		LOD.RestSamplePositionsBuffer 		 = FHairBulkContainer();
 	}
 	Data.LODs.Empty();
 }
diff --git a/Engine/Plugins/Runtime/HairStrands/Source/HairStrandsCore/Public/GroomBindingAsset.h b/Engine/Plugins/Runtime/HairStrands/Source/HairStrandsCore/Public/GroomBindingAsset.h
--- a/Engine/Plugins/Runtime/HairStrands/Source/HairStrandsCore/Public/GroomBindingAsset.h
+++ b/Engine/Plugins/Runtime/HairStrands/Source/HairStrandsCore/Public/GroomBindingAsset.h
@@ -47,14 +47,6 @@ enum class EGroomBindingMeshType : uint8
 	GeometryCache
 };
 
-/** Binding bulk data */
-struct FHairGroupBulkData
-{
-	FHairStrandsRootBulkData		SimRootBulkData;
-	FHairStrandsRootBulkData		RenRootBulkData;
-	TArray<FHairStrandsRootBulkData>CardsRootBulkData;
-};
-
 /**
  * Implements an asset that can be used to store binding information between a groom and a skeletal mesh
  */
@@ -117,8 +109,16 @@ public:
 	   when the binding asset is recomputed */
 	TQueue<FHairGroupResource> HairGroupResourcesToDelete;
 
+	/** Binding bulk data */
+	struct FHairGroupPlatformData
+	{
+		FHairStrandsRootBulkData		SimRootBulkData;
+		FHairStrandsRootBulkData		RenRootBulkData;
+		TArray<FHairStrandsRootBulkData>CardsRootBulkData;
+
+	};
 	/** Root bulk data for each hair gruops */
-	TArray<FHairGroupBulkData> HairGroupBulkDatas;
+	TArray<FHairGroupPlatformData> HairGroupsPlatformData;
 
 	//~ Begin UObject Interface.
 	virtual void GetResourceSizeEx(FResourceSizeEx& CumulativeResourceSize) override;
@@ -185,12 +185,14 @@ public:
 	void Build();
 
 	void CacheDerivedDatas();
+	void CacheDerivedDatas(uint32 InGroupIndex, const FString KeySuffix, bool& bOutValid, bool& bOutReloadResource);
+
 	void InvalidateBinding();
 	void InvalidateBinding(class USkeletalMesh*);
 	bool bRegisterSourceMeshCallback = false;
 	bool bRegisterTargetMeshCallback = false;
 	bool bRegisterGroomAssetCallback = false;
-	FString CachedDerivedDataKey;
+	TArray<FString> CachedDerivedDataKey;
 #endif
 	bool bIsValid = false;
 };
diff --git a/Engine/Plugins/Runtime/HairStrands/Source/HairStrandsCore/Public/GroomBindingBuilder.h b/Engine/Plugins/Runtime/HairStrands/Source/HairStrandsCore/Public/GroomBindingBuilder.h
--- a/Engine/Plugins/Runtime/HairStrands/Source/HairStrandsCore/Public/GroomBindingBuilder.h
+++ b/Engine/Plugins/Runtime/HairStrands/Source/HairStrandsCore/Public/GroomBindingBuilder.h
@@ -17,7 +17,8 @@ struct HAIRSTRANDSCORE_API FGroomBindingBuilder
 	static FString GetVersion();
 
 	// Build binding asset data
-	static bool BuildBinding(class UGroomBindingAsset* BindingAsset, bool bInitResources);
+	static bool BuildBinding(class UGroomBindingAsset* BindingAsset, bool bInitResource);
+	static bool BuildBinding(class UGroomBindingAsset* BindingAsset, uint32 InGroupIndex);
 
 	// Extract root data from bulk data
 	static void GetRootData(FHairStrandsRootData& Out, const FHairStrandsRootBulkData& In);
diff --git a/Engine/Plugins/Runtime/HairStrands/Source/HairStrandsCore/Public/GroomResources.h b/Engine/Plugins/Runtime/HairStrands/Source/HairStrandsCore/Public/GroomResources.h
--- a/Engine/Plugins/Runtime/HairStrands/Source/HairStrandsCore/Public/GroomResources.h
+++ b/Engine/Plugins/Runtime/HairStrands/Source/HairStrandsCore/Public/GroomResources.h
@@ -109,10 +109,10 @@ struct FHairCommonResource : public FRenderResource
 
 	virtual void InternalAllocate() {}
 	virtual void InternalAllocate(FRDGBuilder& GraphBuilder) {}
-	virtual void InternalAllocate(FRDGBuilder& GraphBuilder, int32 InLODIndex) { InternalAllocate(GraphBuilder); }
+	virtual void InternalAllocate(FRDGBuilder& GraphBuilder, uint32 InCurveCount, int32 InLODIndex) { InternalAllocate(GraphBuilder); }
 	virtual void InternalRelease() {}
 	virtual bool InternalIsDataLoaded(uint32 InRequestedCurveCount, uint32 InRequestedPointCount, int32 InLODIndex) { return true; }
-	virtual bool InternalIsLODDataLoaded(int32 InLODIndex) const { return true; }
+	virtual bool InternalIsLODDataLoaded(uint32 InRequestedCurveCount, uint32 InRequestedPointCount, int32 InLODIndex) const { return true; }
 
 	bool bUseRenderGraph = true;
 	bool bIsInitialized = false;
@@ -351,10 +351,10 @@ struct FHairStrandsRestRootResource : public FHairCommonResource
 	FHairStrandsRestRootResource(FHairStrandsRootBulkData& BulkData, EHairStrandsResourcesType CurveType, const FHairResourceName& ResourceName, const FName& OwnerName);
 
 	/* Init/Release buffers */
-	virtual void InternalAllocate(FRDGBuilder& GraphBuilder, int32 InLODIndex) override;
+	virtual void InternalAllocate(FRDGBuilder& GraphBuilder, uint32 InCurveCount, int32 InLODIndex) override;
 	virtual void InternalRelease() override;
 	virtual bool InternalIsDataLoaded(uint32 InRequestedCurveCount, uint32 InRequestedPointCount, int32 InLODIndex) override;
-	virtual bool InternalIsLODDataLoaded(int32 InLODIndex) const override;
+	virtual bool InternalIsLODDataLoaded(uint32 InRequestedCurveCount, uint32 InRequestedPointCount, int32 InLODIndex) const override;
 
 	/* Get the resource name */
 	virtual FString GetFriendlyName() const override { return TEXT("FHairStrandsRestRootResource"); }
@@ -390,6 +390,7 @@ struct FHairStrandsRestRootResource : public FHairCommonResource
 		const bool IsValid() const { return Status == EStatus::Completed; }
 		EStatus Status = EStatus::Invalid;
 		int32 LODIndex = -1;
+		uint32 AvailableCurveCount = 0;
 
 		/* Triangle on which a root is attached */
 		/* When the projection is done with source to target mesh transfer, the projection indices does not match.
@@ -412,9 +413,6 @@ struct FHairStrandsRestRootResource : public FHairCommonResource
 	/* Store the hair projection information for each mesh LOD */
 	TArray<FLOD> LODs;
 
-	/* LOD bulk data requests */
-	TArray<FBulkDataBatchRequest> LODRequests;
-
 	/* Store CPU data for root info & root binding */
 	FHairStrandsRootBulkData& BulkData;
 
@@ -430,9 +428,9 @@ struct FHairStrandsDeformedRootResource : public FHairCommonResource
 	FHairStrandsDeformedRootResource(const FHairStrandsRestRootResource* InRestResources, EHairStrandsResourcesType CurveType, const FHairResourceName& ResourceName, const FName& OwnerName);
 
 	/* Init/Release buffers */
-	virtual void InternalAllocate(FRDGBuilder& GraphBuilder, int32 InLODIndex) override;
+	virtual void InternalAllocate(FRDGBuilder& GraphBuilder, uint32 InCurveCount, int32 InLODIndex) override;
 	virtual void InternalRelease() override;
-	virtual bool InternalIsLODDataLoaded(int32 InLODIndex) const override;
+	virtual bool InternalIsLODDataLoaded(uint32 InRequestedCurveCount, uint32 InRequestedPointCount, int32 InLODIndex) const override;
 
 	/* Get the resource name */
 	virtual FString GetFriendlyName() const override { return TEXT("FHairStrandsDeformedRootResource"); }
@@ -480,6 +478,7 @@ struct FHairStrandsDeformedRootResource : public FHairCommonResource
 		const bool IsValid() const { return Status == EStatus::Initialized || Status == EStatus::Completed; }
 		EStatus Status = EStatus::Invalid;
 		int32 LODIndex = -1;
+		uint32 AvailableCurveCount = 0;
 
 		/* Strand hair roots translation and rotation in triangle-deformed position relative to the bound triangle. Positions are relative the deformed root center*/
 		FRDGExternalBuffer DeformedUniqueTrianglePositionBuffer[2];
diff --git a/Engine/Plugins/Runtime/HairStrands/Source/HairStrandsCore/Public/HairStrandsDatas.h b/Engine/Plugins/Runtime/HairStrands/Source/HairStrandsCore/Public/HairStrandsDatas.h
--- a/Engine/Plugins/Runtime/HairStrands/Source/HairStrandsCore/Public/HairStrandsDatas.h
+++ b/Engine/Plugins/Runtime/HairStrands/Source/HairStrandsCore/Public/HairStrandsDatas.h
@@ -383,6 +383,7 @@ struct FHairStreamingRequest
 	TArray<FChunk> Chunks;
 	uint32 CurveCount = 0;
 	uint32 PointCount = 0;
+	int32 LODIndex = 0;
 
 	// When enabled, data can be loaded from an offset. Otherwisee, start from the beginning of the resource
 	// This is used when cooking data to force the loading of the entire resource (i.e., bSupportOffsetLoad=false)
@@ -413,11 +414,11 @@ struct HAIRSTRANDSCORE_API FHairStrandsBulkCommon
 	virtual ~FHairStrandsBulkCommon() { }
 	void Serialize(FArchive& Ar, UObject* Owner);
 	virtual void SerializeHeader(FArchive& Ar, UObject* Owner) = 0;
-	void SerializeData(FArchive& Ar, UObject* Owner);
+	void SerializeData(FArchive& Ar, UObject* Owner, int32 LODIndex=-1);
 
-	void Write_DDC(UObject* Owner, TArray<UE::DerivedData::FCachePutValueRequest>& Out);
+	void Write_DDC(UObject* Owner, TArray<UE::DerivedData::FCachePutValueRequest>& Out, int32 LODIndex=-1);
 	void Read_DDC(FHairStreamingRequest* In, TArray<UE::DerivedData::FCacheGetChunkRequest>& Out);
-	void Write_IO(UObject* Owner, FArchive& Out);
+	void Write_IO(UObject* Owner, FArchive& Out, int32 LODIndex=-1);
 	void Read_IO(FHairStreamingRequest* In, FBulkDataBatchRequest& Out);
 
 	struct FQuery
@@ -425,7 +426,7 @@ struct HAIRSTRANDSCORE_API FHairStrandsBulkCommon
 		void Add(FHairBulkContainer& In, const TCHAR* InSuffix, uint32 InOffset=0, uint32 InSize=0);
 		uint32 GetCurveCount() const { check(StreamingRequest); return StreamingRequest->CurveCount; }
 		uint32 GetPointCount() const { check(StreamingRequest); return StreamingRequest->PointCount; }
-
+		int32  GetLODIndex() const	 { return StreamingRequest ? StreamingRequest->LODIndex : LODIndex; }
 		enum EQueryType { None, ReadDDC, WriteDDC, ReadIO, ReadWriteIO /* i.e. regular Serialize() */};
 		EQueryType Type = None;
 		FHairStreamingRequest* StreamingRequest = nullptr;
@@ -437,6 +438,7 @@ struct HAIRSTRANDSCORE_API FHairStrandsBulkCommon
 		FString* DerivedDataKey = nullptr;
 	#endif
 		UObject* Owner = nullptr; 
+		int32 LODIndex = -1;
 	};
 
 	virtual uint32 GetResourceCount() const = 0;
@@ -818,11 +820,11 @@ struct FHairStrandsRootData
 };
 
 /* Bulk data for root resources (GPU resources are stored into FHairStrandsRootResources) */
-struct FHairStrandsRootBulkData
+struct FHairStrandsRootBulkData : FHairStrandsBulkCommon
 {
-	void SerializeHeader(FArchive& Ar, UObject* Owner);
-	void SerializeData(FArchive& Ar, UObject* Owner);
-	void Serialize(FArchive& Ar, UObject* Owner);
+	virtual void SerializeHeader(FArchive& Ar, UObject* Owner) override;
+	virtual uint32 GetResourceCount() const override;
+	virtual void GetResources(FQuery& Out) override;
 
 	void Reset();
 	bool IsValid() const { return Header.RootCount > 0; }
@@ -879,15 +881,15 @@ struct FHairStrandsRootBulkData
 		struct FLOD
 		{
 			// Binding
-			FByteBulkData RootToUniqueTriangleIndexBuffer; 	// Map each root onto the unique triangle Id (per-root)
-			FByteBulkData RootBarycentricBuffer; 			// Root's barycentric (per-root)
-			FByteBulkData UniqueTriangleIndexBuffer; 		// Unique triangles list from skeleton mesh section IDs and triangle IDs (per-unique-triangle)
-			FByteBulkData RestUniqueTrianglePositionBuffer;	// Rest triangle positions (per-unique-triangle)
+			FHairBulkContainer RootToUniqueTriangleIndexBuffer; // Map each root onto the unique triangle Id (per-root)
+			FHairBulkContainer RootBarycentricBuffer; 			// Root's barycentric (per-root)
+			FHairBulkContainer UniqueTriangleIndexBuffer; 		// Unique triangles list from skeleton mesh section IDs and triangle IDs (per-unique-triangle)
+			FHairBulkContainer RestUniqueTrianglePositionBuffer;// Rest triangle positions (per-unique-triangle)
 
 			// RBF
-			FByteBulkData MeshInterpolationWeightsBuffer; 	// Store the hair interpolation weights | Size = SamplesCount * SamplesCount (per-sample
-			FByteBulkData MeshSampleIndicesBuffer; 			// Store the samples vertex indices (per-sample)
-			FByteBulkData RestSamplePositionsBuffer; 		// Store the samples rest positions (per-sample)
+			FHairBulkContainer MeshInterpolationWeightsBuffer; 	// Store the hair interpolation weights | Size = SamplesCount * SamplesCount (per-sample
+			FHairBulkContainer MeshSampleIndicesBuffer; 		// Store the samples vertex indices (per-sample)
+			FHairBulkContainer RestSamplePositionsBuffer; 		// Store the samples rest positions (per-sample)
 		};
 		TArray<FLOD> LODs;
 	} Data;
