* Fix groom debug attribute buffer allocation & copy
* Fix groom debug attribute writing

diff --git a/Engine/Plugins/Runtime/HairStrands/Source/HairStrandsCore/Private/GroomResources.cpp b/Engine/Plugins/Runtime/HairStrands/Source/HairStrandsCore/Private/GroomResources.cpp
--- a/Engine/Plugins/Runtime/HairStrands/Source/HairStrandsCore/Private/GroomResources.cpp
+++ b/Engine/Plugins/Runtime/HairStrands/Source/HairStrandsCore/Private/GroomResources.cpp
@@ -1097,7 +1097,7 @@ FRDGExternalBuffer& FHairStrandsDeformedResource::GetDeformerCurveAttributeBuffe
 	// Deformer curve attributes
 	if (DeformerCurveAttributeBuffer.Buffer == nullptr)
 	{
-		const uint32 AllocSize = FMath::DivideAndRoundUp(BulkData.Header.CurveCount, BulkData.Header.Strides.CurveAttributeChunkElementCount) * BulkData.Header.Strides.CurveAttributeChunkStride;
+		const uint32 AllocSize = BulkData.GetCurveAttributeSizeInBytes();
 		check(AllocSize > 0);
 		InternalCreateByteAddressBufferRDG(GraphBuilder, AllocSize, DeformerCurveAttributeBuffer, ToHairResourceDebugName(HAIRSTRANDS_RESOUCE_NAME(CurveType, Hair.StrandsDeformed_DeformerCurveAttributeBuffer), ResourceName), OwnerName, EHairResourceUsageType::Dynamic);
 	}
@@ -1109,7 +1109,7 @@ FRDGExternalBuffer& FHairStrandsDeformedResource::GetDeformerPointAttributeBuffe
 	// Deformer point attributes
 	if (DeformerPointAttributeBuffer.Buffer == nullptr && (BulkData.Header.Flags & FHairStrandsBulkData::DataFlags_HasPointAttribute) && BulkData.Data.PointAttributes.Data.GetBulkDataSize() > 0)
 	{
-		const uint32 AllocSize = FMath::DivideAndRoundUp(BulkData.Header.PointCount, BulkData.Header.Strides.PointAttributeChunkElementCount) * BulkData.Header.Strides.PointAttributeChunkStride;
+		const uint32 AllocSize = BulkData.GetPointAttributeSizeInBytes();
 		check(AllocSize > 0);
 		InternalCreateByteAddressBufferRDG(GraphBuilder, AllocSize, DeformerPointAttributeBuffer, ToHairResourceDebugName(HAIRSTRANDS_RESOUCE_NAME(CurveType, Hair.StrandsDeformedt_DeformerPointAttributeBuffer), ResourceName), OwnerName, EHairResourceUsageType::Dynamic);
 	}
diff --git a/Engine/Plugins/Runtime/HairStrands/Source/HairStrandsCore/Private/HairStrandsDatas.cpp b/Engine/Plugins/Runtime/HairStrands/Source/HairStrandsCore/Private/HairStrandsDatas.cpp
--- a/Engine/Plugins/Runtime/HairStrands/Source/HairStrandsCore/Private/HairStrandsDatas.cpp
+++ b/Engine/Plugins/Runtime/HairStrands/Source/HairStrandsCore/Private/HairStrandsDatas.cpp
@@ -484,8 +484,8 @@ void FHairStrandsBulkData::GetResources(FHairStrandsBulkCommon::FQuery& Out)
 		PointCount = CurveCount > 0 ? Header.CurveToPointCount[CurveCount -1] : 0;
 	}
 
-	const uint32 PointAttributeSize = PointCount > 0 ? FMath::DivideAndRoundUp(PointCount, Header.Strides.PointAttributeChunkElementCount) * Header.Strides.PointAttributeChunkStride : 0;
-	const uint32 CurveAttributeSize = CurveCount > 0 ? FMath::DivideAndRoundUp(CurveCount, Header.Strides.CurveAttributeChunkElementCount) * Header.Strides.CurveAttributeChunkStride : 0;
+	const uint32 PointAttributeSize = GetPointAttributeSizeInBytes(PointCount);
+	const uint32 CurveAttributeSize = GetCurveAttributeSizeInBytes(CurveCount);
 
 	if (!!(Header.Flags & DataFlags_HasData))
 	{
diff --git a/Engine/Plugins/Runtime/HairStrands/Source/HairStrandsCore/Private/HairStrandsInterpolation.cpp b/Engine/Plugins/Runtime/HairStrands/Source/HairStrandsCore/Private/HairStrandsInterpolation.cpp
--- a/Engine/Plugins/Runtime/HairStrands/Source/HairStrandsCore/Private/HairStrandsInterpolation.cpp
+++ b/Engine/Plugins/Runtime/HairStrands/Source/HairStrandsCore/Private/HairStrandsInterpolation.cpp
@@ -785,7 +785,7 @@ static void AddPatchAttributePass(
 	const FIntVector DispatchCount = FIntVector(FMath::DivideAndRoundUp(CurveCount, 1024u), 1, 1);
 
 	// First copy all the rendering attributes into the new bufffer, before overriding some of them
-	AddCopyBufferPass(GraphBuilder, RenAttributeBuffer, OutRenAttributeBuffer.Buffer);
+	AddCopyBufferPass(GraphBuilder, OutRenAttributeBuffer.Buffer, RenAttributeBuffer);
 
 	FHairPatchAttributeCS::FParameters* Parameters = GraphBuilder.AllocParameters<FHairPatchAttributeCS::FParameters>();
 	Parameters->CurveCount = CurveCount;
@@ -2097,7 +2097,7 @@ void ComputeHairStrandsInterpolation(
 				// Special case for debug mode were the attribute buffer is patch with some custom data to show hair properties (strands belonging to the same cluster, ...)
 				if (Instance->Strands.DebugCurveAttributeBuffer.Buffer == nullptr)
 				{
-					CreateHairStrandsDebugAttributeBuffer(GraphBuilder, &Instance->Strands.DebugCurveAttributeBuffer, Instance->Strands.Data->Data.CurveAttributes.GetBulkDataSize(), FName(Instance->Debug.MeshComponentName));
+					CreateHairStrandsDebugAttributeBuffer(GraphBuilder, &Instance->Strands.DebugCurveAttributeBuffer, Instance->Strands.Data->GetCurveAttributeSizeInBytes(), FName(Instance->Debug.MeshComponentName));
 				}
 				FRDGImportedBuffer OutRenCurveAttributeBuffer = Register(GraphBuilder, Instance->Strands.DebugCurveAttributeBuffer, ERDGImportedBufferFlags::CreateUAV);
 
diff --git a/Engine/Plugins/Runtime/HairStrands/Source/HairStrandsCore/Public/HairStrandsDatas.h b/Engine/Plugins/Runtime/HairStrands/Source/HairStrandsCore/Public/HairStrandsDatas.h
--- a/Engine/Plugins/Runtime/HairStrands/Source/HairStrandsCore/Public/HairStrandsDatas.h
+++ b/Engine/Plugins/Runtime/HairStrands/Source/HairStrandsCore/Public/HairStrandsDatas.h
@@ -592,6 +592,9 @@ struct HAIRSTRANDSCORE_API FHairStrandsBulkData : FHairStrandsBulkCommon
 	FVector GetPositionOffset() const { return Header.BoundingBox.GetCenter(); }
 	const FBox& GetBounds() const { return Header.BoundingBox; }
 
+	uint32 GetCurveAttributeSizeInBytes(uint32 InCurveCount=HAIR_MAX_NUM_CURVE_PER_GROUP) const	{ return InCurveCount > 0 ? FMath::DivideAndRoundUp(FMath::Min(Header.CurveCount, InCurveCount), Header.Strides.CurveAttributeChunkElementCount) * Header.Strides.CurveAttributeChunkStride : 0; }
+	uint32 GetPointAttributeSizeInBytes(uint32 InPointCount=HAIR_MAX_NUM_POINT_PER_GROUP) const	{ return InPointCount > 0 ? FMath::DivideAndRoundUp(FMath::Min(Header.PointCount, InPointCount), Header.Strides.PointAttributeChunkElementCount) * Header.Strides.PointAttributeChunkStride : 0; }
+
 	struct FHeader
 	{
 		uint32 CurveCount = 0;
diff --git a/Engine/Shaders/Private/HairStrands/HairStrandsInterpolation.usf b/Engine/Shaders/Private/HairStrands/HairStrandsInterpolation.usf
--- a/Engine/Shaders/Private/HairStrands/HairStrandsInterpolation.usf
+++ b/Engine/Shaders/Private/HairStrands/HairStrandsInterpolation.usf
@@ -581,7 +581,7 @@ void MainCS(uint2 DispatchThreadId : SV_DispatchThreadID, uint GroupThread1D : S
 		uint WriteAttributOffsetInBytes = HAIR_CURVE_ATTRIBUTE_OFFSET_SEED(CurveAttributeOffsets);
 
 		const uint ChunkIndex   	= CurveIndex >> CurveAttributeIndexToChunkDivAsShift;
-		WriteIndex                 -= ChunkIndex * CurveAttributeIndexToChunkDivAsShift;
+		WriteIndex                 -= ChunkIndex * CurveAttributeChunkElementCount;
 		WriteAttributOffsetInBytes += ChunkIndex * CurveAttributeChunkStrideInBytes;
 
 		OutRenAttributeBuffer.Store(WriteAttributOffsetInBytes + WriteIndex, Out);
