Reduce groom cluster resources size (~ -15%)

diff --git a/Engine/Plugins/Runtime/HairStrands/Source/HairStrandsCore/Private/GroomAsset.cpp b/Engine/Plugins/Runtime/HairStrands/Source/HairStrandsCore/Private/GroomAsset.cpp
--- a/Engine/Plugins/Runtime/HairStrands/Source/HairStrandsCore/Private/GroomAsset.cpp
+++ b/Engine/Plugins/Runtime/HairStrands/Source/HairStrandsCore/Private/GroomAsset.cpp
@@ -128,7 +128,7 @@ uint32 FHairGroupPlatformData::FStrands::GetDataSize() const
 	CTotal += ClusterCullingBulkData.Header.LODVisibility.GetAllocatedSize();
 	CTotal += ClusterCullingBulkData.Header.CPULODScreenSize.GetAllocatedSize();
 	CTotal += ClusterCullingBulkData.Data.ClusterLODInfos.IsBulkDataLoaded()	? ClusterCullingBulkData.Data.ClusterLODInfos.GetBulkDataSize() : 0;
-	BTotal += ClusterCullingBulkData.Data.VertexToClusterIds.IsBulkDataLoaded()	? ClusterCullingBulkData.Data.VertexToClusterIds.GetBulkDataSize() : 0;
+	BTotal += ClusterCullingBulkData.Data.CurveToClusterIds.IsBulkDataLoaded()	? ClusterCullingBulkData.Data.CurveToClusterIds.GetBulkDataSize() : 0;
 	BTotal += ClusterCullingBulkData.Data.ClusterVertexIds.IsBulkDataLoaded()	? ClusterCullingBulkData.Data.ClusterVertexIds.GetBulkDataSize() : 0;
 	CTotal += ClusterCullingBulkData.Data.PackedClusterInfos.IsBulkDataLoaded()	? ClusterCullingBulkData.Data.PackedClusterInfos.GetBulkDataSize() : 0;
 
diff --git a/Engine/Plugins/Runtime/HairStrands/Source/HairStrandsCore/Private/GroomBuilder.cpp b/Engine/Plugins/Runtime/HairStrands/Source/HairStrandsCore/Private/GroomBuilder.cpp
--- a/Engine/Plugins/Runtime/HairStrands/Source/HairStrandsCore/Private/GroomBuilder.cpp
+++ b/Engine/Plugins/Runtime/HairStrands/Source/HairStrandsCore/Private/GroomBuilder.cpp
@@ -40,7 +40,7 @@ static FAutoConsoleVariableRef CVarHairGroupIndexBuilder_MaxVoxelResolution(TEXT
 
 FString FGroomBuilder::GetVersion()
 {
-	return TEXT("v8r24");
+	return TEXT("v8r25");
 }
 
 namespace FHairStrandsDecimation
@@ -2694,6 +2694,7 @@ struct FClusterGrid
 		}
 		uint32 Offset = 0;
 		uint32 Count = 0;
+		uint32 CurveIndex = 0;
 		float Area = 0;
 		float AvgRadius = 0;
 		float MaxRadius = 0;
@@ -2891,8 +2892,9 @@ static void BuildClusterData(
 	check(LODCount > 0);
 
 	const uint32 RenCurveCount = InRenStrandsData.GetNumCurves();
-	Out.VertexCount = InRenStrandsData.GetNumPoints();
-	check(Out.VertexCount);
+	Out.PointCount = InRenStrandsData.GetNumPoints();
+	Out.CurveCount = InRenStrandsData.GetNumCurves();
+	check(Out.PointCount > 0 && Out.CurveCount > 0);
 
 	// 1. Allocate cluster per voxel containing contains >=1 render curve root
 	const FVector GroupMinBound = InRenStrandsData.BoundingBox.Min;
@@ -2925,6 +2927,7 @@ static void BuildClusterData(
 	for (uint32 RenCurveIndex = 0; RenCurveIndex < RenCurveCount; ++RenCurveIndex)
 	{
 		FClusterGrid::FCurve RCurve;
+		RCurve.CurveIndex = RenCurveIndex;
 		RCurve.Count = InRenStrandsData.StrandsCurves.CurvesCount[RenCurveIndex];
 		RCurve.Offset = InRenStrandsData.StrandsCurves.CurvesOffset[RenCurveIndex];
 		RCurve.Area = 0.0f;
@@ -2971,7 +2974,7 @@ static void BuildClusterData(
 	}
 	Out.ClusterCount = ValidClusterIndices.Num();
 	Out.ClusterInfos.Init(FHairClusterInfo(), Out.ClusterCount);
-	Out.VertexToClusterIds.SetNum(Out.VertexCount);
+	Out.CurveToClusterIds.SetNum(Out.CurveCount);
 	Out.LODInfos.SetNum(LODCount);
 
 	// Conservative allocation for inserting vertex indices for the various curves LOD
@@ -2986,7 +2989,7 @@ static void BuildClusterData(
 	// Local variable for being capture by the lambda
 	TArray<FHairClusterInfo>& LocalClusterInfos = Out.ClusterInfos;
 	TArray<FHairClusterLODInfo>& LocalClusterLODInfos = Out.ClusterLODInfos;
-	TArray<uint32>& LocalVertexToClusterIds = Out.VertexToClusterIds;
+	TArray<uint32>& LocalCurveToClusterIds = Out.CurveToClusterIds;
 #define USE_PARALLE_FOR 0
 #if USE_PARALLE_FOR
 	ParallelFor(Out.ClusterCount,
@@ -2999,7 +3002,7 @@ static void BuildClusterData(
 			&ClusterGrid,
 			&LocalClusterInfos,
 			&LocalClusterLODInfos,
-			&LocalVertexToClusterIds,
+			&LocalCurveToClusterIds,
 			&VertexLODMasks,
 			&RawClusterVertexIds,
 			&RawClusterVertexCount
@@ -3034,10 +3037,12 @@ static void BuildClusterData(
 		Cluster.Area = 0;
 		for (FClusterGrid::FCurve& ClusterCurve : Cluster.ClusterCurves)
 		{
+			check(ClusterCurve.CurveIndex < uint32(LocalCurveToClusterIds.Num()));
+			LocalCurveToClusterIds[ClusterCurve.CurveIndex] = ClusterIt;
+
 			for (uint32 RenPointIndex = 0; RenPointIndex < ClusterCurve.Count; ++RenPointIndex)
 			{
 				const uint32 PointGlobalIndex = RenPointIndex + ClusterCurve.Offset;
-				LocalVertexToClusterIds[PointGlobalIndex] = ClusterIt;
 
 				const FVector& P = (FVector)InRenStrandsData.StrandsPoints.PointsPosition[PointGlobalIndex];
 				{
@@ -3273,7 +3278,8 @@ static void BuildClusterBulkData(
 	Out.Reset();
 
 	Out.Header.ClusterCount		= In.ClusterCount;
-	Out.Header.VertexCount		= In.VertexCount;
+	Out.Header.PointCount		= In.PointCount;
+	Out.Header.CurveCount		= In.CurveCount;
 	Out.Header.VertexLODCount	= In.ClusterVertexIds.Num();
 	Out.Header.ClusterLODCount 	= In.ClusterLODInfos.Num();
 
@@ -3284,17 +3290,17 @@ static void BuildClusterBulkData(
 	// Sanity check
 	check(Out.Header.ClusterCount		== uint32(In.ClusterInfos.Num()));
 	check(Out.Header.ClusterLODCount	== uint32(In.ClusterLODInfos.Num()));
-	check(Out.Header.VertexCount		== uint32(In.VertexToClusterIds.Num()));
+	check(Out.Header.CurveCount			== uint32(In.CurveToClusterIds.Num()));
 	check(Out.Header.VertexLODCount		== uint32(In.ClusterVertexIds.Num()));
 	
 	HairStrandsBuilder::CopyToBulkData<FHairClusterLODInfoFormat>(Out.Data.ClusterLODInfos, In.ClusterLODInfos);
-	HairStrandsBuilder::CopyToBulkData<FHairClusterIndexFormat>(Out.Data.VertexToClusterIds, In.VertexToClusterIds);
+	HairStrandsBuilder::CopyToBulkData<FHairClusterIndexFormat>(Out.Data.CurveToClusterIds, In.CurveToClusterIds);
 	HairStrandsBuilder::CopyToBulkData<FHairClusterIndexFormat>(Out.Data.ClusterVertexIds, In.ClusterVertexIds);
 
 	// Pack LODInfo into GPU format
 	{
 		check(uint32(In.ClusterInfos.Num()) == Out.Header.ClusterCount);
-		check(uint32(In.VertexToClusterIds.Num()) == Out.Header.VertexCount);
+		check(uint32(In.CurveToClusterIds.Num()) == Out.Header.CurveCount);
 	
 		TArray<FHairClusterInfo::Packed> PackedClusterInfos;
 		PackedClusterInfos.Reserve(In.ClusterInfos.Num());
diff --git a/Engine/Plugins/Runtime/HairStrands/Source/HairStrandsCore/Private/GroomResources.cpp b/Engine/Plugins/Runtime/HairStrands/Source/HairStrandsCore/Private/GroomResources.cpp
--- a/Engine/Plugins/Runtime/HairStrands/Source/HairStrandsCore/Private/GroomResources.cpp
+++ b/Engine/Plugins/Runtime/HairStrands/Source/HairStrandsCore/Private/GroomResources.cpp
@@ -1159,7 +1159,7 @@ void FHairStrandsClusterCullingResource::InternalAllocate(FRDGBuilder& GraphBuil
 	InternalCreateStructuredBufferRDG_FromBulkData<FHairClusterInfoFormat>(GraphBuilder, BulkData.Data.PackedClusterInfos.Data, BulkData.Header.ClusterCount, ClusterInfoBuffer, ToHairResourceDebugName(TEXT("Hair.StrandsClusterCulling_ClusterInfoBuffer"), ResourceName), OwnerName, EHairResourceUsageType::Static);
 	InternalCreateStructuredBufferRDG_FromBulkData<FHairClusterLODInfoFormat>(GraphBuilder, BulkData.Data.ClusterLODInfos.Data, BulkData.Header.ClusterLODCount, ClusterLODInfoBuffer, ToHairResourceDebugName(TEXT("Hair.StrandsClusterCulling_ClusterLODInfoBuffer"), ResourceName), OwnerName, EHairResourceUsageType::Static);
 
-	InternalCreateVertexBufferRDG_FromBulkData<FHairClusterIndexFormat>(GraphBuilder, BulkData.Data.VertexToClusterIds.Data, BulkData.Header.VertexCount, VertexToClusterIdBuffer, ToHairResourceDebugName(TEXT("Hair.StrandsClusterCulling_VertexToClusterIds"), ResourceName), OwnerName, EHairResourceUsageType::Static);
+	InternalCreateVertexBufferRDG_FromBulkData<FHairClusterIndexFormat>(GraphBuilder, BulkData.Data.CurveToClusterIds.Data, BulkData.Header.CurveCount, CurveToClusterIdBuffer, ToHairResourceDebugName(TEXT("Hair.StrandsClusterCulling_CurveToClusterIds"), ResourceName), OwnerName, EHairResourceUsageType::Static);
 	InternalCreateVertexBufferRDG_FromBulkData<FHairClusterIndexFormat>(GraphBuilder, BulkData.Data.ClusterVertexIds.Data, BulkData.Header.VertexLODCount, ClusterVertexIdBuffer, ToHairResourceDebugName(TEXT("Hair.StrandsClusterCulling_ClusterVertexIds"), ResourceName), OwnerName, EHairResourceUsageType::Static);
 }
 
@@ -1168,7 +1168,7 @@ void FHairStrandsClusterCullingResource::InternalRelease()
 	ClusterInfoBuffer.Release();
 	ClusterLODInfoBuffer.Release();
 	ClusterVertexIdBuffer.Release();
-	VertexToClusterIdBuffer.Release();
+	CurveToClusterIdBuffer.Release();
 	MaxAvailableCurveCount = 0;
 }
 
diff --git a/Engine/Plugins/Runtime/HairStrands/Source/HairStrandsCore/Private/HairStrandsDatas.cpp b/Engine/Plugins/Runtime/HairStrands/Source/HairStrandsCore/Private/HairStrandsDatas.cpp
--- a/Engine/Plugins/Runtime/HairStrands/Source/HairStrandsCore/Private/HairStrandsDatas.cpp
+++ b/Engine/Plugins/Runtime/HairStrands/Source/HairStrandsCore/Private/HairStrandsDatas.cpp
@@ -590,7 +590,8 @@ void FHairStrandsClusterCullingBulkData::Reset()
 {
 	Header.ClusterCount = 0;
 	Header.ClusterLODCount = 0;
-	Header.VertexCount = 0;
+	Header.PointCount = 0;
+	Header.CurveCount = 0;
 	Header.VertexLODCount = 0;
 
 	Header.LODVisibility.Empty();
@@ -598,13 +599,13 @@ void FHairStrandsClusterCullingBulkData::Reset()
 	Header.LODInfos.Empty();
 
 	Data.ClusterLODInfos.RemoveBulkData();
-	Data.VertexToClusterIds.RemoveBulkData();
+	Data.CurveToClusterIds.RemoveBulkData();
 	Data.ClusterVertexIds.RemoveBulkData();
 	Data.PackedClusterInfos.RemoveBulkData();
 
 	// Reset the bulk byte buffer to ensure the (serialize) data size is reset to 0
 	Data.ClusterLODInfos 	= FHairBulkContainer();
-	Data.VertexToClusterIds = FHairBulkContainer();
+	Data.CurveToClusterIds = FHairBulkContainer();
 	Data.ClusterVertexIds 	= FHairBulkContainer();
 	Data.PackedClusterInfos = FHairBulkContainer();
 }
@@ -613,7 +614,8 @@ void FHairStrandsClusterCullingBulkData::SerializeHeader(FArchive& Ar, UObject*
 {
 	Ar << Header.ClusterCount;
 	Ar << Header.ClusterLODCount;
-	Ar << Header.VertexCount;
+	Ar << Header.PointCount;
+	Ar << Header.CurveCount;
 	Ar << Header.VertexLODCount;
 	Ar << Header.LODVisibility;
 	Ar << Header.CPULODScreenSize;
@@ -644,9 +646,9 @@ void FHairStrandsClusterCullingBulkData::GetResources(FHairStrandsBulkCommon::FQ
 		Out.Add(Data.ClusterLODInfos, TEXT("_ClusterLODInfos"), 0, 0); // Load all data
 	}
 
-	if (Header.VertexCount)
+	if (Header.CurveCount)
 	{
-		Out.Add(Data.VertexToClusterIds, TEXT("_VertexToClusterIds"), 0, 0); // Load all data
+		Out.Add(Data.CurveToClusterIds, TEXT("_CurveToClusterIds"), 0, 0); // Load all data
 	}
 
 	if (Header.VertexLODCount)
diff --git a/Engine/Plugins/Runtime/HairStrands/Source/HairStrandsCore/Private/HairStrandsInterpolation.cpp b/Engine/Plugins/Runtime/HairStrands/Source/HairStrandsCore/Private/HairStrandsInterpolation.cpp
--- a/Engine/Plugins/Runtime/HairStrands/Source/HairStrandsCore/Private/HairStrandsInterpolation.cpp
+++ b/Engine/Plugins/Runtime/HairStrands/Source/HairStrandsCore/Private/HairStrandsInterpolation.cpp
@@ -742,7 +742,7 @@ class FHairPatchAttributeCS : public FGlobalShader
 		SHADER_PARAMETER_STRUCT_INCLUDE(FHairStrandsInstanceAttributeParameters, RenAttributes)
 		SHADER_PARAMETER_RDG_BUFFER_SRV(Buffer, RenCurveBuffer)
 		SHADER_PARAMETER_RDG_BUFFER_SRV(ByteAddressBuffer, InterpolationBuffer)
-		SHADER_PARAMETER_RDG_BUFFER_SRV(Buffer, RenVertexToClusterIdBuffer)
+		SHADER_PARAMETER_RDG_BUFFER_SRV(Buffer, RenCurveToClusterIdBuffer)
 		SHADER_PARAMETER_RDG_BUFFER_UAV(RWByteAddressBuffer, OutRenAttributeBuffer)
 	END_SHADER_PARAMETER_STRUCT()
 
@@ -774,7 +774,7 @@ static void AddPatchAttributePass(
 	const FHairStrandsBulkData& RenBulkData,
 	const FRDGBufferRef& RenAttributeBuffer,
 	const FRDGBufferSRVRef& RenCurveBuffer,
-	const FRDGBufferSRVRef& RenVertexToClusterIdBuffer,
+	const FRDGBufferSRVRef& RenCurveToClusterIdBuffer,
 	const FRDGBufferSRVRef& InterpolationBuffer,
 	FRDGImportedBuffer& OutRenAttributeBuffer)
 {
@@ -790,7 +790,7 @@ static void AddPatchAttributePass(
 	FHairPatchAttributeCS::FParameters* Parameters = GraphBuilder.AllocParameters<FHairPatchAttributeCS::FParameters>();
 	Parameters->CurveCount = CurveCount;
 	Parameters->RenCurveBuffer = RenCurveBuffer;
-	Parameters->RenVertexToClusterIdBuffer = RenVertexToClusterIdBuffer;
+	Parameters->RenCurveToClusterIdBuffer = RenCurveToClusterIdBuffer;
 	GetHairStrandsAttributeParameter(RenBulkData, Parameters->RenAttributes);
 
 	if (bSimulation)
@@ -1545,7 +1545,7 @@ void RegisterClusterData(FHairGroupInstance* Instance, FHairStrandClusterData* I
 
 	HairGroupCluster.ClusterInfoBuffer = &Instance->Strands.ClusterCullingResource->ClusterInfoBuffer;
 	HairGroupCluster.ClusterLODInfoBuffer = &Instance->Strands.ClusterCullingResource->ClusterLODInfoBuffer;
-	HairGroupCluster.VertexToClusterIdBuffer = &Instance->Strands.ClusterCullingResource->VertexToClusterIdBuffer;
+	HairGroupCluster.CurveToClusterIdBuffer = &Instance->Strands.ClusterCullingResource->CurveToClusterIdBuffer;
 	HairGroupCluster.ClusterVertexIdBuffer = &Instance->Strands.ClusterCullingResource->ClusterVertexIdBuffer;
 
 	HairGroupCluster.HairGroupPublicPtr = Instance->HairGroupPublicData;
@@ -2112,7 +2112,7 @@ void ComputeHairStrandsInterpolation(
 					*Instance->Strands.Data,
 					Register(GraphBuilder, Instance->Strands.RestResource->CurveAttributeBuffer, ERDGImportedBufferFlags::CreateSRV).Buffer,
 					RegisterAsSRV(GraphBuilder, Instance->Strands.RestResource->CurveBuffer),
-					RegisterAsSRV(GraphBuilder, Instance->Strands.ClusterCullingResource->VertexToClusterIdBuffer),
+					RegisterAsSRV(GraphBuilder, Instance->Strands.ClusterCullingResource->CurveToClusterIdBuffer),
 					bValidGuide ? RegisterAsSRV(GraphBuilder, Instance->Strands.InterpolationResource->InterpolationBuffer) : nullptr,
 					OutRenCurveAttributeBuffer);
 			}
diff --git a/Engine/Plugins/Runtime/HairStrands/Source/HairStrandsCore/Public/GroomAsset.h b/Engine/Plugins/Runtime/HairStrands/Source/HairStrandsCore/Public/GroomAsset.h
--- a/Engine/Plugins/Runtime/HairStrands/Source/HairStrandsCore/Public/GroomAsset.h
+++ b/Engine/Plugins/Runtime/HairStrands/Source/HairStrandsCore/Public/GroomAsset.h
@@ -188,7 +188,7 @@ struct HAIRSTRANDSCORE_API FHairGroupPlatformData
 			Total += ClusterCullingBulkData.Header.CPULODScreenSize.GetAllocatedSize();
 			Total += ClusterCullingBulkData.Data.ClusterInfos.IsBulkDataLoaded()		? ClusterCullingBulkData.Data.ClusterInfos.GetBulkDataSize() : 0;
 			Total += ClusterCullingBulkData.Data.ClusterLODInfos.IsBulkDataLoaded()		? ClusterCullingBulkData.Data.ClusterLODInfos.GetBulkDataSize() : 0;
-			Total += ClusterCullingBulkData.Data.VertexToClusterIds.IsBulkDataLoaded()	? ClusterCullingBulkData.Data.VertexToClusterIds.GetBulkDataSize() : 0;
+			Total += ClusterCullingBulkData.Data.CurveToClusterIds.IsBulkDataLoaded()	? ClusterCullingBulkData.Data.CurveToClusterIds.GetBulkDataSize() : 0;
 			Total += ClusterCullingBulkData.Data.ClusterVertexIds.IsBulkDataLoaded()	? ClusterCullingBulkData.Data.ClusterVertexIds.GetBulkDataSize() : 0;
 			return Total;
 		}
diff --git a/Engine/Plugins/Runtime/HairStrands/Source/HairStrandsCore/Public/GroomResources.h b/Engine/Plugins/Runtime/HairStrands/Source/HairStrandsCore/Public/GroomResources.h
--- a/Engine/Plugins/Runtime/HairStrands/Source/HairStrandsCore/Public/GroomResources.h
+++ b/Engine/Plugins/Runtime/HairStrands/Source/HairStrandsCore/Public/GroomResources.h
@@ -326,7 +326,7 @@ struct FHairStrandsClusterCullingResource : public FHairCommonResource
 		uint32 Total = 0;
 		Total += GetBufferTotalNumBytes(ClusterInfoBuffer);
 		Total += GetBufferTotalNumBytes(ClusterLODInfoBuffer);
-		Total += GetBufferTotalNumBytes(VertexToClusterIdBuffer);
+		Total += GetBufferTotalNumBytes(CurveToClusterIdBuffer);
 		Total += GetBufferTotalNumBytes(ClusterVertexIdBuffer);
 		return Total;
 	}
@@ -335,8 +335,8 @@ struct FHairStrandsClusterCullingResource : public FHairCommonResource
 	FRDGExternalBuffer ClusterInfoBuffer;
 	FRDGExternalBuffer ClusterLODInfoBuffer;
 
-	/* VertexId => ClusterId to know which AABB to contribute to*/
-	FRDGExternalBuffer VertexToClusterIdBuffer;
+	/* CurveId => ClusterId to know which AABB to contribute to*/
+	FRDGExternalBuffer CurveToClusterIdBuffer;
 
 	/* Concatenated data for each cluster: list of VertexId pointed to by ClusterInfoBuffer */
 	FRDGExternalBuffer ClusterVertexIdBuffer;
diff --git a/Engine/Plugins/Runtime/HairStrands/Source/HairStrandsCore/Public/HairStrandsDatas.h b/Engine/Plugins/Runtime/HairStrands/Source/HairStrandsCore/Public/HairStrandsDatas.h
--- a/Engine/Plugins/Runtime/HairStrands/Source/HairStrandsCore/Public/HairStrandsDatas.h
+++ b/Engine/Plugins/Runtime/HairStrands/Source/HairStrandsCore/Public/HairStrandsDatas.h
@@ -712,7 +712,7 @@ struct HAIRSTRANDSCORE_API FHairStrandsInterpolationBulkData : FHairStrandsBulkC
 struct HAIRSTRANDSCORE_API FHairStrandsClusterCullingData
 {
 	void Reset();
-	bool IsValid() const { return ClusterCount > 0 && VertexCount > 0; }
+	bool IsValid() const { return ClusterCount > 0 && PointCount > 0; }
 
 	/* Set LOD visibility, allowing to remove the simulation/rendering of certain LOD */
 	TArray<bool>				LODVisibility;
@@ -723,15 +723,13 @@ struct HAIRSTRANDSCORE_API FHairStrandsClusterCullingData
 	/* LOD info for the various clusters for LOD management on GPU */
 	TArray<FHairClusterInfo>	ClusterInfos;
 	TArray<FHairClusterLODInfo> ClusterLODInfos;
-	TArray<uint32>				VertexToClusterIds;
+	TArray<uint32>				CurveToClusterIds;
 	TArray<uint32>				ClusterVertexIds;
 	TArray<FHairLODInfo>		LODInfos;
 
-	/* Number of cluster  */
-	uint32 ClusterCount = 0;
-
-	/* Number of vertex  */
-	uint32 VertexCount = 0;
+	uint32 ClusterCount = 0;		// Number of clusters
+	uint32 PointCount = 0;			// Number of points
+	uint32 CurveCount = 0;			// Number of curves
 };
 
 struct HAIRSTRANDSCORE_API FHairStrandsClusterCullingBulkData : FHairStrandsBulkCommon
@@ -742,7 +740,7 @@ struct HAIRSTRANDSCORE_API FHairStrandsClusterCullingBulkData : FHairStrandsBulk
 	virtual uint32 GetResourceCount() const override;
 	virtual void GetResources(FQuery& Out) override;
 
-	bool IsValid() const { return Header.ClusterCount > 0 && Header.VertexCount > 0; }
+	bool IsValid() const { return Header.ClusterCount > 0 && Header.PointCount > 0; }
 	void Validate(bool bIsSaving);
 
 	struct FHeader
@@ -758,7 +756,8 @@ struct HAIRSTRANDSCORE_API FHairStrandsClusterCullingBulkData : FHairStrandsBulk
 	
 		uint32 ClusterCount = 0;
 		uint32 ClusterLODCount = 0;
-		uint32 VertexCount = 0;
+		uint32 PointCount = 0;
+		uint32 CurveCount = 0;
 		uint32 VertexLODCount = 0;
 	} Header;
 
@@ -767,7 +766,7 @@ struct HAIRSTRANDSCORE_API FHairStrandsClusterCullingBulkData : FHairStrandsBulk
 		/* LOD info for the various clusters for LOD management on GPU */
 		FHairBulkContainer	PackedClusterInfos;		// Size - ClusterCount
 		FHairBulkContainer	ClusterLODInfos;		// Size - ClusterLODCount
-		FHairBulkContainer	VertexToClusterIds;		// Size - VertexCount
+		FHairBulkContainer	CurveToClusterIds;		// Size - CurveCount
 		FHairBulkContainer	ClusterVertexIds;		// Size - VertexLODCount
 	} Data;
 };
diff --git a/Engine/Shaders/Private/HairStrands/HairStrandsInterpolation.usf b/Engine/Shaders/Private/HairStrands/HairStrandsInterpolation.usf
--- a/Engine/Shaders/Private/HairStrands/HairStrandsInterpolation.usf
+++ b/Engine/Shaders/Private/HairStrands/HairStrandsInterpolation.usf
@@ -513,7 +513,7 @@ void MainCS(uint2 DispatchThreadId : SV_DispatchThreadID)
 uint				CurveCount;
 
 Buffer<uint>		RenCurveBuffer;
-Buffer<int>			RenVertexToClusterIdBuffer;
+Buffer<uint>		RenCurveToClusterIdBuffer;
 RWByteAddressBuffer	OutRenAttributeBuffer;
 
 uint 				CurveAttributeIndexToChunkDivAsShift;
@@ -552,8 +552,7 @@ void MainCS(uint2 DispatchThreadId : SV_DispatchThreadID, uint GroupThread1D : S
 		uint Out = 0;
 		#if PERMUTATION_SIMULATION == 0
 		{
-			const uint VertexIndex0		= UnpackHairCurve(RenCurveBuffer[CurveIndex]).PointIndex;
-			const uint ClusterId		= RenVertexToClusterIdBuffer[VertexIndex0];
+			const uint ClusterId		= RenCurveToClusterIdBuffer[CurveIndex];
 			SeedValues[GroupThread1D]	= ToSeedHash(ClusterId);
 		}
 		#else // if (PERMUTATION_SIMULATION > 0)
diff --git a/Engine/Source/Runtime/Renderer/Public/HairStrandsInterface.h b/Engine/Source/Runtime/Renderer/Public/HairStrandsInterface.h
--- a/Engine/Source/Runtime/Renderer/Public/HairStrandsInterface.h
+++ b/Engine/Source/Runtime/Renderer/Public/HairStrandsInterface.h
@@ -439,7 +439,7 @@ struct FHairStrandClusterData
 		FRDGExternalBuffer* ClusterAABBBuffer = nullptr;
 		FRDGExternalBuffer* ClusterInfoBuffer = nullptr;
 		FRDGExternalBuffer* ClusterLODInfoBuffer = nullptr;
-		FRDGExternalBuffer* VertexToClusterIdBuffer = nullptr;
+		FRDGExternalBuffer* CurveToClusterIdBuffer = nullptr;
 		FRDGExternalBuffer* ClusterVertexIdBuffer = nullptr;
 
 		TRefCountPtr<FRDGPooledBuffer> ClusterIdBuffer;
