Fix crash when adding a new cards/meshes entry in a groom asset.

The CacheDerivedData() wasn't called properly.

diff --git a/Engine/Plugins/Runtime/HairStrands/Source/HairStrandsCore/Private/GroomAsset.cpp b/Engine/Plugins/Runtime/HairStrands/Source/HairStrandsCore/Private/GroomAsset.cpp
--- a/Engine/Plugins/Runtime/HairStrands/Source/HairStrandsCore/Private/GroomAsset.cpp
+++ b/Engine/Plugins/Runtime/HairStrands/Source/HairStrandsCore/Private/GroomAsset.cpp
@@ -1563,6 +1563,16 @@ static bool IsStrandsLODAttributes(const FName PropertyName)
 		|| PropertyName == GET_MEMBER_NAME_CHECKED(FHairLODSettings, GlobalInterpolation);
 }
 
+static bool IsCardsAttributes(const FName PropertyName)
+{
+	return PropertyName == GET_MEMBER_NAME_CHECKED(UGroomAsset, HairGroupsCards);
+}
+
+static bool IsMeshesAttributes(const FName PropertyName)
+{
+	return PropertyName == GET_MEMBER_NAME_CHECKED(UGroomAsset, HairGroupsMeshes);
+}
+
 void UGroomAsset::PostEditChangeProperty(FPropertyChangedEvent& PropertyChangedEvent)
 {
 	Super::PostEditChangeProperty(PropertyChangedEvent);
@@ -1593,7 +1603,11 @@ void UGroomAsset::PostEditChangeProperty(FPropertyChangedEvent& PropertyChangedE
 	}
 
 	// Rebuild the groom cached data if interpolation or LODs have changed
-	const bool bNeedRebuildDerivedData = IsStrandsInterpolationAttributes(PropertyName) || IsStrandsLODAttributes(PropertyName);
+	const bool bNeedRebuildDerivedData = 
+		IsStrandsInterpolationAttributes(PropertyName) || 
+		IsStrandsLODAttributes(PropertyName) || 
+		IsCardsAttributes(PropertyName) || 
+		IsMeshesAttributes(PropertyName);
 	if (bNeedRebuildDerivedData)
 	{
 		CacheDerivedDatas();
