Fix crash when streaming in root data while skel. mesh data are not ready.

diff --git a/Engine/Plugins/Runtime/HairStrands/Source/HairStrandsCore/Private/GroomResources.cpp b/Engine/Plugins/Runtime/HairStrands/Source/HairStrandsCore/Private/GroomResources.cpp
--- a/Engine/Plugins/Runtime/HairStrands/Source/HairStrandsCore/Private/GroomResources.cpp
+++ b/Engine/Plugins/Runtime/HairStrands/Source/HairStrandsCore/Private/GroomResources.cpp
@@ -1213,7 +1213,7 @@ bool FHairStrandsRestRootResource::InternalIsLODDataLoaded(uint32 InRequestedCur
 
 bool FHairStrandsRestRootResource::InternalIsDataLoaded(uint32 InRequestedCurveCount, uint32 InRequestedPointCount, int32 InLODIndex)
 {
-	if (StreamingRequest.IsNone())
+	if (StreamingRequest.IsNone() && InLODIndex >= 0)
 	{
 		StreamingRequest.Request(InRequestedCurveCount, InRequestedPointCount, InLODIndex, BulkData, false, false, OwnerName);
 	}
