Fix incorrect resources creation for hair cluster resources.

diff --git a/Engine/Plugins/Runtime/HairStrands/Source/HairStrandsCore/Private/GroomResources.cpp b/Engine/Plugins/Runtime/HairStrands/Source/HairStrandsCore/Private/GroomResources.cpp
--- a/Engine/Plugins/Runtime/HairStrands/Source/HairStrandsCore/Private/GroomResources.cpp
+++ b/Engine/Plugins/Runtime/HairStrands/Source/HairStrandsCore/Private/GroomResources.cpp
@@ -332,7 +332,7 @@ void InternalCreateStructuredBufferRDG_FromHairBulkData(FRDGBuilder& GraphBuilde
 	// Fallback for non-streamble resources (e.g. guides)
 	if (InChunk.ChunkRequest == nullptr)
 	{
-		InternalCreateStructuredBufferRDG_FromHairBulkData<FormatType>(GraphBuilder, InChunk.Data, InDataCount, Out, DebugName, OwnerName, UsageType);
+		InternalCreateStructuredBufferRDG_FromBulkData<FormatType>(GraphBuilder, InChunk.Data, InDataCount, Out, DebugName, OwnerName, UsageType);
 		return;
 	}
 
@@ -1156,11 +1156,11 @@ void FHairStrandsClusterCullingResource::InternalAllocate(FRDGBuilder& GraphBuil
 		BulkData.Validate(false);
 	}
 
-	InternalCreateStructuredBufferRDG_FromBulkData<FHairClusterInfoFormat>(GraphBuilder, BulkData.Data.PackedClusterInfos.Data, BulkData.Header.ClusterCount, ClusterInfoBuffer, ToHairResourceDebugName(TEXT("Hair.StrandsClusterCulling_ClusterInfoBuffer"), ResourceName), OwnerName, EHairResourceUsageType::Static);
-	InternalCreateStructuredBufferRDG_FromBulkData<FHairClusterLODInfoFormat>(GraphBuilder, BulkData.Data.ClusterLODInfos.Data, BulkData.Header.ClusterLODCount, ClusterLODInfoBuffer, ToHairResourceDebugName(TEXT("Hair.StrandsClusterCulling_ClusterLODInfoBuffer"), ResourceName), OwnerName, EHairResourceUsageType::Static);
+	InternalCreateStructuredBufferRDG_FromHairBulkData<FHairClusterInfoFormat>(GraphBuilder, BulkData.Data.PackedClusterInfos, BulkData.Header.ClusterCount, ClusterInfoBuffer, ToHairResourceDebugName(TEXT("Hair.StrandsClusterCulling_ClusterInfoBuffer"), ResourceName), OwnerName, EHairResourceUsageType::Static);
+	InternalCreateStructuredBufferRDG_FromHairBulkData<FHairClusterLODInfoFormat>(GraphBuilder, BulkData.Data.ClusterLODInfos, BulkData.Header.ClusterLODCount, ClusterLODInfoBuffer, ToHairResourceDebugName(TEXT("Hair.StrandsClusterCulling_ClusterLODInfoBuffer"), ResourceName), OwnerName, EHairResourceUsageType::Static);
 
-	InternalCreateVertexBufferRDG_FromBulkData<FHairClusterIndexFormat>(GraphBuilder, BulkData.Data.CurveToClusterIds.Data, BulkData.Header.CurveCount, CurveToClusterIdBuffer, ToHairResourceDebugName(TEXT("Hair.StrandsClusterCulling_CurveToClusterIds"), ResourceName), OwnerName, EHairResourceUsageType::Static);
-	InternalCreateVertexBufferRDG_FromBulkData<FHairClusterIndexFormat>(GraphBuilder, BulkData.Data.ClusterVertexIds.Data, BulkData.Header.VertexLODCount, ClusterVertexIdBuffer, ToHairResourceDebugName(TEXT("Hair.StrandsClusterCulling_ClusterVertexIds"), ResourceName), OwnerName, EHairResourceUsageType::Static);
+	InternalCreateVertexBufferRDG_FromHairBulkData<FHairClusterIndexFormat>(GraphBuilder, BulkData.Data.CurveToClusterIds, BulkData.Header.CurveCount, CurveToClusterIdBuffer, ToHairResourceDebugName(TEXT("Hair.StrandsClusterCulling_CurveToClusterIds"), ResourceName), OwnerName, EHairResourceUsageType::Static);
+	InternalCreateVertexBufferRDG_FromHairBulkData<FHairClusterIndexFormat>(GraphBuilder, BulkData.Data.ClusterVertexIds, BulkData.Header.VertexLODCount, ClusterVertexIdBuffer, ToHairResourceDebugName(TEXT("Hair.StrandsClusterCulling_ClusterVertexIds"), ResourceName), OwnerName, EHairResourceUsageType::Static);
 }
 
 void FHairStrandsClusterCullingResource::InternalRelease()
