Fix long cook time due to incorrect seialization of strands root data.
Root data for strands are now skipped if a groom asset doesn't have strand geometry LOD.

diff --git a/Engine/Plugins/Runtime/HairStrands/Source/HairStrandsCore/Private/GroomBindingAsset.cpp b/Engine/Plugins/Runtime/HairStrands/Source/HairStrandsCore/Private/GroomBindingAsset.cpp
--- a/Engine/Plugins/Runtime/HairStrands/Source/HairStrandsCore/Private/GroomBindingAsset.cpp
+++ b/Engine/Plugins/Runtime/HairStrands/Source/HairStrandsCore/Private/GroomBindingAsset.cpp
@@ -547,10 +547,7 @@ bool UGroomBindingAsset::IsCompatible(const UGroomAsset* InGroom, const UGroomBi
 
 		if (!InBinding->Groom)
 		{
-			if (bIssueWarning)
-			{
-				UE_LOG(LogHairStrands, Warning, TEXT("[Groom] The binding asset (%s) does not reference a groom. Falling back onto non-binding version."), *InBinding->GetName());
-			}
+			UE_CLOG(bIssueWarning, LogHairStrands, Warning, TEXT("[Groom] The binding asset (%s) does not reference a groom. Falling back onto non-binding version."), *InBinding->GetName());
 			return false;
 		}
 
@@ -566,55 +563,58 @@ bool UGroomBindingAsset::IsCompatible(const UGroomAsset* InGroom, const UGroomBi
 		const uint32 GroupCount = InGroom->GetNumHairGroups();
 		if (GroupCount != InBinding->GroupInfos.Num())
 		{
-			if (bIssueWarning)
-			{
-				UE_LOG(LogHairStrands, Warning, TEXT("[Groom] The GroomBinding asset (%s) does not contains the same number of groups (%d vs. %d) than the groom (%s). The binding asset will not be used."),
-					*InBinding->GetName(),
-					GroupCount,
-					InBinding->GroupInfos.Num(),
-					*InGroom->GetName());
-			}
+			UE_CLOG(bIssueWarning, LogHairStrands, Warning, TEXT("[Groom] The GroomBinding asset (%s) does not contains the same number of groups (%d vs. %d) than the groom (%s). The binding asset will not be used."),
+				*InBinding->GetName(),
+				GroupCount,
+				InBinding->GroupInfos.Num(),
+				*InGroom->GetName());
 			return false;
 		}
 
 		for (uint32 GroupIt = 0; GroupIt < GroupCount; ++GroupIt)
 		{
+			// Guides
 			{
 				const uint32 GroomCount = InGroom->HairGroupsPlatformData[GroupIt].Guides.BulkData.GetNumCurves();
 				const uint32 BindingCount = InBinding->GroupInfos[GroupIt].SimRootCount;
 
 				if (GroomCount != 0 && GroomCount != BindingCount)
 				{
-					if (bIssueWarning)
-					{
-						UE_LOG(LogHairStrands, Warning, TEXT("[Groom] The GroomBinding asset (%s) does not contains the same guides in group %d (%d vs. %d) than the groom (%s). The binding asset will not be used."),
-							*InBinding->GetName(),
-							GroupIt,
-							GroomCount,
-							BindingCount,
-							*InGroom->GetName());
-					}
+					UE_CLOG(bIssueWarning, LogHairStrands, Warning, TEXT("[Groom] The GroomBinding asset (%s) does not contains the same guides in group %d (%d vs. %d) than the groom (%s). The binding asset will not be used."),
+						*InBinding->GetName(),
+						GroupIt,
+						GroomCount,
+						BindingCount,
+						*InGroom->GetName());
 					return false;
 				}
 			}
 
+			// Strands
 			if (IsHairStrandsEnabled(EHairStrandsShaderType::Strands))
 			{
 				const uint32 GroomCount = InGroom->HairGroupsPlatformData[GroupIt].Strands.BulkData.GetNumCurves();
 				const uint32 BindingCount = InBinding->GroupInfos[GroupIt].RenRootCount;
 
-				// Groom may have stripped strands data so GroomCount would be 0
-				if (GroomCount != 0 && GroomCount != BindingCount)
+				bool bNeedStrands = false;
+				for (const FHairLODSettings& LODSettings : InGroom->HairGroupsLOD[GroupIt].LODs)
 				{
-					if (bIssueWarning)
+					if (LODSettings.GeometryType == EGroomGeometryType::Strands)
 					{
-						UE_LOG(LogHairStrands, Warning, TEXT("[Groom] The GroomBinding asset (%s) does not contains the same curves in group %d (%d vs. %d) than the groom (%s). The binding asset will not be used."),
-							*InBinding->GetName(),
-							GroupIt,
-							GroomCount,
-							BindingCount,
-							*InGroom->GetName());
+						bNeedStrands = true;
+						break;
 					}
+				}
+
+				// Groom may have stripped strands data so GroomCount would be 0
+				if (bNeedStrands && GroomCount != 0 && GroomCount != BindingCount)
+				{
+					UE_CLOG(bIssueWarning, LogHairStrands, Warning, TEXT("[Groom] The GroomBinding asset (%s) does not contains the same curves in group %d (%d vs. %d) than the groom (%s). The binding asset will not be used."),
+						*InBinding->GetName(),
+						GroupIt,
+						GroomCount,
+						BindingCount,
+						*InGroom->GetName());
 					return false;
 				}
 			}
diff --git a/Engine/Plugins/Runtime/HairStrands/Source/HairStrandsCore/Private/GroomBindingBuilder.cpp b/Engine/Plugins/Runtime/HairStrands/Source/HairStrandsCore/Private/GroomBindingBuilder.cpp
--- a/Engine/Plugins/Runtime/HairStrands/Source/HairStrandsCore/Private/GroomBindingBuilder.cpp
+++ b/Engine/Plugins/Runtime/HairStrands/Source/HairStrandsCore/Private/GroomBindingBuilder.cpp
@@ -57,7 +57,7 @@ static FAutoConsoleVariableRef CVarHairStrandsBindingBuilderWarningEnable(TEXT("
 FString FGroomBindingBuilder::GetVersion()
 {
 	// Important to update the version when groom building changes
-	return TEXT("3c");
+	return TEXT("3e");
 }
 
 ///////////////////////////////////////////////////////////////////////////////////////////////////
@@ -700,6 +700,7 @@ namespace GroomBinding_RBFWeighting
 
 	void ComputeInterpolationWeights(
 		FHairRootGroupData& Out, 
+		const bool bNeedStrandsRoot,
 		const uint32 NumInterpolationPoints, 
 		const int32 MatchingSection, 
 		const IMeshData* MeshData, 
@@ -738,8 +739,14 @@ namespace GroomBinding_RBFWeighting
 					FWeightsBuilder InterpolationWeights(SampleCount, SampleCount,
 						PointsSampler.SamplePositions.GetData(), PointsSampler.SamplePositions.GetData());
 
+					// Guides
 					UpdateInterpolationWeights(InterpolationWeights, PointsSampler, LODIndex, Out.SimRootData);
-					UpdateInterpolationWeights(InterpolationWeights, PointsSampler, LODIndex, Out.RenRootData);
+
+					// Strands
+					if (bNeedStrandsRoot)
+					{
+						UpdateInterpolationWeights(InterpolationWeights, PointsSampler, LODIndex, Out.RenRootData);
+					}
 				}
 			}
 			else
@@ -754,8 +761,12 @@ namespace GroomBinding_RBFWeighting
 				FWeightsBuilder InterpolationWeights(SampleCount, SampleCount,
 					PointsSampler.SamplePositions.GetData(), PointsSampler.SamplePositions.GetData());
 
+				// Guides
+				UpdateInterpolationWeights(InterpolationWeights, PointsSampler, LODIndex, Out.SimRootData);
+
+				// Strands
+				if (bNeedStrandsRoot)
 				{
-					UpdateInterpolationWeights(InterpolationWeights, PointsSampler, LODIndex, Out.SimRootData);
 					UpdateInterpolationWeights(InterpolationWeights, PointsSampler, LODIndex, Out.RenRootData);
 				}
 			}
@@ -2007,16 +2018,18 @@ static void BuildRootBulkData(
 	UGroomBindingAsset::FHairGroupPlatformData& Out,
 	const FHairRootGroupData& In)
 {
-	{
-		BuildRootBulkData(Out.SimRootBulkData, In.SimRootData);
-		BuildRootBulkData(Out.RenRootBulkData, In.RenRootData);
+	// Guides
+	BuildRootBulkData(Out.SimRootBulkData, In.SimRootData);
 
-		const uint32 LODCount = In.CardsRootData.Num();
-		Out.CardsRootBulkData.SetNum(LODCount);
-		for (uint32 LODIt = 0; LODIt < LODCount; ++LODIt)
-		{
-			BuildRootBulkData(Out.CardsRootBulkData[LODIt], In.CardsRootData[LODIt]);
-		}
+	// Strands
+	BuildRootBulkData(Out.RenRootBulkData, In.RenRootData);
+
+	// Cards
+	const uint32 LODCount = In.CardsRootData.Num();
+	Out.CardsRootBulkData.SetNum(LODCount);
+	for (uint32 LODIt = 0; LODIt < LODCount; ++LODIt)
+	{
+		BuildRootBulkData(Out.CardsRootBulkData[LODIt], In.CardsRootData[LODIt]);
 	}
 }
 
@@ -2039,7 +2052,6 @@ static bool InternalBuildBinding_CPU(UGroomBindingAsset* BindingAsset, uint32 In
 	// 1. Build groom root data
 	FHairRootGroupData OutData;
 	{
-
 		BindingAsset->Groom->ConditionalPostLoad();
 
 		const int32 NumInterpolationPoints = BindingAsset->NumInterpolationPoints;
@@ -2079,15 +2091,36 @@ static bool InternalBuildBinding_CPU(UGroomBindingAsset* BindingAsset, uint32 In
 		const uint32 MeshLODCount = TargetMeshData->GetNumLODs();
 
 		check(InGroupIndex < uint32(GroomAsset->HairGroupsPlatformData.Num()));
-		const FHairGroupPlatformData& GroupData = GroomAsset->HairGroupsPlatformData[InGroupIndex];
+
+		// Check if root data are needs for strands
+		bool bNeedStrandsRoot = false;
+		for (const FHairLODSettings& LODSettings : GroomAsset->HairGroupsLOD[InGroupIndex].LODs)
 		{
-			FHairStrandsDatas StrandsData;
-			FHairStrandsDatas GuidesData;
-			GroomAsset->GetHairStrandsDatas(InGroupIndex, StrandsData, GuidesData);
+			if (LODSettings.GeometryType == EGroomGeometryType::Strands)
+			{
+				bNeedStrandsRoot = true;
+				break;
+			}
+		}
+		
+		// 1.1 Build guide/strands data
+		FHairStrandsDatas StrandsData;
+		FHairStrandsDatas GuidesData;
+		GroomAsset->GetHairStrandsDatas(InGroupIndex, StrandsData, GuidesData);
 
-			InitHairStrandsRootData(OutData.RenRootData, &StrandsData, MeshLODCount, NumInterpolationPoints);
+		// 1.2 Init root data for guides/strands/cards
+		const FHairGroupPlatformData& GroupData = GroomAsset->HairGroupsPlatformData[InGroupIndex];
+		{
+			// Guides
 			InitHairStrandsRootData(OutData.SimRootData, &GuidesData, MeshLODCount, NumInterpolationPoints);
 
+			// Strands
+			if (bNeedStrandsRoot)
+			{
+				InitHairStrandsRootData(OutData.RenRootData, &StrandsData, MeshLODCount, NumInterpolationPoints);
+			}
+
+			// Cards
 			const uint32 CardsLODCount = GroupData.Cards.LODs.Num();
 			OutData.CardsRootData.SetNum(GroupData.Cards.LODs.Num());
 			for (uint32 CardsLODIt = 0; CardsLODIt < CardsLODCount; ++CardsLODIt)
@@ -2108,65 +2141,67 @@ static bool InternalBuildBinding_CPU(UGroomBindingAsset* BindingAsset, uint32 In
 		const bool bNeedTransferPosition = SourceMeshData->IsValid();
 
 		// Create mapping between the source & target using their UV
-		uint32 WorkItemCount = 1 + (bNeedTransferPosition ? 1 : 0); //RBF + optional position transfer
-		{
-			WorkItemCount += 2; // Sim & Render
-			const uint32 CardsLODCount = OutData.CardsRootData.Num();
-			WorkItemCount += CardsLODCount;
+		uint32 WorkItemCount = 1 + (bNeedTransferPosition ? 1 : 0); // RBF + optional position transfer
+		{
+			// Guides
+			WorkItemCount += 1;
+			// Strands
+			WorkItemCount += bNeedStrandsRoot ? 1 : 0;
+			// Cards
+			WorkItemCount += OutData.CardsRootData.Num();
 		}
 	
 		uint32 WorkItemIndex = 0;
 		FScopedSlowTask SlowTask(WorkItemCount, LOCTEXT("BuildBindingData", "Building groom binding data"));
 		SlowTask.MakeDialog();
 
+		// 1.3 Transfer positions
 		TArray<TArray<FVector3f>> TransferredPositions;
 		if (bNeedTransferPosition)
 		{
-			bool bSucceed = GroomBinding_Transfer::Transfer( 
+			if (!GroomBinding_Transfer::Transfer(
 				SourceMeshData.Get(),
 				TargetMeshData.Get(),
-				TransferredPositions, BindingAsset->MatchingSection);
-
-			if (!bSucceed)
+				TransferredPositions, BindingAsset->MatchingSection))
 			{
+				UE_LOG(LogHairStrands, Error, TEXT("[Groom] Binding asset could not be built. Positions transfer between source and target mesh failed."));
 				return false;
 			}
-
 			SlowTask.EnterProgressFrame();
 		}
 
-		bool bSucceed = false;
+		// 1.4 Build root data for guides/strands/cards
 		{
-			FHairStrandsDatas StrandsData;
-			FHairStrandsDatas GuidesData;
-			GroomAsset->GetHairStrandsDatas(InGroupIndex, StrandsData, GuidesData);
-
-			bSucceed = GroomBinding_RootProjection::Project(
-				StrandsData,
-				TargetMeshData.Get(),
-				TransferredPositions,
-				OutData.RenRootData);
-			if (!bSucceed)
+			// Guides
 			{
-				UE_LOG(LogHairStrands, Error, TEXT("[Groom] Binding asset could not be built. Some strand roots are not close enough to the target mesh to be projected onto it."));
-				return false;
+				if (!GroomBinding_RootProjection::Project(
+					GuidesData,
+					TargetMeshData.Get(),
+					TransferredPositions,
+					OutData.SimRootData))
+				{
+					UE_LOG(LogHairStrands, Error, TEXT("[Groom] Binding asset could not be built. Some guide roots are not close enough to the target mesh to be projected onto it."));
+					return false; 
+				}
+				SlowTask.EnterProgressFrame();
 			}
 
-			SlowTask.EnterProgressFrame();
-
-			bSucceed = GroomBinding_RootProjection::Project(
-				GuidesData,
-				TargetMeshData.Get(),
-				TransferredPositions,
-				OutData.SimRootData);
-			if (!bSucceed) 
+			// Strands
+			if (bNeedStrandsRoot)
 			{
-				UE_LOG(LogHairStrands, Error, TEXT("[Groom] Binding asset could not be built. Some guide roots are not close enough to the target mesh to be projected onto it."));
-				return false; 
+				if (!GroomBinding_RootProjection::Project(
+					StrandsData,
+					TargetMeshData.Get(),
+					TransferredPositions,
+					OutData.RenRootData))
+				{
+					UE_LOG(LogHairStrands, Error, TEXT("[Groom] Binding asset could not be built. Some strand roots are not close enough to the target mesh to be projected onto it."));
+					return false;
+				}
+				SlowTask.EnterProgressFrame();
 			}
 
-			SlowTask.EnterProgressFrame();
-
+			// Cards
 			const uint32 CardsLODCount = OutData.CardsRootData.Num();
 			for (uint32 CardsLODIt = 0; CardsLODIt < CardsLODCount; ++CardsLODIt)
 			{
@@ -2176,26 +2211,26 @@ static bool InternalBuildBinding_CPU(UGroomBindingAsset* BindingAsset, uint32 In
 					const bool bIsValid = GroomAsset->GetHairCardsGuidesDatas(InGroupIndex, CardsLODIt, LODGuidesData);
 					if (bIsValid)
 					{
-						bSucceed = GroomBinding_RootProjection::Project(
+						if (!GroomBinding_RootProjection::Project(
 							LODGuidesData,
 							TargetMeshData.Get(),
 							TransferredPositions,
-							OutData.CardsRootData[CardsLODIt]);
-
-						if (!bSucceed) 
+							OutData.CardsRootData[CardsLODIt]))
 						{
 							UE_LOG(LogHairStrands, Error, TEXT("[Groom] Binding asset could not be built. Some cards guide roots are not close enough to the target mesh to be projected onto it."));
 							return false; 
 						}
 					}
 				}
-
 				SlowTask.EnterProgressFrame();
 			}
 		}
 		
-		GroomBinding_RBFWeighting::ComputeInterpolationWeights(OutData, BindingAsset->NumInterpolationPoints, BindingAsset->MatchingSection, TargetMeshData.Get(), TransferredPositions);
-		SlowTask.EnterProgressFrame();
+		// 1.5 RBF building
+		{
+			GroomBinding_RBFWeighting::ComputeInterpolationWeights(OutData, bNeedStrandsRoot, BindingAsset->NumInterpolationPoints, BindingAsset->MatchingSection, TargetMeshData.Get(), TransferredPositions);
+			SlowTask.EnterProgressFrame();
+		}
 	}
 
 	// 2. Release existing resources data
@@ -2210,7 +2245,7 @@ static bool InternalBuildBinding_CPU(UGroomBindingAsset* BindingAsset, uint32 In
 	}
 	check(OutHairGroupResources.Num() == 0);
 
-	// 2. Convert data to bulk data
+	// 3. Convert data to bulk data
 	GroomBinding_BulkCopy::BuildRootBulkData(BindingAsset->HairGroupsPlatformData[InGroupIndex], OutData);
 
 	BindingAsset->QueryStatus = UGroomBindingAsset::EQueryStatus::Completed;
diff --git a/Engine/Plugins/Runtime/HairStrands/Source/HairStrandsCore/Private/HairStrandsDatas.cpp b/Engine/Plugins/Runtime/HairStrands/Source/HairStrandsCore/Private/HairStrandsDatas.cpp
--- a/Engine/Plugins/Runtime/HairStrands/Source/HairStrandsCore/Private/HairStrandsDatas.cpp
+++ b/Engine/Plugins/Runtime/HairStrands/Source/HairStrandsCore/Private/HairStrandsDatas.cpp
@@ -317,7 +317,7 @@ bool FHairStreamingRequest::IsCompleted()
 #endif
 }
 // Request fullfil 2 use cases:
-// * Load DDC data and upload them to GPU
+// * Load IO/DDC data and upload them to GPU
 // * Load DDC data and store them into bulkdata for serialization
 void FHairStreamingRequest::Request(uint32 InRequestedCurveCount, uint32 InRequestedPointCount, int32 InLODIndex, FHairStrandsBulkCommon& In, bool bWait, bool bFillBulkdata, const FName& InOwnerName)
 {
