Fix random rebuilding of Groom asset/Groom binding.

diff --git a/Engine/Plugins/Runtime/HairStrands/Source/HairStrandsCore/Private/GroomAsset.cpp b/Engine/Plugins/Runtime/HairStrands/Source/HairStrandsCore/Private/GroomAsset.cpp
--- a/Engine/Plugins/Runtime/HairStrands/Source/HairStrandsCore/Private/GroomAsset.cpp
+++ b/Engine/Plugins/Runtime/HairStrands/Source/HairStrandsCore/Private/GroomAsset.cpp
@@ -2291,18 +2291,23 @@ bool UGroomAsset::IsFullyCached()
 	return GetDerivedDataCacheRef().AllCachedDataProbablyExists(CacheKeys);
 }
 
-FString UGroomAsset::GetDerivedDataKey()
+FString UGroomAsset::GetDerivedDataKey(bool bUseStrandCachedKey)
 {
-	FString Key;
-	const uint32 GroupCount = GetNumHairGroups();
-	for (uint32 GroupIt = 0; GroupIt < GroupCount; ++GroupIt)
-	{
-		FString StrandsKey = GetDerivedDataKeyForStrands(GroupIt);
-		Key += StrandsKey;
-		Key += GetDerivedDataKeyForCards(GroupIt, StrandsKey);
-		Key += GetDerivedDataKeyForMeshes(GroupIt);
-	}
-	return Key;
+	FString OutKey;
+	for (uint32 GroupIt = 0, GroupCount = GetNumHairGroups(); GroupIt < GroupCount; ++GroupIt)
+	{		
+		// Optionnaly use the cached strands key. This is to workaround an issue with legacy HairDescription which uses FEditorBulkData, 
+		// and whose StringId is altered when loaded (only for legacy hair description converted to FEditorBulkData). This StringId modification
+		// changes the DDC key of a groom asset depending if has been loaded or not, causing DDC miss
+		const bool bComputeStrandsKey = !bUseStrandCachedKey || StrandsDerivedDataKey[GroupIt].IsEmpty();
+		const FString StrandsKey = bComputeStrandsKey ? GetDerivedDataKeyForStrands(GroupIt) : StrandsDerivedDataKey[GroupIt];
+
+		// Compute group key
+		OutKey += StrandsKey;
+		OutKey += GetDerivedDataKeyForCards(GroupIt, StrandsKey);
+		OutKey += GetDerivedDataKeyForMeshes(GroupIt);
+	}
+	return OutKey;
 }
 
 FString UGroomAsset::GetDerivedDataKeyForCards(uint32 GroupIndex, const FString& StrandsKey)
diff --git a/Engine/Plugins/Runtime/HairStrands/Source/HairStrandsCore/Private/GroomBindingAsset.cpp b/Engine/Plugins/Runtime/HairStrands/Source/HairStrandsCore/Private/GroomBindingAsset.cpp
--- a/Engine/Plugins/Runtime/HairStrands/Source/HairStrandsCore/Private/GroomBindingAsset.cpp
+++ b/Engine/Plugins/Runtime/HairStrands/Source/HairStrandsCore/Private/GroomBindingAsset.cpp
@@ -765,7 +765,9 @@ static FString BuildDerivedDataKeySuffix(const UGroomBindingAsset& BindingAsset)
 		SourceKey = BindingAsset.SourceGeometryCache ? BindingAsset.SourceGeometryCache->GetHash() : FString();
 		TargetKey = BindingAsset.TargetGeometryCache ? BindingAsset.TargetGeometryCache->GetHash() : FString();
 	}
-	FString GroomKey  = BindingAsset.Groom ? BindingAsset.Groom->GetDerivedDataKey()  : FString();
+	// When possible, use the GroomAsset 'cached DDC key'. This allows to avoid a bug where the DDC key would change 
+	// when loading GroomAsset's hair description, which would modify the hair description hash ID with legacy content.
+	FString GroomKey  = BindingAsset.Groom ? BindingAsset.Groom->GetDerivedDataKey(true /*bUseCachedKey*/) : FString();
 	FString PointKey  = FString::FromInt(BindingAsset.NumInterpolationPoints);
 	FString SectionKey = FString::FromInt(BindingAsset.MatchingSection);
 
diff --git a/Engine/Plugins/Runtime/HairStrands/Source/HairStrandsCore/Public/GroomAsset.h b/Engine/Plugins/Runtime/HairStrands/Source/HairStrandsCore/Public/GroomAsset.h
--- a/Engine/Plugins/Runtime/HairStrands/Source/HairStrandsCore/Public/GroomAsset.h
+++ b/Engine/Plugins/Runtime/HairStrands/Source/HairStrandsCore/Public/GroomAsset.h
@@ -690,7 +690,7 @@ public:
 	bool CacheCardsData(uint32 GroupIndex, const FString& StrandsKey);
 	bool CacheMeshesData(uint32 GroupIndex);
 
-	FString GetDerivedDataKey();
+	FString GetDerivedDataKey(bool bUseCacheKey=false);
 	FString GetDerivedDataKeyForCards(uint32 GroupIt, const FString& StrandsKey);
 	FString GetDerivedDataKeyForStrands(uint32 GroupIndex);
 	FString GetDerivedDataKeyForMeshes(uint32 GroupIndex);
