Fix crash when updating groom resource, as the 'data loaded marker' wasn't reset properly.

diff --git a/Engine/Plugins/Runtime/HairStrands/Source/HairStrandsCore/Private/GroomResources.cpp b/Engine/Plugins/Runtime/HairStrands/Source/HairStrandsCore/Private/GroomResources.cpp
--- a/Engine/Plugins/Runtime/HairStrands/Source/HairStrandsCore/Private/GroomResources.cpp
+++ b/Engine/Plugins/Runtime/HairStrands/Source/HairStrandsCore/Private/GroomResources.cpp
@@ -613,11 +613,23 @@ OwnerName(InOwnerName)
 
 void FHairCommonResource::InitRHI()
 {
-	if (bIsInitialized || AllocationType == EHairStrandsAllocationType::Deferred || GUsingNullRHI) { return; }
+	// 1. If the resource is Updated (ReleaseRHI, IniRHI), reset its loaded data counter
+	if (!bIsInitialized)
+	{
+		InternalResetLoadedSize();
+	}
+
+	// 2. If the resouce creation is deferred, skip the initialization
+	if (bIsInitialized || AllocationType == EHairStrandsAllocationType::Deferred || GUsingNullRHI) 
+	{ 
+		return; 
+	}
 
+	// 3. Loaded the data
 	const int32 DummyLODIndex = -1;
 	check(InternalIsDataLoaded(HAIR_MAX_NUM_CURVE_PER_GROUP, HAIR_MAX_NUM_POINT_PER_GROUP, DummyLODIndex));
 
+	// 4. Allocate the resource, and update the data
 	if (bUseRenderGraph)
 	{
 		FRHICommandListImmediate& RHICmdList = FRHICommandListExecutor::GetImmediateCommandList();
@@ -1056,6 +1068,11 @@ void FHairStrandsRestResource::InternalRelease()
 	MaxAvailableCurveCount = 0;
 }
 
+void FHairStrandsRestResource::InternalResetLoadedSize()
+{
+	BulkData.ResetLoadedSize();
+}
+
 /////////////////////////////////////////////////////////////////////////////////////////
 
 FHairStrandsDeformedResource::FHairStrandsDeformedResource(FHairStrandsBulkData& InBulkData, EHairStrandsResourcesType InCurveType, const FHairResourceName& InResourceName, const FName& InOwnerName) :
@@ -1172,6 +1189,11 @@ void FHairStrandsClusterCullingResource::InternalRelease()
 	MaxAvailableCurveCount = 0;
 }
 
+void FHairStrandsClusterCullingResource::InternalResetLoadedSize()
+{
+	BulkData.ResetLoadedSize();
+}
+
 /////////////////////////////////////////////////////////////////////////////////////////
 
 FHairStrandsRestRootResource::FHairStrandsRestRootResource(FHairStrandsRootBulkData& InBulkData, EHairStrandsResourcesType InCurveType, const FHairResourceName& InResourceName, const FName& InOwnerName) :
@@ -1285,6 +1307,11 @@ void FHairStrandsRestRootResource::InternalRelease()
 	LODs.Empty();
 }
 
+void FHairStrandsRestRootResource::InternalResetLoadedSize()
+{
+	BulkData.ResetLoadedSize();
+}
+
 /////////////////////////////////////////////////////////////////////////////////////////
 
 FHairStrandsDeformedRootResource::FHairStrandsDeformedRootResource(EHairStrandsResourcesType InCurveType, const FHairResourceName& InResourceName, const FName& InOwnerName):
@@ -1415,6 +1442,11 @@ void FHairStrandsInterpolationResource::InternalRelease()
 	MaxAvailableCurveCount = 0;
 }
 
+void FHairStrandsInterpolationResource::InternalResetLoadedSize()
+{
+	BulkData.ResetLoadedSize();
+}
+
 /////////////////////////////////////////////////////////////////////////////////////////
 
 FHairCardsInterpolationResource::FHairCardsInterpolationResource(FHairCardsInterpolationBulkData& InBulkData, const FHairResourceName& InResourceName, const FName& InOwnerName) :
diff --git a/Engine/Plugins/Runtime/HairStrands/Source/HairStrandsCore/Private/HairStrandsDatas.cpp b/Engine/Plugins/Runtime/HairStrands/Source/HairStrandsCore/Private/HairStrandsDatas.cpp
--- a/Engine/Plugins/Runtime/HairStrands/Source/HairStrandsCore/Private/HairStrandsDatas.cpp
+++ b/Engine/Plugins/Runtime/HairStrands/Source/HairStrandsCore/Private/HairStrandsDatas.cpp
@@ -531,6 +531,14 @@ void FHairStrandsBulkData::Reset()
 	Data.Curves			= FHairBulkContainer();
 }
 
+void FHairStrandsBulkData::ResetLoadedSize()
+{
+	Data.Positions.LoadedSize		= 0;
+	Data.CurveAttributes.LoadedSize	= 0;
+	Data.PointAttributes.LoadedSize	= 0;
+	Data.PointToCurve.LoadedSize	= 0;
+	Data.Curves.LoadedSize			= 0;
+}
 
 /////////////////////////////////////////////////////////////////////////////////////////
 // Interpolation bulk data
@@ -583,6 +591,12 @@ void FHairStrandsInterpolationBulkData::GetResources(FHairStrandsBulkCommon::FQu
 	}
 }
 
+void FHairStrandsInterpolationBulkData::ResetLoadedSize()
+{
+	Data.Interpolation.LoadedSize		= 0;
+	Data.SimRootPointIndex.LoadedSize	= 0;
+}
+
 /////////////////////////////////////////////////////////////////////////////////////////
 // Cluster culling bulk data
 
@@ -610,6 +624,14 @@ void FHairStrandsClusterCullingBulkData::Reset()
 	Data.PackedClusterInfos = FHairBulkContainer();
 }
 
+void FHairStrandsClusterCullingBulkData::ResetLoadedSize()
+{
+	Data.ClusterLODInfos.LoadedSize 	= 0;
+	Data.CurveToClusterIds.LoadedSize 	= 0;
+	Data.ClusterVertexIds.LoadedSize 	= 0;
+	Data.PackedClusterInfos.LoadedSize	= 0;
+}
+
 void FHairStrandsClusterCullingBulkData::SerializeHeader(FArchive& Ar, UObject* Owner)
 {
 	Ar << Header.ClusterCount;
@@ -843,3 +865,20 @@ void FHairStrandsRootBulkData::Reset()
 	}
 	Data.LODs.Empty();
 }
+
+void FHairStrandsRootBulkData::ResetLoadedSize()
+{
+	for (FData::FLOD& LOD : Data.LODs)
+	{
+		// Binding
+		LOD.RootBarycentricBuffer.LoadedSize 			= 0;
+		LOD.RootToUniqueTriangleIndexBuffer.LoadedSize 	= 0;
+		LOD.UniqueTriangleIndexBuffer.LoadedSize 		= 0;
+		LOD.RestUniqueTrianglePositionBuffer.LoadedSize = 0;
+		      
+		// RBF
+		LOD.MeshInterpolationWeightsBuffer.LoadedSize	= 0;
+		LOD.MeshSampleIndicesBuffer.LoadedSize 			= 0;
+		LOD.RestSamplePositionsBuffer.LoadedSize 		= 0;
+	}
+}
\ No newline at end of file
diff --git a/Engine/Plugins/Runtime/HairStrands/Source/HairStrandsCore/Public/GroomResources.h b/Engine/Plugins/Runtime/HairStrands/Source/HairStrandsCore/Public/GroomResources.h
--- a/Engine/Plugins/Runtime/HairStrands/Source/HairStrandsCore/Public/GroomResources.h
+++ b/Engine/Plugins/Runtime/HairStrands/Source/HairStrandsCore/Public/GroomResources.h
@@ -113,6 +113,7 @@ struct FHairCommonResource : public FRenderResource
 	virtual void InternalRelease() {}
 	virtual bool InternalIsDataLoaded(uint32 InRequestedCurveCount, uint32 InRequestedPointCount, int32 InLODIndex) { return true; }
 	virtual bool InternalIsLODDataLoaded(uint32 InRequestedCurveCount, uint32 InRequestedPointCount, int32 InLODIndex) const { return true; }
+	virtual void InternalResetLoadedSize() { }
 
 	bool bUseRenderGraph = true;
 	bool bIsInitialized = false;
@@ -140,6 +141,7 @@ struct FHairStrandsRestResource : public FHairCommonResource
 	virtual void InternalAllocate(FRDGBuilder& GraphBuilder) override;
 	virtual void InternalRelease() override;
 	virtual bool InternalIsDataLoaded(uint32 InRequestedCurveCount, uint32 InRequestedPointCount, int32 InLODIndex) override;
+	virtual void InternalResetLoadedSize() override;
 
 	/* Get the resource name */
 	virtual FString GetFriendlyName() const override { return TEXT("FHairStrandsResource"); }
@@ -285,6 +287,7 @@ struct FHairStrandsInterpolationResource : public FHairCommonResource
 	virtual void InternalAllocate(FRDGBuilder& GraphBuilder) override;
 	virtual void InternalRelease() override;
 	virtual bool InternalIsDataLoaded(uint32 InRequestedCurveCount, uint32 InRequestedPointCount, int32 InLODIndex) override;
+	virtual void InternalResetLoadedSize() override;
 
 	/* Get the resource name */
 	virtual FString GetFriendlyName() const override { return TEXT("FHairStrandsInterplationResource"); }
@@ -316,6 +319,7 @@ struct FHairStrandsClusterCullingResource : public FHairCommonResource
 	virtual void InternalAllocate(FRDGBuilder& GraphBuilder) override;
 	virtual void InternalRelease() override;
 	virtual bool InternalIsDataLoaded(uint32 InRequestedCurveCount, uint32 InRequestedPointCount, int32 InLODIndex) override;
+	virtual void InternalResetLoadedSize() override;
 
 	/* Get the resource name */
 	virtual FString GetFriendlyName() const override { return TEXT("FHairStrandsClusterResource"); }
@@ -355,6 +359,7 @@ struct FHairStrandsRestRootResource : public FHairCommonResource
 	virtual void InternalRelease() override;
 	virtual bool InternalIsDataLoaded(uint32 InRequestedCurveCount, uint32 InRequestedPointCount, int32 InLODIndex) override;
 	virtual bool InternalIsLODDataLoaded(uint32 InRequestedCurveCount, uint32 InRequestedPointCount, int32 InLODIndex) const override;
+	virtual void InternalResetLoadedSize() override;
 
 	/* Get the resource name */
 	virtual FString GetFriendlyName() const override { return TEXT("FHairStrandsRestRootResource"); }
diff --git a/Engine/Plugins/Runtime/HairStrands/Source/HairStrandsCore/Public/HairStrandsDatas.h b/Engine/Plugins/Runtime/HairStrands/Source/HairStrandsCore/Public/HairStrandsDatas.h
--- a/Engine/Plugins/Runtime/HairStrands/Source/HairStrandsCore/Public/HairStrandsDatas.h
+++ b/Engine/Plugins/Runtime/HairStrands/Source/HairStrandsCore/Public/HairStrandsDatas.h
@@ -444,6 +444,7 @@ struct HAIRSTRANDSCORE_API FHairStrandsBulkCommon
 	virtual uint32 GetResourceCount() const = 0;
 	virtual void GetResources(FQuery& Out) = 0;
 	virtual void GetResourceVersion(FArchive& Ar) const {}
+	virtual void ResetLoadedSize() = 0;
 
 #if WITH_EDITORONLY_DATA
 	// Transient Name/DDCkey for streaming
@@ -584,6 +585,7 @@ struct HAIRSTRANDSCORE_API FHairStrandsBulkData : FHairStrandsBulkCommon
 
 	bool IsValid() const { return Header.CurveCount > 0 && Header.PointCount > 0; }
 	void Reset();
+	virtual void ResetLoadedSize() override;
 
 	uint32 GetNumCurves() const { return Header.CurveCount;  };
 	uint32 GetNumPoints() const { return Header.PointCount; };
@@ -680,6 +682,7 @@ struct HAIRSTRANDSCORE_API FHairStrandsInterpolationBulkData : FHairStrandsBulkC
 	};
 
 	void Reset();
+	virtual void ResetLoadedSize() override;
 	virtual void SerializeHeader(FArchive& Ar, UObject* Owner) override;
 
 	virtual uint32 GetResourceCount() const override;
@@ -735,6 +738,7 @@ struct HAIRSTRANDSCORE_API FHairStrandsClusterCullingData
 struct HAIRSTRANDSCORE_API FHairStrandsClusterCullingBulkData : FHairStrandsBulkCommon
 {
 	void Reset();
+	virtual void ResetLoadedSize() override;
 
 	virtual void SerializeHeader(FArchive& Ar, UObject* Owner) override;
 	virtual uint32 GetResourceCount() const override;
@@ -829,6 +833,7 @@ struct FHairStrandsRootBulkData : FHairStrandsBulkCommon
 	virtual void GetResources(FQuery& Out) override;
 
 	void Reset();
+	virtual void ResetLoadedSize() override;
 	bool IsValid() const { return Header.RootCount > 0; }
 	const TArray<uint32>& GetValidSectionIndices(int32 LODIndex) const;
 	uint32 GetLODCount() const { return Header.LODs.Num(); }
