Rename groom InternalIsDataLoaded into InternalGetOrRequestData for clarity

diff --git a/Engine/Plugins/Runtime/HairStrands/Source/HairStrandsCore/Private/GroomResources.cpp b/Engine/Plugins/Runtime/HairStrands/Source/HairStrandsCore/Private/GroomResources.cpp
--- a/Engine/Plugins/Runtime/HairStrands/Source/HairStrandsCore/Private/GroomResources.cpp
+++ b/Engine/Plugins/Runtime/HairStrands/Source/HairStrandsCore/Private/GroomResources.cpp
@@ -627,7 +627,7 @@ void FHairCommonResource::InitRHI()
 
 	// 3. Loaded the data
 	const int32 DummyLODIndex = -1;
-	check(InternalIsDataLoaded(HAIR_MAX_NUM_CURVE_PER_GROUP, HAIR_MAX_NUM_POINT_PER_GROUP, DummyLODIndex));
+	check(InternalGetOrRequestData(HAIR_MAX_NUM_CURVE_PER_GROUP, HAIR_MAX_NUM_POINT_PER_GROUP, DummyLODIndex));
 
 	// 4. Allocate the resource, and update the data
 	if (bUseRenderGraph)
@@ -694,7 +694,7 @@ void FHairCommonResource::Allocate(FRDGBuilder& GraphBuilder, EHairResourceLoadi
 			Status |= EHairResourceStatus::EStatus::Valid; 
 		}
 		// 2. If more curves are requested, issue a streaming request
-		else if (InternalIsDataLoaded(InRequestedCurveCount, InRequestedPointCount, InLODIndex))
+		else if (InternalGetOrRequestData(InRequestedCurveCount, InRequestedPointCount, InLODIndex))
 		{ 
 			// 2.1 Curve data are availble, and update GPU resources
 			if (!bIsInitialized)
@@ -727,7 +727,7 @@ void FHairCommonResource::StreamInData(int32 InLODIndex)
 	if (!bIsInitialized)
 	{
 		// TODO
-		InternalIsDataLoaded(HAIR_MAX_NUM_CURVE_PER_GROUP, HAIR_MAX_NUM_POINT_PER_GROUP, InLODIndex);
+		InternalGetOrRequestData(HAIR_MAX_NUM_CURVE_PER_GROUP, HAIR_MAX_NUM_POINT_PER_GROUP, InLODIndex);
 	}
 }
 
@@ -978,7 +978,7 @@ FHairStrandsRestResource::FHairStrandsRestResource(FHairStrandsBulkData& InBulkD
 	check(!!(BulkData.Header.Flags & FHairStrandsBulkData::DataFlags_HasData));
 }
 
-bool FHairStrandsRestResource::InternalIsDataLoaded(uint32 InRequestedCurveCount, uint32 InRequestedPointCount, int32 InLODIndex)
+bool FHairStrandsRestResource::InternalGetOrRequestData(uint32 InRequestedCurveCount, uint32 InRequestedPointCount, int32 InLODIndex)
 {
 	if (StreamingRequest.IsNone())
 	{
@@ -1155,7 +1155,7 @@ FHairStrandsClusterCullingResource::FHairStrandsClusterCullingResource(FHairStra
 	MaxAvailableCurveCount = 0;
 }
 
-bool FHairStrandsClusterCullingResource::InternalIsDataLoaded(uint32 InRequestedCurveCount, uint32 InRequestedPointCount, int32 InLODIndex)
+bool FHairStrandsClusterCullingResource::InternalGetOrRequestData(uint32 InRequestedCurveCount, uint32 InRequestedPointCount, int32 InLODIndex)
 {
 	if (StreamingRequest.IsNone())
 	{
@@ -1233,7 +1233,7 @@ bool FHairStrandsRestRootResource::InternalIsLODDataLoaded(uint32 InRequestedCur
 	return true;
 }
 
-bool FHairStrandsRestRootResource::InternalIsDataLoaded(uint32 InRequestedCurveCount, uint32 InRequestedPointCount, int32 InLODIndex)
+bool FHairStrandsRestRootResource::InternalGetOrRequestData(uint32 InRequestedCurveCount, uint32 InRequestedPointCount, int32 InLODIndex)
 {
 	if (StreamingRequest.IsNone() && InLODIndex >= 0)
 	{
@@ -1417,7 +1417,7 @@ FHairStrandsInterpolationResource::FHairStrandsInterpolationResource(FHairStrand
 	check(!!(BulkData.Header.Flags & FHairStrandsInterpolationBulkData::DataFlags_HasData));
 }
 
-bool FHairStrandsInterpolationResource::InternalIsDataLoaded(uint32 InRequestedCurveCount, uint32 InRequestedPointCount, int32 InLODIndex)
+bool FHairStrandsInterpolationResource::InternalGetOrRequestData(uint32 InRequestedCurveCount, uint32 InRequestedPointCount, int32 InLODIndex)
 {
 	if (StreamingRequest.IsNone())
 	{
diff --git a/Engine/Plugins/Runtime/HairStrands/Source/HairStrandsCore/Public/GroomResources.h b/Engine/Plugins/Runtime/HairStrands/Source/HairStrandsCore/Public/GroomResources.h
--- a/Engine/Plugins/Runtime/HairStrands/Source/HairStrandsCore/Public/GroomResources.h
+++ b/Engine/Plugins/Runtime/HairStrands/Source/HairStrandsCore/Public/GroomResources.h
@@ -111,7 +111,7 @@ struct FHairCommonResource : public FRenderResource
 	virtual void InternalAllocate(FRDGBuilder& GraphBuilder) {}
 	virtual void InternalAllocate(FRDGBuilder& GraphBuilder, uint32 InCurveCount, int32 InLODIndex) { InternalAllocate(GraphBuilder); }
 	virtual void InternalRelease() {}
-	virtual bool InternalIsDataLoaded(uint32 InRequestedCurveCount, uint32 InRequestedPointCount, int32 InLODIndex) { return true; }
+	virtual bool InternalGetOrRequestData(uint32 InRequestedCurveCount, uint32 InRequestedPointCount, int32 InLODIndex) { return true; }
 	virtual bool InternalIsLODDataLoaded(uint32 InRequestedCurveCount, uint32 InRequestedPointCount, int32 InLODIndex) const { return true; }
 	virtual void InternalResetLoadedSize() { }
 
@@ -140,7 +140,7 @@ struct FHairStrandsRestResource : public FHairCommonResource
 	/* Init/Release buffers */
 	virtual void InternalAllocate(FRDGBuilder& GraphBuilder) override;
 	virtual void InternalRelease() override;
-	virtual bool InternalIsDataLoaded(uint32 InRequestedCurveCount, uint32 InRequestedPointCount, int32 InLODIndex) override;
+	virtual bool InternalGetOrRequestData(uint32 InRequestedCurveCount, uint32 InRequestedPointCount, int32 InLODIndex) override;
 	virtual void InternalResetLoadedSize() override;
 
 	/* Get the resource name */
@@ -286,7 +286,7 @@ struct FHairStrandsInterpolationResource : public FHairCommonResource
 	/* Init/Release buffers */
 	virtual void InternalAllocate(FRDGBuilder& GraphBuilder) override;
 	virtual void InternalRelease() override;
-	virtual bool InternalIsDataLoaded(uint32 InRequestedCurveCount, uint32 InRequestedPointCount, int32 InLODIndex) override;
+	virtual bool InternalGetOrRequestData(uint32 InRequestedCurveCount, uint32 InRequestedPointCount, int32 InLODIndex) override;
 	virtual void InternalResetLoadedSize() override;
 
 	/* Get the resource name */
@@ -318,7 +318,7 @@ struct FHairStrandsClusterCullingResource : public FHairCommonResource
 	/* Init/Release buffers */
 	virtual void InternalAllocate(FRDGBuilder& GraphBuilder) override;
 	virtual void InternalRelease() override;
-	virtual bool InternalIsDataLoaded(uint32 InRequestedCurveCount, uint32 InRequestedPointCount, int32 InLODIndex) override;
+	virtual bool InternalGetOrRequestData(uint32 InRequestedCurveCount, uint32 InRequestedPointCount, int32 InLODIndex) override;
 	virtual void InternalResetLoadedSize() override;
 
 	/* Get the resource name */
@@ -357,7 +357,7 @@ struct FHairStrandsRestRootResource : public FHairCommonResource
 	/* Init/Release buffers */
 	virtual void InternalAllocate(FRDGBuilder& GraphBuilder, uint32 InCurveCount, int32 InLODIndex) override;
 	virtual void InternalRelease() override;
-	virtual bool InternalIsDataLoaded(uint32 InRequestedCurveCount, uint32 InRequestedPointCount, int32 InLODIndex) override;
+	virtual bool InternalGetOrRequestData(uint32 InRequestedCurveCount, uint32 InRequestedPointCount, int32 InLODIndex) override;
 	virtual bool InternalIsLODDataLoaded(uint32 InRequestedCurveCount, uint32 InRequestedPointCount, int32 InLODIndex) const override;
 	virtual void InternalResetLoadedSize() override;
 
