Hair Strands: Fix missing GROUP_SIZE on some shaders

This is a followup to 25203902

diff --git a/Engine/Plugins/Runtime/HairStrands/Source/HairStrandsCore/Private/HairStrandsInterpolation.cpp b/Engine/Plugins/Runtime/HairStrands/Source/HairStrandsCore/Private/HairStrandsInterpolation.cpp
--- a/Engine/Plugins/Runtime/HairStrands/Source/HairStrandsCore/Private/HairStrandsInterpolation.cpp
+++ b/Engine/Plugins/Runtime/HairStrands/Source/HairStrandsCore/Private/HairStrandsInterpolation.cpp
@@ -821,7 +821,7 @@ class FHairClusterAABBCS : public FGlobalShader
 	DECLARE_GLOBAL_SHADER(FHairClusterAABBCS);
 	SHADER_USE_PARAMETER_STRUCT(FHairClusterAABBCS, FGlobalShader);
 
-	class FGroupSize : SHADER_PERMUTATION_SPARSE_INT("PERMUTATION_GROUP_SIZE", 32, 64);
+	class FGroupSize : SHADER_PERMUTATION_SPARSE_INT("GROUP_SIZE", 32, 64);
 	class FCPUAABB : SHADER_PERMUTATION_BOOL("PERMUTATION_USE_CPU_AABB");
 	using FPermutationDomain = TShaderPermutationDomain<FGroupSize, FCPUAABB>;
 
diff --git a/Engine/Shaders/Private/HairStrands/HairStrandsInterpolation.usf b/Engine/Shaders/Private/HairStrands/HairStrandsInterpolation.usf
--- a/Engine/Shaders/Private/HairStrands/HairStrandsInterpolation.usf
+++ b/Engine/Shaders/Private/HairStrands/HairStrandsInterpolation.usf
@@ -590,14 +590,6 @@ void MainCS(uint2 DispatchThreadId : SV_DispatchThreadID, uint GroupThread1D : S
 
 #if SHADER_CLUSTERAABB
 
-#if PERMUTATION_GROUP_SIZE == 64
-	#define GROUP_SIZE 64
-#elif PERMUTATION_GROUP_SIZE == 32
-	#define GROUP_SIZE 32
-#else
-	#error Unknown group size	
-#endif
-
 uint TotalClusterCount;
 float3 CPUBoundMin;
 float3 CPUBoundMax;
