Fix missing binding loading in CitySample due to to incorrect check.

Binding asset no longer store strands data if this is not requried (e.g., a groom doesn't have LOD with strand geometry). The binding validation code was incorretly asserting the goom binding had strands root data.

diff --git a/Engine/Plugins/Runtime/HairStrands/Source/HairStrandsCore/Private/GroomBindingAsset.cpp b/Engine/Plugins/Runtime/HairStrands/Source/HairStrandsCore/Private/GroomBindingAsset.cpp
--- a/Engine/Plugins/Runtime/HairStrands/Source/HairStrandsCore/Private/GroomBindingAsset.cpp
+++ b/Engine/Plugins/Runtime/HairStrands/Source/HairStrandsCore/Private/GroomBindingAsset.cpp
@@ -535,6 +535,21 @@ bool UGroomBindingAsset::IsCompatible(const UGeometryCache* InGeometryCache, con
 	return true;
 }
 
+static bool DoesGroomNeedStrandsBinding(const UGroomAsset* InGroom, uint32 InGroupIndex)
+{
+	if (IsHairStrandsEnabled(EHairStrandsShaderType::Strands))
+	{
+		for (const FHairLODSettings& LODSettings : InGroom->HairGroupsLOD[InGroupIndex].LODs)
+		{
+			if (LODSettings.bVisible && LODSettings.GeometryType == EGroomGeometryType::Strands)
+			{
+				return true;
+			}
+		}
+	}
+	return false;
+}
+
 bool UGroomBindingAsset::IsCompatible(const UGroomAsset* InGroom, const UGroomBindingAsset* InBinding, bool bIssueWarning)
 {
 	if (InBinding && InGroom && IsHairStrandsBindingEnable())
@@ -591,23 +606,14 @@ bool UGroomBindingAsset::IsCompatible(const UGroomAsset* InGroom, const UGroomBi
 			}
 
 			// Strands
-			if (IsHairStrandsEnabled(EHairStrandsShaderType::Strands))
+			const bool bNeedStrandsRoot = DoesGroomNeedStrandsBinding(InGroom, GroupIt);
+			if (bNeedStrandsRoot)
 			{
 				const uint32 GroomCount = InGroom->HairGroupsPlatformData[GroupIt].Strands.BulkData.GetNumCurves();
 				const uint32 BindingCount = InBinding->GroupInfos[GroupIt].RenRootCount;
 
-				bool bNeedStrands = false;
-				for (const FHairLODSettings& LODSettings : InGroom->HairGroupsLOD[GroupIt].LODs)
-				{
-					if (LODSettings.GeometryType == EGroomGeometryType::Strands)
-					{
-						bNeedStrands = true;
-						break;
-					}
-				}
-
 				// Groom may have stripped strands data so GroomCount would be 0
-				if (bNeedStrands && GroomCount != 0 && GroomCount != BindingCount)
+				if (GroomCount != 0 && GroomCount != BindingCount)
 				{
 					UE_CLOG(bIssueWarning, LogHairStrands, Warning, TEXT("[Groom] The GroomBinding asset (%s) does not contains the same curves in group %d (%d vs. %d) than the groom (%s). The binding asset will not be used."),
 						*InBinding->GetName(),
@@ -653,6 +659,7 @@ bool UGroomBindingAsset::IsBindingAssetValid(const UGroomBindingAsset* InBinding
 			return false;
 		}
 
+		uint32 GroupIt = 0;
 		for (const FGoomBindingGroupInfo& Info : InBinding->GroupInfos)
 		{
 			if (Info.SimRootCount == 0)
@@ -664,7 +671,8 @@ bool UGroomBindingAsset::IsBindingAssetValid(const UGroomBindingAsset* InBinding
 				return false;
 			}
 
-			if (Info.RenRootCount == 0 && IsHairStrandsEnabled(EHairStrandsShaderType::Strands))
+			const bool bNeedStrandsRoot = DoesGroomNeedStrandsBinding(InBinding->Groom, GroupIt);
+			if (bNeedStrandsRoot && Info.RenRootCount == 0 && IsHairStrandsEnabled(EHairStrandsShaderType::Strands))
 			{
 				if (bIssueWarning)
 				{
@@ -672,6 +680,8 @@ bool UGroomBindingAsset::IsBindingAssetValid(const UGroomBindingAsset* InBinding
 				}
 				return false;
 			}
+
+			++GroupIt;
 		}
 	}
 	return true;
