Fix non-unity
#rnx

diff --git a/Engine/Plugins/Runtime/HairStrands/Source/HairStrandsCore/Private/GroomResources.cpp b/Engine/Plugins/Runtime/HairStrands/Source/HairStrandsCore/Private/GroomResources.cpp
--- a/Engine/Plugins/Runtime/HairStrands/Source/HairStrandsCore/Private/GroomResources.cpp
+++ b/Engine/Plugins/Runtime/HairStrands/Source/HairStrandsCore/Private/GroomResources.cpp
@@ -26,6 +26,7 @@
 #include "RenderGraph.h"
 #include "RenderGraphUtils.h"
 #include "RenderTargetPool.h"
+#include "RHIUtilities.h"
 #include "GroomBindingBuilder.h"
 #if WITH_EDITORONLY_DATA
 #include "Containers/StringView.h"
