Fix PIE crash with groom

diff --git a/Engine/Plugins/Runtime/HairStrands/Source/HairStrandsCore/Private/Deformer/DeformerDataInterfaceGroomWrite.cpp b/Engine/Plugins/Runtime/HairStrands/Source/HairStrandsCore/Private/Deformer/DeformerDataInterfaceGroomWrite.cpp
--- a/Engine/Plugins/Runtime/HairStrands/Source/HairStrandsCore/Private/Deformer/DeformerDataInterfaceGroomWrite.cpp
+++ b/Engine/Plugins/Runtime/HairStrands/Source/HairStrandsCore/Private/Deformer/DeformerDataInterfaceGroomWrite.cpp
@@ -225,15 +225,19 @@ void FOptimusGroomWriteDataProviderProxy::AllocateResources(FRDGBuilder& GraphBu
 				FRDGExternalBuffer CurveAttributeBufferExt = GroomInstance->Strands.DeformedResource->GetDeformerCurveAttributeBuffer(GraphBuilder);
 				if (CurveAttributeBufferExt.Buffer)
 				{
-					FRDGImportedBuffer CurveAttributeBuffer = Register(GraphBuilder, CurveAttributeBufferExt, ERDGImportedBufferFlags::CreateUAV);
-					AddCopyBufferPass(GraphBuilder, CurveAttributeBuffer.Buffer, Register(GraphBuilder, GroomInstance->Strands.RestResource->CurveAttributeBuffer, ERDGImportedBufferFlags::None).Buffer);
-					R.CurveAttributeBufferUAV = CurveAttributeBuffer.UAV;
+					// Always copy the attributes from the rest asset, so that if deformer write different attribute at different tick, it all remains consistent with the source data
+					FRDGImportedBuffer DstCurveAttributeBuffer = Register(GraphBuilder, CurveAttributeBufferExt, ERDGImportedBufferFlags::CreateUAV);
+					FRDGImportedBuffer SrcCurveAttributeBuffer = Register(GraphBuilder, GroomInstance->Strands.RestResource->CurveAttributeBuffer, ERDGImportedBufferFlags::None);
+					if (SrcCurveAttributeBuffer.Buffer)
+					{
+						AddCopyBufferPass(GraphBuilder, DstCurveAttributeBuffer.Buffer, SrcCurveAttributeBuffer.Buffer);
+					}
+					R.CurveAttributeBufferUAV = DstCurveAttributeBuffer.UAV;
 				}
 			}
 			if (R.CurveAttributeBufferUAV == nullptr)
 			{
 				R.CurveAttributeBufferUAV = R.AttributeBufferUAV_fallback;
-				//R.CurveAttributeBufferUAV = GraphBuilder.CreateUAV(GraphBuilder.RegisterExternalBuffer(GWhiteVertexBufferWithRDG->Buffer), PF_R32_UINT);
 			}
 
 			// Point Attributes
