Fix groom details refreshing in BP.

diff --git a/Engine/Plugins/Runtime/HairStrands/Source/HairStrandsEditor/Private/GroomComponentDetailsCustomization.cpp b/Engine/Plugins/Runtime/HairStrands/Source/HairStrandsEditor/Private/GroomComponentDetailsCustomization.cpp
--- a/Engine/Plugins/Runtime/HairStrands/Source/HairStrandsEditor/Private/GroomComponentDetailsCustomization.cpp
+++ b/Engine/Plugins/Runtime/HairStrands/Source/HairStrandsEditor/Private/GroomComponentDetailsCustomization.cpp
@@ -72,6 +72,17 @@ void FGroomComponentDetailsCustomization::CustomizeDetails(IDetailLayoutBuilder&
 	}
 	GroomComponentPtr = GroomComponent;
 
+	if (GroomComponentPtr.IsValid())
+	{
+		// Ensure the group desc is update to date
+		GroomComponentPtr->UpdateHairGroupsDesc();
+
+		// Add a delegate to refresh details panel when the groom asset changed
+		TSharedRef<IPropertyHandle> GroomAssetHandle = DetailLayout.GetProperty(GET_MEMBER_NAME_CHECKED(UGroomComponent, GroomAsset));
+		FSimpleDelegate OnGroomAssetChangedDelegate = FSimpleDelegate::CreateSP(this, &FGroomComponentDetailsCustomization::OnGroomAssetChanged, &DetailLayout);
+		GroomAssetHandle->SetOnPropertyValueChanged(OnGroomAssetChangedDelegate);
+	}
+
 	IDetailCategoryBuilder& HairGroupCategory = DetailLayout.EditCategory("GroomGroupsDesc", FText::GetEmpty(), ECategoryPriority::TypeSpecific);
 	CustomizeDescGroupProperties(DetailLayout, HairGroupCategory);
 }
@@ -266,6 +277,11 @@ void FGroomComponentDetailsCustomization::OnGenerateElementForHairGroup(TSharedR
 	}
 }
 
+void FGroomComponentDetailsCustomization::OnGroomAssetChanged(IDetailLayoutBuilder* LayoutBuilder)
+{
+	LayoutBuilder->ForceRefreshDetails();
+}
+
 //////////////////////////////////////////////////////////////////////////
 
 #undef LOCTEXT_NAMESPACE
diff --git a/Engine/Plugins/Runtime/HairStrands/Source/HairStrandsEditor/Private/GroomComponentDetailsCustomization.h b/Engine/Plugins/Runtime/HairStrands/Source/HairStrandsEditor/Private/GroomComponentDetailsCustomization.h
--- a/Engine/Plugins/Runtime/HairStrands/Source/HairStrandsEditor/Private/GroomComponentDetailsCustomization.h
+++ b/Engine/Plugins/Runtime/HairStrands/Source/HairStrandsEditor/Private/GroomComponentDetailsCustomization.h
@@ -37,4 +37,5 @@ private:
 	bool ShouldResetToDefault(TSharedPtr<IPropertyHandle> ChildHandle, int32 GroupIndex);
 	void ResetToDefault(TSharedPtr<IPropertyHandle> ChildHandle, int32 GroupIndex);
 	void AddPropertyWithCustomReset(TSharedPtr<IPropertyHandle>& PropertyHandle, IDetailChildrenBuilder& Builder, int32 GroupIndex);
+	void OnGroomAssetChanged(IDetailLayoutBuilder* LayoutBuilder);
 };
