Fix crash when streaming hair resources due to mismatching LODIndex between the request and allocation.
Fix resources being reallocated unecessarely.

diff --git a/Engine/Plugins/Runtime/HairStrands/Source/HairStrandsCore/Private/GroomResources.cpp b/Engine/Plugins/Runtime/HairStrands/Source/HairStrandsCore/Private/GroomResources.cpp
--- a/Engine/Plugins/Runtime/HairStrands/Source/HairStrandsCore/Private/GroomResources.cpp
+++ b/Engine/Plugins/Runtime/HairStrands/Source/HairStrandsCore/Private/GroomResources.cpp
@@ -651,6 +651,16 @@ void FHairCommonResource::ReleaseRHI()
 	bIsInitialized = false;
 }
 
+bool FHairCommonResource::InternalGetOrRequestData(uint32 InRequestedCurveCount, uint32 InRequestedPointCount, int32 InLODIndex) 
+{ 
+	// Initialize the streaming request so that resource which actually don't emit streaming request get correct #curve/#point/LODIndx
+	// request when allocating its resources
+	StreamingRequest.CurveCount = InRequestedCurveCount; 
+	StreamingRequest.PointCount = InRequestedPointCount;  
+	StreamingRequest.LODIndex   = InLODIndex;
+	return true; 
+}
+
 void FHairCommonResource::Allocate(FRDGBuilder& GraphBuilder, EHairResourceLoadingType LoadingType)
 {
 	EHairResourceStatus Status;
@@ -702,7 +712,7 @@ void FHairCommonResource::Allocate(FRDGBuilder& GraphBuilder, EHairResourceLoadi
 			{
 				FRenderResource::InitResource(); // Call RenderResource InitResource() so that the resources is marked as initialized
 			}
-			InternalAllocate(GraphBuilder, StreamingRequest.CurveCount, InLODIndex);
+			InternalAllocate(GraphBuilder, StreamingRequest.CurveCount, StreamingRequest.LODIndex);
 			bIsInitialized = true;
 
 			// Update the max curve count available
@@ -711,7 +721,17 @@ void FHairCommonResource::Allocate(FRDGBuilder& GraphBuilder, EHairResourceLoadi
 			// Reset streaming request. When the request is delete, the DDC request becomes cancelled. 
 			StreamingRequest = FHairStreamingRequest();
 
-			Status |= EHairResourceStatus::EStatus::Valid;
+			// Validate the resource the resource again, as the LODIndex at request time (StreamingRequest.LODIndex) 
+			// and the current LODIndex (InLODIndex) might not match.
+			// If they don't match a new streaming request will be emitted next frame
+			if (InternalIsLODDataLoaded(InRequestedCurveCount, InRequestedPointCount, InLODIndex))
+			{
+				Status |= EHairResourceStatus::EStatus::Valid;
+			}
+			else
+			{
+				Status |= EHairResourceStatus::EStatus::Loading;
+			}
 		}
 		else
 		{
@@ -728,7 +748,7 @@ void FHairCommonResource::StreamInData(int32 InLODIndex)
 	if (!bIsInitialized)
 	{
 		// TODO
-		InternalGetOrRequestData(HAIR_MAX_NUM_CURVE_PER_GROUP, HAIR_MAX_NUM_POINT_PER_GROUP, InLODIndex);
+		// InternalGetOrRequestData(HAIR_MAX_NUM_CURVE_PER_GROUP, HAIR_MAX_NUM_POINT_PER_GROUP, InLODIndex);
 	}
 }
 
@@ -1269,7 +1289,6 @@ void FHairStrandsRestRootResource::InternalAllocate(FRDGBuilder& GraphBuilder, u
 		const FHairStrandsRootBulkData::FHeader::FLOD& LODHeader = BulkData.Header.LODs[InLODIndex];
 		FHairStrandsRootBulkData::FData::FLOD& CPUData = BulkData.Data.LODs[InLODIndex];
 		{
-			
 			InternalCreateVertexBufferRDG_FromHairBulkData<FHairStrandsRootBarycentricFormat>(GraphBuilder, CPUData.RootBarycentricBuffer, BulkData.Header.RootCount, GPUData.RootBarycentricBuffer, ToHairResourceDebugName(HAIRSTRANDS_RESOUCE_NAME(CurveType, Hair.StrandsRestRoot_RootTriangleBarycentricBuffer), ResourceName), OwnerName, EHairResourceUsageType::Static);
 			InternalCreateVertexBufferRDG_FromHairBulkData<FHairStrandsRootToUniqueTriangleIndexFormat>(GraphBuilder, CPUData.RootToUniqueTriangleIndexBuffer, BulkData.Header.RootCount, GPUData.RootToUniqueTriangleIndexBuffer, ToHairResourceDebugName(HAIRSTRANDS_RESOUCE_NAME(CurveType, Hair.StrandsRestRoot_RootToUniqueTriangleIndexBuffer), ResourceName), OwnerName, EHairResourceUsageType::Static);
 			InternalCreateVertexBufferRDG_FromHairBulkData<FHairStrandsUniqueTriangleIndexFormat>(GraphBuilder, CPUData.UniqueTriangleIndexBuffer, LODHeader.UniqueTriangleCount, GPUData.UniqueTriangleIndexBuffer, ToHairResourceDebugName(HAIRSTRANDS_RESOUCE_NAME(CurveType, Hair.StrandsRestRoot_UniqueTriangleIndexBuffer), ResourceName), OwnerName, EHairResourceUsageType::Static);
diff --git a/Engine/Plugins/Runtime/HairStrands/Source/HairStrandsCore/Public/GroomResources.h b/Engine/Plugins/Runtime/HairStrands/Source/HairStrandsCore/Public/GroomResources.h
--- a/Engine/Plugins/Runtime/HairStrands/Source/HairStrandsCore/Public/GroomResources.h
+++ b/Engine/Plugins/Runtime/HairStrands/Source/HairStrandsCore/Public/GroomResources.h
@@ -111,7 +111,7 @@ struct FHairCommonResource : public FRenderResource
 	virtual void InternalAllocate(FRDGBuilder& GraphBuilder) {}
 	virtual void InternalAllocate(FRDGBuilder& GraphBuilder, uint32 InCurveCount, int32 InLODIndex) { InternalAllocate(GraphBuilder); }
 	virtual void InternalRelease() {}
-	virtual bool InternalGetOrRequestData(uint32 InRequestedCurveCount, uint32 InRequestedPointCount, int32 InLODIndex) { return true; }
+	virtual bool InternalGetOrRequestData(uint32 InRequestedCurveCount, uint32 InRequestedPointCount, int32 InLODIndex);
 	virtual bool InternalIsLODDataLoaded(uint32 InRequestedCurveCount, uint32 InRequestedPointCount, int32 InLODIndex) const { return true; }
 	virtual void InternalResetLoadedSize() { }
 
