Add cache warming when cooking groom asset to reduce high DDC latency when doing GetChunk query.

This improves cooking time only at the moment.

diff --git a/Engine/Plugins/Runtime/HairStrands/Source/HairStrandsCore/Private/GroomAsset.cpp b/Engine/Plugins/Runtime/HairStrands/Source/HairStrandsCore/Private/GroomAsset.cpp
--- a/Engine/Plugins/Runtime/HairStrands/Source/HairStrandsCore/Private/GroomAsset.cpp
+++ b/Engine/Plugins/Runtime/HairStrands/Source/HairStrandsCore/Private/GroomAsset.cpp
@@ -1789,9 +1789,9 @@ static void InternalSerializeStrand(FArchive& Ar, UObject* Owner, FHairGroupPlat
 	// from the 'start' (i.e., without offset)
 	if (Ar.IsCooking() && Ar.IsSaving())
 	{
-		{ FHairStreamingRequest R; R.Request(HAIR_MAX_NUM_CURVE_PER_GROUP, HAIR_MAX_NUM_POINT_PER_GROUP, -1/*LODIndex*/, StrandData.BulkData,               true /*bWait*/, true /*bFillBulkdata*/, Owner->GetFName()); }
-		{ FHairStreamingRequest R; R.Request(HAIR_MAX_NUM_CURVE_PER_GROUP, HAIR_MAX_NUM_POINT_PER_GROUP, -1/*LODIndex*/, StrandData.InterpolationBulkData,  true /*bWait*/, true /*bFillBulkdata*/, Owner->GetFName()); }
-		{ FHairStreamingRequest R; R.Request(HAIR_MAX_NUM_CURVE_PER_GROUP, HAIR_MAX_NUM_POINT_PER_GROUP, -1/*LODIndex*/, StrandData.ClusterCullingBulkData, true /*bWait*/, true /*bFillBulkdata*/, Owner->GetFName()); }
+		{ FHairStreamingRequest R; R.Request(HAIR_MAX_NUM_CURVE_PER_GROUP, HAIR_MAX_NUM_POINT_PER_GROUP, -1/*LODIndex*/, StrandData.BulkData,               true /*bWait*/, true /*bFillBulkdata*/, true /*bWarmCache*/, Owner->GetFName()); }
+		{ FHairStreamingRequest R; R.Request(HAIR_MAX_NUM_CURVE_PER_GROUP, HAIR_MAX_NUM_POINT_PER_GROUP, -1/*LODIndex*/, StrandData.InterpolationBulkData,  true /*bWait*/, true /*bFillBulkdata*/, true /*bWarmCache*/, Owner->GetFName()); }
+		{ FHairStreamingRequest R; R.Request(HAIR_MAX_NUM_CURVE_PER_GROUP, HAIR_MAX_NUM_POINT_PER_GROUP, -1/*LODIndex*/, StrandData.ClusterCullingBulkData, true /*bWait*/, true /*bFillBulkdata*/, true /*bWarmCache*/, Owner->GetFName()); }
 	}
 
 	if (!Ar.IsCooking() || !StrandData.bIsCookedOut)
diff --git a/Engine/Plugins/Runtime/HairStrands/Source/HairStrandsCore/Private/GroomBindingAsset.cpp b/Engine/Plugins/Runtime/HairStrands/Source/HairStrandsCore/Private/GroomBindingAsset.cpp
--- a/Engine/Plugins/Runtime/HairStrands/Source/HairStrandsCore/Private/GroomBindingAsset.cpp
+++ b/Engine/Plugins/Runtime/HairStrands/Source/HairStrandsCore/Private/GroomBindingAsset.cpp
@@ -55,7 +55,7 @@ static void InternalSerializeStrands(FArchive& Ar, UObject* Owner, FHairStrandsR
 		{
 			for (int32 LODIndex = 0, LODCount = Data.GetLODCount(); LODIndex < LODCount; ++LODIndex)
 			{
-				FHairStreamingRequest R; R.Request(HAIR_MAX_NUM_CURVE_PER_GROUP, HAIR_MAX_NUM_POINT_PER_GROUP, LODIndex, Data, true /*bWait*/, true /*bFillBulkdata*/, Owner->GetFName());
+				FHairStreamingRequest R; R.Request(HAIR_MAX_NUM_CURVE_PER_GROUP, HAIR_MAX_NUM_POINT_PER_GROUP, LODIndex, Data, true /*bWait*/, true /*bFillBulkdata*/, true /*bWarmCache*/, Owner->GetFName());
 			}
 		}
 
diff --git a/Engine/Plugins/Runtime/HairStrands/Source/HairStrandsCore/Private/GroomResources.cpp b/Engine/Plugins/Runtime/HairStrands/Source/HairStrandsCore/Private/GroomResources.cpp
--- a/Engine/Plugins/Runtime/HairStrands/Source/HairStrandsCore/Private/GroomResources.cpp
+++ b/Engine/Plugins/Runtime/HairStrands/Source/HairStrandsCore/Private/GroomResources.cpp
@@ -1003,7 +1003,7 @@ bool FHairStrandsRestResource::InternalGetOrRequestData(uint32 InRequestedCurveC
 {
 	if (StreamingRequest.IsNone())
 	{
-		StreamingRequest.Request(InRequestedCurveCount, InRequestedPointCount, InLODIndex, BulkData, false, false, OwnerName);
+		StreamingRequest.Request(InRequestedCurveCount, InRequestedPointCount, InLODIndex, BulkData, false, false, false, OwnerName);
 	}
 	return StreamingRequest.IsCompleted();
 }
@@ -1180,7 +1180,7 @@ bool FHairStrandsClusterCullingResource::InternalGetOrRequestData(uint32 InReque
 {
 	if (StreamingRequest.IsNone())
 	{
-		StreamingRequest.Request(InRequestedCurveCount, InRequestedPointCount, InLODIndex, BulkData, true, true, OwnerName);
+		StreamingRequest.Request(InRequestedCurveCount, InRequestedPointCount, InLODIndex, BulkData, true, true, false, OwnerName);
 	}
 	return StreamingRequest.IsCompleted();
 }
@@ -1258,7 +1258,7 @@ bool FHairStrandsRestRootResource::InternalGetOrRequestData(uint32 InRequestedCu
 {
 	if (StreamingRequest.IsNone() && InLODIndex >= 0)
 	{
-		StreamingRequest.Request(InRequestedCurveCount, InRequestedPointCount, InLODIndex, BulkData, false, false, OwnerName);
+		StreamingRequest.Request(InRequestedCurveCount, InRequestedPointCount, InLODIndex, BulkData, false, false, false, OwnerName);
 	}
 	return StreamingRequest.IsCompleted();
 }
@@ -1441,7 +1441,7 @@ bool FHairStrandsInterpolationResource::InternalGetOrRequestData(uint32 InReques
 {
 	if (StreamingRequest.IsNone())
 	{
-		StreamingRequest.Request(InRequestedCurveCount, InRequestedPointCount, InLODIndex, BulkData, false, false, OwnerName);
+		StreamingRequest.Request(InRequestedCurveCount, InRequestedPointCount, InLODIndex, BulkData, false, false, false, OwnerName);
 	}
 	return StreamingRequest.IsCompleted();
 }
diff --git a/Engine/Plugins/Runtime/HairStrands/Source/HairStrandsCore/Private/HairStrandsDatas.cpp b/Engine/Plugins/Runtime/HairStrands/Source/HairStrandsCore/Private/HairStrandsDatas.cpp
--- a/Engine/Plugins/Runtime/HairStrands/Source/HairStrandsCore/Private/HairStrandsDatas.cpp
+++ b/Engine/Plugins/Runtime/HairStrands/Source/HairStrandsCore/Private/HairStrandsDatas.cpp
@@ -319,7 +319,7 @@ bool FHairStreamingRequest::IsCompleted()
 // Request fullfil 2 use cases:
 // * Load IO/DDC data and upload them to GPU
 // * Load DDC data and store them into bulkdata for serialization
-void FHairStreamingRequest::Request(uint32 InRequestedCurveCount, uint32 InRequestedPointCount, int32 InLODIndex, FHairStrandsBulkCommon& In, bool bWait, bool bFillBulkdata, const FName& InOwnerName)
+void FHairStreamingRequest::Request(uint32 InRequestedCurveCount, uint32 InRequestedPointCount, int32 InLODIndex, FHairStrandsBulkCommon& In, bool bWait, bool bFillBulkdata, bool bWarmCache, const FName& InOwnerName)
 {
 	if (In.GetResourceCount() == 0 || InRequestedCurveCount == 0)
 	{
@@ -352,7 +352,21 @@ void FHairStreamingRequest::Request(uint32 InRequestedCurveCount, uint32 InReque
 	
 		check(DDCRequestOwner == nullptr);
 		DDCRequestOwner = MakeUnique<FRequestOwner>(bWait ? UE::DerivedData::EPriority::Blocking : UE::DerivedData::EPriority::Normal); // <= Move this onto the resource (Buffer, in order to ensure no race condition)
-	
+		
+		// Currently GetChunk request have long latency if not cached locally. 
+		// To reduce cook time, we warming cache by issuing 'Value' request (vs.'Chunk' request)
+		if (bWarmCache)
+		{
+			const ECachePolicy Policy = ECachePolicy::Default | ECachePolicy::SkipData;
+			TArray<FCacheGetValueRequest, TInlineAllocator<16>> WarmRequest;
+			for (const FCacheGetChunkRequest& R : Requests)
+			{				
+				WarmRequest.Add({R.Name, R.Key, Policy, 0});
+			}
+			GetCache().GetValue(WarmRequest, *DDCRequestOwner, [](FCacheGetValueResponse && Response) { check(Response.Status == EStatus::Ok); });
+			DDCRequestOwner->Wait();
+		}
+
 		//FRequestBarrier Barrier(*DDCRequestOwner);	// This is a critical section on the owner. It does not constrain ordering
 		GetCache().GetChunks(Requests, *DDCRequestOwner,
 		[this, &In, bFillBulkdata, InOwnerName](FCacheGetChunkResponse && Response)
diff --git a/Engine/Plugins/Runtime/HairStrands/Source/HairStrandsCore/Public/HairStrandsDatas.h b/Engine/Plugins/Runtime/HairStrands/Source/HairStrandsCore/Public/HairStrandsDatas.h
--- a/Engine/Plugins/Runtime/HairStrands/Source/HairStrandsCore/Public/HairStrandsDatas.h
+++ b/Engine/Plugins/Runtime/HairStrands/Source/HairStrandsCore/Public/HairStrandsDatas.h
@@ -368,7 +368,7 @@ struct FHairStreamingRequest
 		void Release();
 	};
 
-	void Request(uint32 InRequestedCurveCount, uint32 InRequestedPointCount, int32 InLODIndex, FHairStrandsBulkCommon& In, bool bWait=false, bool bFillBulkData=false, const FName& InOwnerName = NAME_None);
+	void Request(uint32 InRequestedCurveCount, uint32 InRequestedPointCount, int32 InLODIndex, FHairStrandsBulkCommon& In, bool bWait=false, bool bFillBulkData=false, bool bWarmCache=false, const FName& InOwnerName = NAME_None);
 	bool IsNone() const;
 	bool IsCompleted();
 
