Fix missing CVar ThreadSafe flag on r.HairStrands.ContinuousLOD.

diff --git a/Engine/Plugins/Runtime/HairStrands/Source/HairStrandsCore/Private/GroomManager.cpp b/Engine/Plugins/Runtime/HairStrands/Source/HairStrandsCore/Private/GroomManager.cpp
--- a/Engine/Plugins/Runtime/HairStrands/Source/HairStrandsCore/Private/GroomManager.cpp
+++ b/Engine/Plugins/Runtime/HairStrands/Source/HairStrandsCore/Private/GroomManager.cpp
@@ -43,7 +43,7 @@ static int32 GHairStrands_InterpolationFrustumCullingEnable = 1;
 static FAutoConsoleVariableRef CVarHairStrands_InterpolationFrustumCullingEnable(TEXT("r.HairStrands.Interoplation.FrustumCulling"), GHairStrands_InterpolationFrustumCullingEnable, TEXT("Swap rendering buffer at the end of frame. This is an experimental toggle. Default:1"));
 
 static int32 GHairStrands_ContinousLOD = 0;
-static FAutoConsoleVariableRef CVarHairStrands_ContinousLOD(TEXT("r.HairStrands.ContinuousLOD"), GHairStrands_ContinousLOD, TEXT("Continuous LODing toggle."));
+static FAutoConsoleVariableRef CVarHairStrands_ContinousLOD(TEXT("r.HairStrands.ContinuousLOD"), GHairStrands_ContinousLOD, TEXT("Continuous LODing toggle."), ECVF_RenderThreadSafe);
 
 EHairBufferSwapType GetHairSwapBufferType()
 {
