Add option on the physics settings to force the Niagara component to be visible.

This allows user to leverage Niagara rendering when using custom solver.

diff --git a/Engine/Plugins/Runtime/HairStrands/Source/HairStrandsCore/Private/GroomAssetPhysics.cpp b/Engine/Plugins/Runtime/HairStrands/Source/HairStrandsCore/Private/GroomAssetPhysics.cpp
--- a/Engine/Plugins/Runtime/HairStrands/Source/HairStrandsCore/Private/GroomAssetPhysics.cpp
+++ b/Engine/Plugins/Runtime/HairStrands/Source/HairStrandsCore/Private/GroomAssetPhysics.cpp
@@ -11,6 +11,7 @@ FHairSolverSettings::FHairSolverSettings()
 	SubSteps = 5;
 	IterationCount = 5;
 	GravityPreloading = 0.0f;
+	bForceVisible = false;
 }
 
 FHairExternalForces::FHairExternalForces()
@@ -88,7 +89,8 @@ bool FHairSolverSettings::operator==(const FHairSolverSettings& A) const
 		EnableSimulation == A.EnableSimulation &&
 		NiagaraSolver == A.NiagaraSolver &&
 		SubSteps == A.SubSteps &&
-		IterationCount == A.IterationCount;
+		IterationCount == A.IterationCount &&
+		bForceVisible == A.bForceVisible;
 }
 
 bool FHairExternalForces::operator==(const FHairExternalForces& A) const
diff --git a/Engine/Plugins/Runtime/HairStrands/Source/HairStrandsCore/Private/GroomComponent.cpp b/Engine/Plugins/Runtime/HairStrands/Source/HairStrandsCore/Private/GroomComponent.cpp
--- a/Engine/Plugins/Runtime/HairStrands/Source/HairStrandsCore/Private/GroomComponent.cpp
+++ b/Engine/Plugins/Runtime/HairStrands/Source/HairStrandsCore/Private/GroomComponent.cpp
@@ -1342,7 +1342,7 @@ void UGroomComponent::CreateHairSimulation(const int32 GroupIndex, const int32 L
 					(GroomAsset->HairGroupsPhysics[GroupIndex].SolverSettings.NiagaraSolver == EGroomNiagaraSolvers::AngularSprings) ? ToRawPtr(AngularSpringsSystem) :
 					(GroomAsset->HairGroupsPhysics[GroupIndex].SolverSettings.NiagaraSolver == EGroomNiagaraSolvers::CosseratRods) ? ToRawPtr(CosseratRodsSystem) :
 					GroomAsset->HairGroupsPhysics[GroupIndex].SolverSettings.CustomSystem.LoadSynchronous();
-				NiagaraComponent->SetVisibleFlag(SimulationSettings.SimulationSetup.bDebugSimulation);
+				NiagaraComponent->SetVisibleFlag(SimulationSettings.SimulationSetup.bDebugSimulation || GroomAsset->HairGroupsPhysics[GroupIndex].SolverSettings.bForceVisible);
 				NiagaraComponent->SetAsset(NiagaraAsset);
 				NiagaraComponent->ReinitializeSystem();
 			}
diff --git a/Engine/Plugins/Runtime/HairStrands/Source/HairStrandsCore/Public/GroomAssetPhysics.h b/Engine/Plugins/Runtime/HairStrands/Source/HairStrandsCore/Public/GroomAssetPhysics.h
--- a/Engine/Plugins/Runtime/HairStrands/Source/HairStrandsCore/Public/GroomAssetPhysics.h
+++ b/Engine/Plugins/Runtime/HairStrands/Source/HairStrandsCore/Public/GroomAssetPhysics.h
@@ -72,6 +72,10 @@ struct HAIRSTRANDSCORE_API FHairSolverSettings
 	UPROPERTY(EditAnywhere, Category = "SolverSettings", meta = (ToolTip = "Number of iterations to solve the constraints with the xpbd solver"))
 	int32 IterationCount;
 
+	/** Force the Niagara solver component to be visible  */
+	UPROPERTY(EditAnywhere, Category = "SolverSettings", meta = (ToolTip = "Force the Niagara solver component to be visible"))
+	bool bForceVisible;
+
 	bool operator==(const FHairSolverSettings& A) const;
 };
 
