Reduce hair divergence from MH branch.

diff --git a/Engine/Plugins/Runtime/HairStrands/Source/HairStrandsCore/Private/GroomBuilder.cpp b/Engine/Plugins/Runtime/HairStrands/Source/HairStrandsCore/Private/GroomBuilder.cpp
--- a/Engine/Plugins/Runtime/HairStrands/Source/HairStrandsCore/Private/GroomBuilder.cpp
+++ b/Engine/Plugins/Runtime/HairStrands/Source/HairStrandsCore/Private/GroomBuilder.cpp
@@ -2008,7 +2008,8 @@ void FGroomBuilder::BuildData(
 	const FHairGroupsInterpolation& InSettings, 
 	FHairGroupInfo& OutGroupInfo, 
 	FHairStrandsDatas& OutRen,
-	FHairStrandsDatas& OutSim)
+	FHairStrandsDatas& OutSim,
+	bool bAllowCurveReordering)
 {
 	TRACE_CPUPROFILER_EVENT_SCOPE(FGroomBuilder::BuildData);
 
@@ -2017,7 +2018,7 @@ void FGroomBuilder::BuildData(
 	const float CurveDecimation		= FMath::Clamp(InSettings.DecimationSettings.CurveDecimation, 0.f, 1.f);
 	const float VertexDecimation	= FMath::Clamp(InSettings.DecimationSettings.VertexDecimation, 0.f, 1.f);
 	const float HairToGuideDensity	= FMath::Clamp(InSettings.InterpolationSettings.HairToGuideDensity, 0.f, 1.f);	
-	const bool bCurveReordering 	= true;
+	const bool bCurveReordering 	= bAllowCurveReordering;
 
 	{
 		OutGroupInfo.GroupID = InHairDescriptionGroup.Info.GroupID;
diff --git a/Engine/Plugins/Runtime/HairStrands/Source/HairStrandsCore/Private/GroomRBFDeformer.cpp b/Engine/Plugins/Runtime/HairStrands/Source/HairStrandsCore/Private/GroomRBFDeformer.cpp
--- a/Engine/Plugins/Runtime/HairStrands/Source/HairStrandsCore/Private/GroomRBFDeformer.cpp
+++ b/Engine/Plugins/Runtime/HairStrands/Source/HairStrandsCore/Private/GroomRBFDeformer.cpp
@@ -8,6 +8,8 @@
 #include "Rendering/SkeletalMeshRenderData.h"
 #include "MeshAttributes.h"
 #include "Engine/StaticMesh.h"
+#include "Engine/TextureDefines.h"
+#include "Engine/Texture.h"
 #include "GroomBindingBuilder.h"
 
 // HairStrandsMesh.usf
@@ -259,14 +261,53 @@ struct FRBFDeformedPositions
 {
 	TArray<FVector3f> RenderStrands;
 	TArray<FVector3f> GuideStrands;
+
+	// Trimmed data
+	TArray<bool> bIsRenderCurveValid;
+	TArray<bool> bIsRenderVertexValid;
+	bool bHasTrimmedData = false;
 };
 
 #if WITH_EDITORONLY_DATA
 static void ApplyDeformationToGroom(const TArray<FRBFDeformedPositions>& DeformedPositions, UGroomAsset* GroomAsset)
 {
+	// Notes:
+	// * This function applies the deformed position onto the hair description.
+	// * In addition, the groom can contain vertex/curves which needs to be trimmed (e.g., which are masked or cut.
+	//   For handling this, the following function will first applied the deformed vertices and then remap all
+	//   curves/vertices so that hair description contains only valid curves/vertices
+	//
+	// /!\ Only rendering strands can be trimmed. Guide always remains unchanged.
+
+	// Check if any group has trimmed data
+	bool bHasTrimmedData = false;
+	for (const FRBFDeformedPositions& Group : DeformedPositions)
+	{
+		if (Group.bHasTrimmedData)
+		{
+			bHasTrimmedData = true;
+			break;
+		}
+	}
+
+	// If there are some trimmed data, compute the curves & vertices remapping
+	int32 CurrentStrandID = 0;
+	int32 CurrentVertexID = 0;
+
+	uint32 TotalNumStrands = 0;
+	uint32 TotalNumVertices = 0;
+
+	uint32 TotalValidNumStrands = 0;
+	uint32 TotalValidNumVertices = 0;
+
 	// The deformation must be stored in the HairDescription to rebuild the hair data when the groom is loaded
 	FHairDescription HairDescription = GroomAsset->GetHairDescription();
 
+	TotalNumVertices = HairDescription.GetNumVertices();
+	TotalNumStrands = HairDescription.GetNumStrands();
+	TotalValidNumVertices = TotalNumVertices;
+	TotalValidNumStrands = TotalNumStrands;
+
 	// Strands attributes as inputs
 	TStrandAttributesConstRef<int> StrandNumVertices = HairDescription.StrandAttributes().GetAttributesRef<int>(HairAttribute::Strand::VertexCount);
 	TStrandAttributesConstRef<int> StrandGuides = HairDescription.StrandAttributes().GetAttributesRef<int>(HairAttribute::Strand::Guide);
@@ -296,12 +337,14 @@ static void ApplyDeformationToGroom(const TArray<FRBFDeformedPositions>& Deforme
 	const int32 GroomNumVertices = HairDescription.GetNumVertices();
 	TArray<FVector3f> FlattenedDeformedPositions;
 	FlattenedDeformedPositions.Reserve(GroomNumVertices);
+	TArray<bool> FlattenedIsValid;
+	FlattenedIsValid.Reserve(GroomNumVertices);
 
 	// Mapping of GroupID to GroupIndex to preserver ordering
 	TMap<int32, int32> GroupIDToGroupIndex;
 
-	const int32 NumStrands = HairDescription.GetNumStrands();
-	for (int32 StrandIndex = 0; StrandIndex < NumStrands; ++StrandIndex)
+	const int32 GroomNumStrands = HairDescription.GetNumStrands();
+	for (int32 StrandIndex = 0; StrandIndex < GroomNumStrands; ++StrandIndex)
 	{
 		FStrandID StrandID(StrandIndex);
 
@@ -331,6 +374,8 @@ static void ApplyDeformationToGroom(const TArray<FRBFDeformedPositions>& Deforme
 			{
 				if (GroupInfo.CurrentGuideVertexIndex < DeformedPositions[GroupIndex].GuideStrands.Num())
 				{
+					// Guide vertex are never trimmed
+					FlattenedIsValid.Add(true);
 					FlattenedDeformedPositions.Add(DeformedPositions[GroupIndex].GuideStrands[GroupInfo.CurrentGuideVertexIndex++]);
 				}
 			}
@@ -342,6 +387,10 @@ static void ApplyDeformationToGroom(const TArray<FRBFDeformedPositions>& Deforme
 			{
 				if (GroupInfo.CurrentRenderVertexIndex < DeformedPositions[GroupIndex].RenderStrands.Num())
 				{
+					FVertexID VertexID(FlattenedDeformedPositions.Num());
+
+					const bool bIsValid = bHasTrimmedData ? DeformedPositions[GroupIndex].bIsRenderVertexValid[GroupInfo.CurrentRenderVertexIndex] : true;
+					FlattenedIsValid.Add(bIsValid);
 					FlattenedDeformedPositions.Add(DeformedPositions[GroupIndex].RenderStrands[GroupInfo.CurrentRenderVertexIndex++]);
 				}
 			}
@@ -356,6 +405,21 @@ static void ApplyDeformationToGroom(const TArray<FRBFDeformedPositions>& Deforme
 		VertexPositions[VertexID] = FlattenedDeformedPositions[VertexIndex];
 	}
 
+	// Apply hair description trimming if needed
+	if (bHasTrimmedData)
+	{
+		TVertexAttributesRef<float> RWVertexWidth = HairDescription.VertexAttributes().GetAttributesRef<float>(HairAttribute::Vertex::Width);
+		for (int32 VertexIndex = 0; VertexIndex < GroomNumVertices; ++VertexIndex)
+		{
+			if (!FlattenedIsValid[VertexIndex])
+			{
+				FVertexID VertexID(VertexIndex);
+				RWVertexWidth[VertexID] = 0.f;
+			}
+		}
+	}
+
+	// Regenerate data based on the new/updated hair description
 	{
 		FHairDescriptionGroups HairDescriptionGroups;
 		FGroomBuilder::BuildHairDescriptionGroups(HairDescription, HairDescriptionGroups);
@@ -465,7 +529,59 @@ namespace GroomDerivedDataCacheUtils
 }
 #endif // #if WITH_EDITORONLY_DATA
 
-void FGroomRBFDeformer::GetRBFDeformedGroomAsset(const UGroomAsset* InGroomAsset, const UGroomBindingAsset* BindingAsset, UGroomAsset* OutGroomAsset)
+template<typename T>
+float SampleTexture(const FUintPoint& InCoord, const uint8 ComponentIndex, const uint8 NumComponent, const FUintPoint& Resolution, const T* InData)
+{
+	const uint32 Index = ComponentIndex + InCoord.X * NumComponent + InCoord.Y * Resolution.X * NumComponent;
+	return InData[Index];
+}
+
+template<typename T>
+float SampleTexture(const FVector2f& InUV, const uint8 InComponentIndex, const uint8 InNumComponent, const FUintPoint& InResolution, const uint8* InRawData, float InNormalizationScale)
+{
+	const T* TypedData = (const T*)InRawData;
+
+	const FVector2f  P(float(InResolution.X) * InUV.X, float(InResolution.Y) * InUV.Y);
+	const FUintPoint P00(FMath::Floor(P.X), FMath::Floor(P.Y));
+	const FUintPoint P10 = P00 + FUintPoint(1, 0);
+	const FUintPoint P11 = P00 + FUintPoint(1, 1);
+	const FUintPoint P01 = P00 + FUintPoint(0, 1);
+
+	const T T00 = SampleTexture<T>(P00, InComponentIndex, InNumComponent, InResolution, TypedData);
+	const T T10 = SampleTexture<T>(P10, InComponentIndex, InNumComponent, InResolution, TypedData);
+	const T T11 = SampleTexture<T>(P11, InComponentIndex, InNumComponent, InResolution, TypedData);
+	const T T01 = SampleTexture<T>(P11, InComponentIndex, InNumComponent, InResolution, TypedData);
+
+	const float V00 = float(T00) * InNormalizationScale;
+	const float V10 = float(T10) * InNormalizationScale;
+	const float V11 = float(T11) * InNormalizationScale;
+	const float V01 = float(T01) * InNormalizationScale;
+
+	const FVector2f S(FMath::Frac(P.X), FMath::Frac(P.Y));
+	return FMath::Lerp(
+		FMath::Lerp(V00, V10, S.X),
+		FMath::Lerp(V01, V11, S.X),
+		S.Y);
+}
+
+static float SampleMaskTexture(const FVector2f& InUV, const FUintPoint& InResolution, const uint8* InData, ETextureSourceFormat Format)
+{
+	switch(Format)
+	{
+		case TSF_G8:		return SampleTexture<uint8>		(InUV, 0, 1, InResolution, InData, 1.f/255.f);
+		case TSF_BGRA8:		return SampleTexture<uint8>		(InUV, 3, 4, InResolution, InData, 1.f/255.f);
+		case TSF_RGBA16:	return SampleTexture<uint16>	(InUV, 0, 4, InResolution, InData, 1.f/65536.f);
+		case TSF_RGBA16F:	return SampleTexture<FFloat16>	(InUV, 0, 4, InResolution, InData, 1.f);
+		case TSF_G16:		return SampleTexture<uint16>	(InUV, 0, 1, InResolution, InData, 1.f/65536.f);
+		case TSF_RGBA32F:	return SampleTexture<float>		(InUV, 0, 4, InResolution, InData, 1.f);
+		case TSF_R16F:		return SampleTexture<FFloat16>	(InUV, 0, 1, InResolution, InData, 1.f);
+		case TSF_R32F:		return SampleTexture<float>		(InUV, 0, 1, InResolution, InData, 1.f);
+		default: check(false);
+	}
+	return 1.0f;
+}
+
+void FGroomRBFDeformer::GetRBFDeformedGroomAsset(const UGroomAsset* InGroomAsset, const UGroomBindingAsset* BindingAsset, FTextureSource* MaskTextureSource, const float MaskScale, UGroomAsset* OutGroomAsset)
 {
 #if WITH_EDITORONLY_DATA
 	if (InGroomAsset && BindingAsset && BindingAsset->TargetSkeletalMesh && BindingAsset->SourceSkeletalMesh)
@@ -524,7 +640,9 @@ void FGroomRBFDeformer::GetRBFDeformedGroomAsset(const UGroomAsset* InGroomAsset
 			FHairStrandsDatas StrandsData;
 			FHairStrandsDatas GuidesData;
 			FHairGroupInfo DummyInfo;
-			FGroomBuilder::BuildData(Group, InGroomAsset->HairGroupsInterpolation[GroupIndex], DummyInfo, StrandsData, GuidesData);
+
+			// Disable explicitely curve reordering as we want to preserve vertex mapping to write deformed postions back to the hair description, after RBF deformation
+			FGroomBuilder::BuildData(Group, InGroomAsset->HairGroupsInterpolation[GroupIndex], DummyInfo, StrandsData, GuidesData, false /*bAllowCurveReordering*/);
 
 			TArray<uint32> SimRootDataPointToCurveBuffer;
 			TArray<uint32> RenRootDataPointToCurveBuffer;
@@ -565,6 +683,91 @@ void FGroomRBFDeformer::GetRBFDeformedGroomAsset(const UGroomAsset* InGroomAsset
 					RenRootDataPointToCurveBuffer,
 					RenRootData.MeshProjectionLODs[MeshLODIndex]);
 			}
+
+			// Trim strands based on mask texture
+			DeformedPositions[GroupIndex].bIsRenderVertexValid.Init(true, StrandsData.GetNumPoints());
+			DeformedPositions[GroupIndex].bIsRenderCurveValid.Init(true, StrandsData.GetNumCurves());
+			DeformedPositions[GroupIndex].bHasTrimmedData = false;
+			if (MaskTextureSource)
+			{
+				check(MaskTextureSource->GetNumBlocks() == 1);
+
+				const ETextureSourceFormat Format = MaskTextureSource->GetFormat();
+				const FUintPoint Resolution(MaskTextureSource->GetSizeX(), MaskTextureSource->GetSizeY());
+
+				// Access texture CPU data
+				FTextureSourceBlock Block;
+				MaskTextureSource->GetBlock(0, Block);
+				const uint8* DataInBytes = MaskTextureSource->LockMipReadOnly(0 /*MipIndex*/);
+
+				const uint32 CurveCount = StrandsData.GetNumCurves();
+				for (uint32 CurveIndex =0; CurveIndex < CurveCount; ++CurveIndex)
+				{
+					const uint32 CurveOffset = StrandsData.StrandsCurves.CurvesOffset[CurveIndex];
+					const uint16 CurveNumVertices = StrandsData.StrandsCurves.CurvesCount[CurveIndex];
+
+					FVector2f RootUV = StrandsData.StrandsCurves.CurvesRootUV[CurveIndex];
+					RootUV.Y = FMath::Clamp(1.f - RootUV.Y,0.f, 1.0f);
+					const float Maskhreshold = SampleMaskTexture(RootUV, Resolution, DataInBytes, Format);
+					const float CoordUThreshold = FMath::Lerp(1.f, Maskhreshold, MaskScale);
+
+					DeformedPositions[GroupIndex].bIsRenderCurveValid[CurveIndex] = true;
+					for (uint16 VertexIndex = 0; VertexIndex < CurveNumVertices; ++VertexIndex)
+					{
+						const uint32 PointGlobalIndex0 = VertexIndex + CurveOffset;
+						const uint32 PointGlobalIndex1 = VertexIndex + CurveOffset + 1;
+
+						DeformedPositions[GroupIndex].bIsRenderVertexValid[PointGlobalIndex0] = true;
+
+						// Trimmed valid
+						if (CoordUThreshold < 1.f)
+						{
+							const float CoordU0 = StrandsData.StrandsPoints.PointsCoordU[PointGlobalIndex0];
+							if (CoordUThreshold <= 0.f)
+							{
+								DeformedPositions[GroupIndex].bIsRenderVertexValid[PointGlobalIndex0] = false;
+								DeformedPositions[GroupIndex].bHasTrimmedData = true;
+							}
+							else if (VertexIndex + 1 < CurveNumVertices)
+							{
+								DeformedPositions[GroupIndex].bIsRenderVertexValid[PointGlobalIndex1] = true;
+
+								// Interpolate position or trim vertex
+								const float CoordU1 = StrandsData.StrandsPoints.PointsCoordU[PointGlobalIndex1];
+								if (CoordU0 <= CoordUThreshold && CoordU1 > CoordUThreshold)
+								{
+									const float S = FMath::Clamp((CoordUThreshold - CoordU0) / FMath::Max(0.0001f, CoordU1 - CoordU0), 0.f, 1.f);
+									StrandsData.StrandsPoints.PointsPosition[PointGlobalIndex1] =
+										FMath::Lerp(
+										StrandsData.StrandsPoints.PointsPosition[PointGlobalIndex0],
+										StrandsData.StrandsPoints.PointsPosition[PointGlobalIndex1],
+										S);
+
+									DeformedPositions[GroupIndex].bHasTrimmedData = true;
+								}
+								else if (CoordU0 > CoordUThreshold)
+								{
+									DeformedPositions[GroupIndex].bIsRenderVertexValid[PointGlobalIndex0] = false;
+									DeformedPositions[GroupIndex].bIsRenderVertexValid[PointGlobalIndex1] = false;
+								}
+							}
+							else if (CoordU0 > CoordUThreshold)
+							{
+								DeformedPositions[GroupIndex].bIsRenderVertexValid[PointGlobalIndex0] = false;
+								DeformedPositions[GroupIndex].bHasTrimmedData = true;
+							}
+
+							// Mark the entire curve as trimmed if the first or second vertex are trimmed, because a curve needs to have at least two valid points
+							if (!DeformedPositions[GroupIndex].bIsRenderVertexValid[PointGlobalIndex0] && (VertexIndex == 0 || VertexIndex == 1))
+							{
+								DeformedPositions[GroupIndex].bIsRenderCurveValid[CurveIndex] = false;
+							}
+						}
+					}
+				}
+
+				MaskTextureSource->UnlockMip(0);
+			}
 		}
 
 		// Apply changes onto cards and meshes (OutGroomASset already contain duplicated mesh asset
diff --git a/Engine/Plugins/Runtime/HairStrands/Source/HairStrandsCore/Public/GroomAssetPhysics.h b/Engine/Plugins/Runtime/HairStrands/Source/HairStrandsCore/Public/GroomAssetPhysics.h
--- a/Engine/Plugins/Runtime/HairStrands/Source/HairStrandsCore/Public/GroomAssetPhysics.h
+++ b/Engine/Plugins/Runtime/HairStrands/Source/HairStrandsCore/Public/GroomAssetPhysics.h
@@ -35,7 +35,7 @@ enum class EGroomStrandsSize : uint8
 UENUM(BlueprintType)
 enum class EGroomInterpolationType : uint8
 {
-	None = 0 UMETA(Hidden),
+	None = 0 UMETA(DisplayName = "No skinning"),
 	RigidTransform = 0x02 UMETA(DisplayName = "Rigid Transform"),
 	OffsetTransform = 0x04 UMETA(DisplayName = "Offset Transform"),
 	SmoothTransform = 0x08 UMETA(DisplayName = "Smooth Transform")
diff --git a/Engine/Plugins/Runtime/HairStrands/Source/HairStrandsCore/Public/GroomBuilder.h b/Engine/Plugins/Runtime/HairStrands/Source/HairStrandsCore/Public/GroomBuilder.h
--- a/Engine/Plugins/Runtime/HairStrands/Source/HairStrandsCore/Public/GroomBuilder.h
+++ b/Engine/Plugins/Runtime/HairStrands/Source/HairStrandsCore/Public/GroomBuilder.h
@@ -67,7 +67,8 @@ struct HAIRSTRANDSCORE_API FGroomBuilder
 		const FHairGroupsInterpolation& InSettings,
 		FHairGroupInfo& OutGroupInfo,
 		FHairStrandsDatas& OutStrands,
-		FHairStrandsDatas& OutGuides);
+		FHairStrandsDatas& OutGuides,
+		bool bAllowCurveReordering=true);
 
 	// 2.b Build FHairStrandsDatas for Strands or Guides. 
 	// This version:
diff --git a/Engine/Plugins/Runtime/HairStrands/Source/HairStrandsCore/Public/GroomRBFDeformer.h b/Engine/Plugins/Runtime/HairStrands/Source/HairStrandsCore/Public/GroomRBFDeformer.h
--- a/Engine/Plugins/Runtime/HairStrands/Source/HairStrandsCore/Public/GroomRBFDeformer.h
+++ b/Engine/Plugins/Runtime/HairStrands/Source/HairStrandsCore/Public/GroomRBFDeformer.h
@@ -5,9 +5,11 @@
 #include "CoreMinimal.h"
 
 class UGroomAsset;
+class UGroomBindingAsset;
+struct FTextureSource;
 
 struct HAIRSTRANDSCORE_API FGroomRBFDeformer
 {
 	// Return a new GroomAsset with the RBF deformation from the BindingAsset baked into it
-	static void GetRBFDeformedGroomAsset(const UGroomAsset* GroomAsset, const class UGroomBindingAsset* BindingAsset, UGroomAsset* OutGroomAsset);
+	void GetRBFDeformedGroomAsset(const UGroomAsset* InGroomAsset, const UGroomBindingAsset* BindingAsset, FTextureSource* MaskTextureSource, const float MaskScale, UGroomAsset* OutGroomAsset);
 };
\ No newline at end of file
