Add cache warming for all strands data at load time.

diff --git a/Engine/Plugins/Runtime/HairStrands/Source/HairStrandsCore/Private/GroomAsset.cpp b/Engine/Plugins/Runtime/HairStrands/Source/HairStrandsCore/Private/GroomAsset.cpp
--- a/Engine/Plugins/Runtime/HairStrands/Source/HairStrandsCore/Private/GroomAsset.cpp
+++ b/Engine/Plugins/Runtime/HairStrands/Source/HairStrandsCore/Private/GroomAsset.cpp
@@ -1787,7 +1787,8 @@ static void InternalSerializeStrand(FArchive& Ar, UObject* Owner, FHairGroupPlat
 	// When cooking data, force loading of *all* bulk data prior to saving them
 	// Note: bFillBulkdata is true for filling in the bulkdata container prior to serialization. This also forces the resources loading 
 	// from the 'start' (i.e., without offset)
-	if (Ar.IsCooking() && Ar.IsSaving())
+	const bool bFillInBulkDataForCooking = Ar.IsCooking() && Ar.IsSaving();
+	if (bFillInBulkDataForCooking)
 	{
 		{ FHairStreamingRequest R; R.Request(HAIR_MAX_NUM_CURVE_PER_GROUP, HAIR_MAX_NUM_POINT_PER_GROUP, -1/*LODIndex*/, StrandData.BulkData,               true /*bWait*/, true /*bFillBulkdata*/, true /*bWarmCache*/, Owner->GetFName()); }
 		{ FHairStreamingRequest R; R.Request(HAIR_MAX_NUM_CURVE_PER_GROUP, HAIR_MAX_NUM_POINT_PER_GROUP, -1/*LODIndex*/, StrandData.InterpolationBulkData,  true /*bWait*/, true /*bFillBulkdata*/, true /*bWarmCache*/, Owner->GetFName()); }
@@ -1805,6 +1806,17 @@ static void InternalSerializeStrand(FArchive& Ar, UObject* Owner, FHairGroupPlat
 			if (bHeader){ StrandData.ClusterCullingBulkData.SerializeHeader(Ar, Owner); }
 			if (bData)	{ StrandData.ClusterCullingBulkData.SerializeData(Ar, Owner); }
 		}
+
+		#if WITH_EDITORONLY_DATA
+		// Pre-warm DDC cache
+		const bool bPreWarmCache = IsLoading() && bHeader && !bData;
+		if (bPreWarmCache)
+		{
+			{ FHairStreamingRequest R; R.WarmCache(HAIR_MAX_NUM_CURVE_PER_GROUP, HAIR_MAX_NUM_POINT_PER_GROUP, -1/*LODIndex*/, StrandData.BulkData); }
+			{ FHairStreamingRequest R; R.WarmCache(HAIR_MAX_NUM_CURVE_PER_GROUP, HAIR_MAX_NUM_POINT_PER_GROUP, -1/*LODIndex*/, StrandData.InterpolationBulkData); }
+			{ FHairStreamingRequest R; R.WarmCache(HAIR_MAX_NUM_CURVE_PER_GROUP, HAIR_MAX_NUM_POINT_PER_GROUP, -1/*LODIndex*/, StrandData.ClusterCullingBulkData); }
+		}
+		#endif
 	}
 	else
 	{
diff --git a/Engine/Plugins/Runtime/HairStrands/Source/HairStrandsCore/Private/GroomBindingAsset.cpp b/Engine/Plugins/Runtime/HairStrands/Source/HairStrandsCore/Private/GroomBindingAsset.cpp
--- a/Engine/Plugins/Runtime/HairStrands/Source/HairStrandsCore/Private/GroomBindingAsset.cpp
+++ b/Engine/Plugins/Runtime/HairStrands/Source/HairStrandsCore/Private/GroomBindingAsset.cpp
@@ -67,6 +67,18 @@ static void InternalSerializeStrands(FArchive& Ar, UObject* Owner, FHairStrandsR
 				Data.SerializeData(Ar, Owner, LODIndex);
 			}
 		}
+
+		// Pre-warm DDC cache
+		#if WITH_EDITORONLY_DATA
+		const bool bPreWarmCache = IsLoading() && bHeader && !bData;
+		if (bPreWarmCache)
+		{
+			for (int32 LODIndex = 0, LODCount = Data.GetLODCount(); LODIndex < LODCount; ++LODIndex)
+			{
+				FHairStreamingRequest R; R.WarmCache(HAIR_MAX_NUM_CURVE_PER_GROUP, HAIR_MAX_NUM_POINT_PER_GROUP, LODIndex, Data);
+			}
+		}
+		#endif
 	}
 }
 
diff --git a/Engine/Plugins/Runtime/HairStrands/Source/HairStrandsCore/Private/HairStrandsDatas.cpp b/Engine/Plugins/Runtime/HairStrands/Source/HairStrandsCore/Private/HairStrandsDatas.cpp
--- a/Engine/Plugins/Runtime/HairStrands/Source/HairStrandsCore/Private/HairStrandsDatas.cpp
+++ b/Engine/Plugins/Runtime/HairStrands/Source/HairStrandsCore/Private/HairStrandsDatas.cpp
@@ -316,6 +316,47 @@ bool FHairStreamingRequest::IsCompleted()
 	return true; 
 #endif
 }
+#if WITH_EDITORONLY_DATA
+static void RequestWarmCache(UE::DerivedData::FRequestOwner* RequestOwner, const TArray<UE::DerivedData::FCacheGetChunkRequest>& Requests)
+{
+	using namespace UE::DerivedData;
+
+	// Currently GetChunk request have long latency if not cached locally. 
+	// To reduce cook time, we warming cache by issuing 'Value' request (vs.'Chunk' request)
+	const ECachePolicy Policy = ECachePolicy::Default | ECachePolicy::SkipData;
+	TArray<FCacheGetValueRequest, TInlineAllocator<16>> WarmRequest;
+	for (const FCacheGetChunkRequest& R : Requests)
+	{				
+		WarmRequest.Add({R.Name, R.Key, Policy, 0});
+	}
+	GetCache().GetValue(WarmRequest, *RequestOwner, [](FCacheGetValueResponse && Response) { check(Response.Status == EStatus::Ok); });
+	RequestOwner->Wait();
+}
+
+void FHairStreamingRequest::WarmCache(uint32 InRequestedCurveCount, uint32 InRequestedPointCount, int32 InLODIndex, FHairStrandsBulkCommon& In)
+{
+	if (In.GetResourceCount() == 0 || InRequestedCurveCount == 0)
+	{
+		CurveCount = 0;
+		PointCount = 0;
+		LODIndex = -1;
+		return;
+	}
+	CurveCount = InRequestedCurveCount;
+	PointCount = InRequestedPointCount;
+	LODIndex = InLODIndex;
+	bSupportOffsetLoad = false; // Whole resource loading/caching
+
+	using namespace UE::DerivedData;
+	TArray<FCacheGetChunkRequest> Requests;
+	In.Read_DDC(this, Requests);
+
+	check(DDCRequestOwner == nullptr);
+	DDCRequestOwner = MakeUnique<FRequestOwner>(UE::DerivedData::EPriority::Blocking);
+	RequestWarmCache(DDCRequestOwner.Get(), Requests);
+}
+#endif
+
 // Request fullfil 2 use cases:
 // * Load IO/DDC data and upload them to GPU
 // * Load DDC data and store them into bulkdata for serialization
@@ -357,14 +398,7 @@ void FHairStreamingRequest::Request(uint32 InRequestedCurveCount, uint32 InReque
 		// To reduce cook time, we warming cache by issuing 'Value' request (vs.'Chunk' request)
 		if (bWarmCache)
 		{
-			const ECachePolicy Policy = ECachePolicy::Default | ECachePolicy::SkipData;
-			TArray<FCacheGetValueRequest, TInlineAllocator<16>> WarmRequest;
-			for (const FCacheGetChunkRequest& R : Requests)
-			{				
-				WarmRequest.Add({R.Name, R.Key, Policy, 0});
-			}
-			GetCache().GetValue(WarmRequest, *DDCRequestOwner, [](FCacheGetValueResponse && Response) { check(Response.Status == EStatus::Ok); });
-			DDCRequestOwner->Wait();
+			RequestWarmCache(DDCRequestOwner.Get(), Requests);
 		}
 
 		//FRequestBarrier Barrier(*DDCRequestOwner);	// This is a critical section on the owner. It does not constrain ordering
diff --git a/Engine/Plugins/Runtime/HairStrands/Source/HairStrandsCore/Public/HairStrandsDatas.h b/Engine/Plugins/Runtime/HairStrands/Source/HairStrandsCore/Public/HairStrandsDatas.h
--- a/Engine/Plugins/Runtime/HairStrands/Source/HairStrandsCore/Public/HairStrandsDatas.h
+++ b/Engine/Plugins/Runtime/HairStrands/Source/HairStrandsCore/Public/HairStrandsDatas.h
@@ -372,6 +372,10 @@ struct FHairStreamingRequest
 	bool IsNone() const;
 	bool IsCompleted();
 
+#if WITH_EDITORONLY_DATA
+	void WarmCache(uint32 InRequestedCurveCount, uint32 InRequestedPointCount, int32 InLODIndex, FHairStrandsBulkCommon& In);
+#endif
+
 #if !WITH_EDITORONLY_DATA
 	// IO
 	FBulkDataBatchRequest IORequest;
