Add Min/Max/Avg point per curve stats to hair data.

This will allow to select appropriate culling/interpolation permutation.

diff --git a/Engine/Plugins/Runtime/HairStrands/Source/HairStrandsCore/Private/GroomBuilder.cpp b/Engine/Plugins/Runtime/HairStrands/Source/HairStrandsCore/Private/GroomBuilder.cpp
--- a/Engine/Plugins/Runtime/HairStrands/Source/HairStrandsCore/Private/GroomBuilder.cpp
+++ b/Engine/Plugins/Runtime/HairStrands/Source/HairStrandsCore/Private/GroomBuilder.cpp
@@ -40,7 +40,7 @@ static FAutoConsoleVariableRef CVarHairGroupIndexBuilder_MaxVoxelResolution(TEXT
 
 FString FGroomBuilder::GetVersion()
 {
-	return TEXT("v8r25");
+	return TEXT("v8r30");
 }
 
 namespace FHairStrandsDecimation
@@ -235,12 +235,22 @@ namespace HairStrandsBuilder
 				uint8 Packed;
 			};
 		};
+
+		uint32 MinPointPerCurve = HAIR_MAX_NUM_POINT_PER_CURVE;
+		uint32 MaxPointPerCurve = 0;
+		uint32 AccPointPerCurve = 0;
+
 		static_assert(sizeof(FPackedRadiusAndType) == sizeof(uint8));
 		const bool bSeedValid = RandomSeeds.Num() > 0;
 		for (uint32 CurveIndex = 0; CurveIndex < NumCurves; ++CurveIndex)
 		{
 			const int32 IndexOffset = Curves.CurvesOffset[CurveIndex];
-			const uint16& PointCount = Curves.CurvesCount[CurveIndex];
+			const uint16 PointCount = Curves.CurvesCount[CurveIndex];
+
+			MinPointPerCurve = FMath::Min(MinPointPerCurve, uint32(PointCount));
+			MaxPointPerCurve = FMath::Max(MaxPointPerCurve, uint32(PointCount));
+			AccPointPerCurve += PointCount;
+
 			for (int32 PointIndex = 0; PointIndex < PointCount; ++PointIndex)
 			{
 				const uint32 PrevIndex = FMath::Max(0, PointIndex - 1);
@@ -369,6 +379,9 @@ namespace HairStrandsBuilder
 		OutBulkData.Header.BoundingBox = HairStrands.BoundingBox;
 		OutBulkData.Header.CurveCount = HairStrands.GetNumCurves();
 		OutBulkData.Header.PointCount = HairStrands.GetNumPoints();
+		OutBulkData.Header.MinPointPerCurve = MinPointPerCurve;
+		OutBulkData.Header.MaxPointPerCurve = MaxPointPerCurve;
+		OutBulkData.Header.AvgPointPerCurve = FMath::DivideAndRoundDown(AccPointPerCurve, NumCurves);
 		OutBulkData.Header.MaxLength = MaxLength;
 		OutBulkData.Header.MaxRadius = MaxRadius;
 		OutBulkData.Header.Flags = FHairStrandsBulkData::DataFlags_HasData;
diff --git a/Engine/Plugins/Runtime/HairStrands/Source/HairStrandsCore/Private/HairStrandsDatas.cpp b/Engine/Plugins/Runtime/HairStrands/Source/HairStrandsCore/Private/HairStrandsDatas.cpp
--- a/Engine/Plugins/Runtime/HairStrands/Source/HairStrandsCore/Private/HairStrandsDatas.cpp
+++ b/Engine/Plugins/Runtime/HairStrands/Source/HairStrandsCore/Private/HairStrandsDatas.cpp
@@ -329,7 +329,7 @@ static void RequestWarmCache(UE::DerivedData::FRequestOwner* RequestOwner, const
 	{				
 		WarmRequest.Add({R.Name, R.Key, Policy, 0});
 	}
-	GetCache().GetValue(WarmRequest, *RequestOwner, [](FCacheGetValueResponse && Response) { check(Response.Status == EStatus::Ok); });
+	GetCache().GetValue(WarmRequest, *RequestOwner, [](FCacheGetValueResponse && Response) { /* If the data are not built the cache quert can return false. check(Response.Status == EStatus::Ok);*/ });
 	RequestOwner->Wait();
 }
 
@@ -481,6 +481,9 @@ void FHairStrandsBulkData::SerializeHeader(FArchive& Ar, UObject* Owner)
 	Ar << Header.MaxLength;
 	Ar << Header.MaxRadius;
 	Ar << Header.BoundingBox;
+	Ar << Header.MinPointPerCurve;
+	Ar << Header.MaxPointPerCurve;
+	Ar << Header.AvgPointPerCurve;
 	Ar << Header.Flags;
 	for (uint8 AttributeIt = 0; AttributeIt < HAIR_CURVE_ATTRIBUTE_COUNT; ++AttributeIt)
 	{
@@ -667,7 +670,7 @@ void FHairStrandsClusterCullingBulkData::Reset()
 
 	// Reset the bulk byte buffer to ensure the (serialize) data size is reset to 0
 	Data.ClusterLODInfos 	= FHairBulkContainer();
-	Data.CurveToClusterIds = FHairBulkContainer();
+	Data.CurveToClusterIds  = FHairBulkContainer();
 	Data.ClusterVertexIds 	= FHairBulkContainer();
 	Data.PackedClusterInfos = FHairBulkContainer();
 }
diff --git a/Engine/Plugins/Runtime/HairStrands/Source/HairStrandsCore/Public/HairStrandsDatas.h b/Engine/Plugins/Runtime/HairStrands/Source/HairStrandsCore/Public/HairStrandsDatas.h
--- a/Engine/Plugins/Runtime/HairStrands/Source/HairStrandsCore/Public/HairStrandsDatas.h
+++ b/Engine/Plugins/Runtime/HairStrands/Source/HairStrandsCore/Public/HairStrandsDatas.h
@@ -605,6 +605,9 @@ struct HAIRSTRANDSCORE_API FHairStrandsBulkData : FHairStrandsBulkCommon
 	{
 		uint32 CurveCount = 0;
 		uint32 PointCount = 0;
+		uint32 MinPointPerCurve = 0;
+		uint32 MaxPointPerCurve = 0;
+		uint32 AvgPointPerCurve = 0;
 		float  MaxLength = 0;
 		float  MaxRadius = 0;
 		FBox   BoundingBox = FBox(EForceInit::ForceInit);
