Move hair culling code into HairStrands plugin.

diff --git a/Engine/Plugins/Runtime/HairStrands/Source/HairStrandsCore/Private/GroomDebug.cpp b/Engine/Plugins/Runtime/HairStrands/Source/HairStrandsCore/Private/GroomDebug.cpp
--- a/Engine/Plugins/Runtime/HairStrands/Source/HairStrandsCore/Private/GroomDebug.cpp
+++ b/Engine/Plugins/Runtime/HairStrands/Source/HairStrandsCore/Private/GroomDebug.cpp
@@ -17,6 +17,7 @@
 #include "ShaderCompilerCore.h"
 #include "GroomVisualizationData.h"
 #include "GroomCacheData.h"
+#include "HairStrandsClusterCulling.h"
 
 ///////////////////////////////////////////////////////////////////////////////////////////////////
 
@@ -459,6 +460,98 @@ static void AddVoxelPlainRaymarchingPass(
 
 ///////////////////////////////////////////////////////////////////////////////////////////////////
 
+class FDrawDebugClusterAABBCS : public FGlobalShader
+{
+	DECLARE_GLOBAL_SHADER(FDrawDebugClusterAABBCS);
+	SHADER_USE_PARAMETER_STRUCT(FDrawDebugClusterAABBCS, FGlobalShader);
+
+	class FDebugAABBBuffer : SHADER_PERMUTATION_INT("PERMUTATION_DEBUGAABBBUFFER", 2);
+	using FPermutationDomain = TShaderPermutationDomain<FDebugAABBBuffer>;
+
+	BEGIN_SHADER_PARAMETER_STRUCT(FParameters, )
+		SHADER_PARAMETER_STRUCT_REF(FViewUniformShaderParameters, ViewUniformBuffer)
+		SHADER_PARAMETER_SRV(Buffer, ClusterAABBBuffer)
+		SHADER_PARAMETER_SRV(Buffer, GroupAABBBuffer)
+		SHADER_PARAMETER_RDG_BUFFER_SRV(Buffer, CulledDispatchIndirectParametersClusterCountBuffer)
+		SHADER_PARAMETER_RDG_BUFFER_SRV(StructuredBuffer, ClusterDebugInfoBuffer)
+		SHADER_PARAMETER_SRV(Buffer, CulledDrawIndirectParameters)
+		SHADER_PARAMETER(uint32, ClusterCount)
+		SHADER_PARAMETER(uint32, TriangleCount)
+		SHADER_PARAMETER(uint32, HairGroupId)
+		SHADER_PARAMETER(int32, ClusterDebugMode)
+		SHADER_PARAMETER_STRUCT_INCLUDE(ShaderPrint::FShaderParameters, ShaderPrintParameters)
+		END_SHADER_PARAMETER_STRUCT()
+
+public:
+	static bool ShouldCompilePermutation(const FGlobalShaderPermutationParameters& Parameters) { return IsHairStrandsSupported(EHairStrandsShaderType::Tool, Parameters.Platform); }
+	static void ModifyCompilationEnvironment(const FGlobalShaderPermutationParameters& Parameters, FShaderCompilerEnvironment& OutEnvironment)
+	{
+		FGlobalShader::ModifyCompilationEnvironment(Parameters, OutEnvironment);
+		OutEnvironment.SetDefine(TEXT("SHADER_DRAWDEBUGAABB"), 1);
+
+		// Skip optimization for avoiding long compilation time due to large UAV writes
+		OutEnvironment.CompilerFlags.Add(CFLAG_Debug);
+	}
+};
+
+IMPLEMENT_GLOBAL_SHADER(FDrawDebugClusterAABBCS, "/Engine/Private/HairStrands/HairStrandsClusterCulling.usf", "MainDrawDebugAABBCS", SF_Compute);
+
+void AddDrawDebugClusterPass(
+	FRDGBuilder& GraphBuilder,
+	const FSceneView& View,
+	FGlobalShaderMap* ShaderMap,
+	const FShaderPrintData* ShaderPrintData,
+	EGroomViewMode ViewMode,
+	FHairStrandClusterData& HairClusterData)
+{
+	// Force ShaderPrint on.
+	ShaderPrint::SetEnabled(true);
+	ShaderPrint::RequestSpaceForLines(64000u);
+	ShaderPrint::RequestSpaceForCharacters(2000);
+	if (!ShaderPrintData) { return; }
+
+	const bool bDebugAABB = ViewMode == EGroomViewMode::ClusterAABB;
+
+	uint32 DataIndex = 0;
+	for (const FHairStrandClusterData::FHairGroup& HairGroupClusters : HairClusterData.HairGroups)
+	{
+		if (HairGroupClusters.CulledClusterCountBuffer)
+		{
+			FRDGExternalBuffer& DrawIndirectBuffer = HairGroupClusters.HairGroupPublicPtr->GetDrawIndirectBuffer();
+
+			FDrawDebugClusterAABBCS::FPermutationDomain Permutation;
+			Permutation.Set<FDrawDebugClusterAABBCS::FDebugAABBBuffer>(bDebugAABB ? 1 : 0);
+			TShaderMapRef<FDrawDebugClusterAABBCS> ComputeShader(ShaderMap, Permutation);
+
+			FDrawDebugClusterAABBCS::FParameters* Parameters = GraphBuilder.AllocParameters<FDrawDebugClusterAABBCS::FParameters>();
+			Parameters->ViewUniformBuffer = View.ViewUniformBuffer;
+			Parameters->ClusterCount = HairGroupClusters.ClusterCount;
+			Parameters->TriangleCount = HairGroupClusters.VertexCount * 2; // VertexCount is actually the number of control points
+			Parameters->HairGroupId = DataIndex++;
+			Parameters->ClusterDebugMode = ViewMode == EGroomViewMode::Cluster ? 1 : (ViewMode == EGroomViewMode::ClusterAABB ? 2 : 0);
+			Parameters->ClusterAABBBuffer = HairGroupClusters.ClusterAABBBuffer->SRV;
+			Parameters->CulledDispatchIndirectParametersClusterCountBuffer = GraphBuilder.CreateSRV(HairGroupClusters.CulledClusterCountBuffer, EPixelFormat::PF_R32_UINT);
+			Parameters->CulledDrawIndirectParameters = DrawIndirectBuffer.SRV;
+			Parameters->GroupAABBBuffer = HairGroupClusters.GroupAABBBuffer->SRV;
+
+			if (HairGroupClusters.ClusterDebugInfoBuffer && bDebugAABB)
+			{
+				FRDGBufferRef ClusterDebugInfoBuffer = GraphBuilder.RegisterExternalBuffer(HairGroupClusters.ClusterDebugInfoBuffer);
+				Parameters->ClusterDebugInfoBuffer = GraphBuilder.CreateSRV(ClusterDebugInfoBuffer);
+			}
+			//ShaderPrint::SetParameters(GraphBuilder, ShaderPrintData, Parameters->ShaderPrintParameters);
+			ShaderPrint::SetParameters(GraphBuilder, *ShaderPrintData, Parameters->ShaderPrintParameters);
+
+			const FIntVector DispatchCount = DispatchCount.DivideAndRoundUp(FIntVector(Parameters->ClusterCount, 1, 1), FIntVector(64, 1, 1));// FIX ME, this could get over 65535
+			FComputeShaderUtils::AddPass(
+				GraphBuilder, RDG_EVENT_NAME("DrawDebugClusterAABB"),
+				ComputeShader, Parameters, DispatchCount);
+		}
+	}
+}
+
+///////////////////////////////////////////////////////////////////////////////////////////////////
+
 class FDrawDebugCardAtlasCS : public FGlobalShader
 {
 	DECLARE_GLOBAL_SHADER(FDrawDebugCardAtlasCS);
@@ -1235,23 +1328,31 @@ void RunHairStrandsDebug(
 		}
 	}
 
-	for (FHairStrandsInstance* AbstractInstance : Instances)
-	{
-		FHairGroupInstance* Instance = static_cast<FHairGroupInstance*>(AbstractInstance);
-
-		if (ViewMode == EGroomViewMode::CardGuides)
-		{
-			AddDrawDebugCardsGuidesPass(GraphBuilder, View, ShaderMap, Instance, ShaderPrintData, true /*bDeformed*/, true);
-		}
-
-		if (ViewMode == EGroomViewMode::SimHairStrands)
-		{
-			AddDrawDebugCardsGuidesPass(GraphBuilder, View, ShaderMap, Instance, ShaderPrintData, true /*bDeformed*/, false);
-		}
+	const bool bEnabled =
+		ViewMode == EGroomViewMode::CardGuides ||
+		ViewMode == EGroomViewMode::SimHairStrands ||
+		ViewMode == EGroomViewMode::ControlPoints;
 
-		if (ViewMode == EGroomViewMode::ControlPoints)
+	if (bEnabled)
+	{	
+		for (FHairStrandsInstance* AbstractInstance : Instances)
 		{
-			AddDrawDebugStrandsCVsPass(GraphBuilder, View, ShaderMap, Instance, ShaderPrintData, SceneColorTexture, SceneDepthTexture);
+			FHairGroupInstance* Instance = static_cast<FHairGroupInstance*>(AbstractInstance);
+	
+			if (ViewMode == EGroomViewMode::CardGuides)
+			{
+				AddDrawDebugCardsGuidesPass(GraphBuilder, View, ShaderMap, Instance, ShaderPrintData, true /*bDeformed*/, true);
+			}
+	
+			if (ViewMode == EGroomViewMode::SimHairStrands)
+			{
+				AddDrawDebugCardsGuidesPass(GraphBuilder, View, ShaderMap, Instance, ShaderPrintData, true /*bDeformed*/, false);
+			}
+	
+			if (ViewMode == EGroomViewMode::ControlPoints)
+			{
+				AddDrawDebugStrandsCVsPass(GraphBuilder, View, ShaderMap, Instance, ShaderPrintData, SceneColorTexture, SceneDepthTexture);
+			}
 		}
 	}
 }
\ No newline at end of file
diff --git a/Engine/Plugins/Runtime/HairStrands/Source/HairStrandsCore/Private/GroomManager.cpp b/Engine/Plugins/Runtime/HairStrands/Source/HairStrandsCore/Private/GroomManager.cpp
--- a/Engine/Plugins/Runtime/HairStrands/Source/HairStrandsCore/Private/GroomManager.cpp
+++ b/Engine/Plugins/Runtime/HairStrands/Source/HairStrandsCore/Private/GroomManager.cpp
@@ -26,6 +26,8 @@
 #include "GroomCacheData.h"
 #include "SceneInterface.h"
 #include "PrimitiveSceneInfo.h"
+#include "HairStrandsClusterCulling.h"
+#include "GroomVisualizationData.h"
 
 static int32 GHairStrandsMinLOD = 0;
 static FAutoConsoleVariableRef CVarGHairStrandsMinLOD(TEXT("r.HairStrands.MinLOD"), GHairStrandsMinLOD, TEXT("Clamp the min hair LOD to this value, preventing to reach lower/high-quality LOD."), ECVF_Scalability);
@@ -437,8 +439,7 @@ static void RunHairStrandsInterpolation_Guide(
 	const uint32 ViewUniqueID,
 	const FHairStrandsInstances& Instances,
 	const FShaderPrintData* ShaderPrintData,
-	FGlobalShaderMap* ShaderMap,
-	FHairStrandClusterData* ClusterData)
+	FGlobalShaderMap* ShaderMap)
 {
 	check(IsInRenderingThread());
 
@@ -455,18 +456,26 @@ static void RunHairStrandsInterpolation_Guide(
 		ShaderPrintData,
 		ShaderMap,
 		EHairStrandsInterpolationType::SimulationStrands,
-		ClusterData);
+		nullptr);
 }
 
+void AddDrawDebugClusterPass(
+	FRDGBuilder& GraphBuilder,
+	const FSceneView& View,
+	FGlobalShaderMap* ShaderMap,
+	const FShaderPrintData* ShaderPrintData,
+	EGroomViewMode ViewMode,
+	FHairStrandClusterData& HairClusterData);
+
 static void RunHairStrandsInterpolation_Strands(
 	FRDGBuilder& GraphBuilder,
 	FSceneInterface* Scene,
+	const TArray<const FSceneView*>& Views, 
 	const FSceneView* View,
 	const uint32 ViewUniqueID,
 	const FHairStrandsInstances& Instances,
 	const FShaderPrintData* ShaderPrintData,
-	FGlobalShaderMap* ShaderMap,
-	FHairStrandClusterData* ClusterData)
+	FGlobalShaderMap* ShaderMap)
 {
 	check(IsInRenderingThread());
 
@@ -474,6 +483,34 @@ static void RunHairStrandsInterpolation_Strands(
 	RDG_EVENT_SCOPE(GraphBuilder, "HairStrandsInterpolation");
 	RDG_GPU_STAT_SCOPE(GraphBuilder, HairStrandsInterpolation);
 
+	FHairStrandClusterData ClusterData;
+	if (IsHairStrandsEnabled(EHairStrandsShaderType::Strands, Scene->GetShaderPlatform()))
+	{
+		for (FHairStrandsInstance* AbstractInstance : Instances)
+		{
+			FHairGroupInstance* Instance = static_cast<FHairGroupInstance*>(AbstractInstance);
+			if (Instance->GeometryType == EHairGeometryType::Strands && Instance->Strands.IsValid())
+			{
+				if (Instance->Strands.bCullingEnable)
+				{
+					AddInstanceToClusterData(Instance, ClusterData);
+				}
+				else
+				{
+					Instance->HairGroupPublicData->SetCullingResultAvailable(false);
+				}
+			}
+		}
+		ComputeHairStrandsClustersCulling(GraphBuilder, *ShaderMap, Views, ClusterData);
+
+		// Run cluster debug view here (instead of GroomDebug.h/.cpp, as we need to have the (transient) cluster data 
+		const EGroomViewMode ViewMode = GetGroomViewMode(*View);
+		if (ViewMode == EGroomViewMode::Cluster || ViewMode == EGroomViewMode::ClusterAABB)
+		{
+			AddDrawDebugClusterPass(GraphBuilder, *View, ShaderMap, ShaderPrintData, ViewMode, ClusterData);
+		}
+	}
+
 	RunInternalHairStrandsInterpolation(
 		GraphBuilder,
 		Scene,
@@ -483,7 +520,7 @@ static void RunHairStrandsInterpolation_Strands(
 		ShaderPrintData,
 		ShaderMap,
 		EHairStrandsInterpolationType::RenderStrands,
-		ClusterData);
+		&ClusterData);
 }
 
 bool IsHairStrandsContinousLODEnabled()
@@ -491,27 +528,6 @@ bool IsHairStrandsContinousLODEnabled()
 	return GHairStrands_ContinousLOD > 0;
 }
 
-static void RunHairStrandsGatherCluster(
-	const FHairStrandsInstances& Instances,
-	FHairStrandClusterData* ClusterData)
-{
-	for (FHairStrandsInstance* AbstractInstance : Instances)
-	{
-		FHairGroupInstance* Instance = static_cast<FHairGroupInstance*>(AbstractInstance);
-		if (Instance->GeometryType == EHairGeometryType::Strands && Instance->Strands.IsValid())
-		{
-			if (Instance->Strands.bCullingEnable)
-			{
-				RegisterClusterData(Instance, ClusterData);
-			}
-			else
-			{
-				Instance->HairGroupPublicData->SetCullingResultAvailable(false);
-			}
-		}
-	}
-}
-
 // Return the LOD which should be used for a given screen size and LOD bias value
 // This function is mirrored in HairStrandsClusterCommon.ush
 static float GetHairInstanceLODIndex(const TArray<float>& InLODScreenSizes, float InScreenSize, float InLODBias)
@@ -1245,14 +1261,7 @@ void ProcessHairStrandsBookmark(
 			Parameters.ViewUniqueID,
 			Instances,
 			Parameters.ShaderPrintData,
-			Parameters.ShaderMap,
-			&Parameters.HairClusterData);
-	}
-	else if (Bookmark == EHairStrandsBookmark::ProcessGatherCluster)
-	{
-		RunHairStrandsGatherCluster(
-			Instances,
-			&Parameters.HairClusterData);
+			Parameters.ShaderMap);
 	}
 	else if (Bookmark == EHairStrandsBookmark::ProcessStrandsInterpolation)
 	{
@@ -1260,12 +1269,12 @@ void ProcessHairStrandsBookmark(
 		RunHairStrandsInterpolation_Strands(
 			*GraphBuilder,
 			Parameters.Scene,
+			Parameters.AllViews,
 			Parameters.View,
 			Parameters.ViewUniqueID,
 			Instances,
 			Parameters.ShaderPrintData,
-			Parameters.ShaderMap,
-			&Parameters.HairClusterData);
+			Parameters.ShaderMap);
 	}
 	else if (Bookmark == EHairStrandsBookmark::ProcessDebug)
 	{
diff --git a/Engine/Plugins/Runtime/HairStrands/Source/HairStrandsCore/Private/HairStrandsClusterCulling.cpp b/Engine/Plugins/Runtime/HairStrands/Source/HairStrandsCore/Private/HairStrandsClusterCulling.cpp
new file mode 100644
--- /dev/null
+++ b/Engine/Plugins/Runtime/HairStrands/Source/HairStrandsCore/Private/HairStrandsClusterCulling.cpp
@@ -0,0 +1,624 @@
+// Copyright Epic Games, Inc. All Rights Reserved.
+
+#include "HairStrandsClusterCulling.h"
+#include "RendererInterface.h"
+#include "Shader.h"
+#include "RenderGraph.h"
+#include "GlobalShader.h"
+#include "ShaderParameters.h"
+#include "ShaderParameterStruct.h"
+#include "ProfilingDebugging/CpuProfilerTrace.h"
+#include "GroomVisualizationData.h"
+#include "GroomInstance.h"
+
+///////////////////////////////////////////////////////////////////////////////////////////////////
+
+static int32 GHairStrandsClusterForceLOD = -1;
+static FAutoConsoleVariableRef CVarHairClusterCullingLodMode(TEXT("r.HairStrands.Cluster.ForceLOD"), GHairStrandsClusterForceLOD, TEXT("Force a specific hair LOD."));
+
+static int32 GHairStrandsClusterCullingFreezeCamera = 0;
+static FAutoConsoleVariableRef CVarHairStrandsClusterCullingFreezeCamera(TEXT("r.HairStrands.Cluster.CullingFreezeCamera"), GHairStrandsClusterCullingFreezeCamera, TEXT("Freeze camera when enabled. It will disable HZB culling because hzb buffer is not frozen."));
+
+bool IsHairStrandsClusterCullingEnable()
+{
+	// At the moment it is not possible to disable cluster culling, as this pass is in charge of LOD selection, 
+	// and preparing the buffer which will be need for the cluster AABB pass (used later on by the voxelization pass)
+	return true;
+}
+
+// Temp: to remove
+enum EHairVisibilityVendor
+{
+	HairVisibilityVendor_AMD,
+	HairVisibilityVendor_NVIDIA,
+	HairVisibilityVendor_INTEL,
+	HairVisibilityVendorCount
+};
+
+// Temp: to remove
+inline EHairVisibilityVendor GetVendor()
+{
+	return IsRHIDeviceAMD() ? HairVisibilityVendor_AMD : (IsRHIDeviceNVIDIA() ? HairVisibilityVendor_NVIDIA : HairVisibilityVendor_INTEL);
+}
+// Temp: to remove
+static uint32 GetVendorOptimalGroupSize1D()
+{
+	switch (GetVendor())
+	{
+		case HairVisibilityVendor_AMD:		return 64;
+		case HairVisibilityVendor_NVIDIA:	return 32;
+		case HairVisibilityVendor_INTEL:	return 64;
+		default:							return 64;
+	}
+}
+
+///////////////////////////////////////////////////////////////////////////////////////////////////
+
+class FHairIndBufferClearCS : public FGlobalShader
+{
+	DECLARE_GLOBAL_SHADER(FHairIndBufferClearCS);
+	SHADER_USE_PARAMETER_STRUCT(FHairIndBufferClearCS, FGlobalShader);
+
+	class FSetIndirectDraw : SHADER_PERMUTATION_BOOL("PERMUTATION_SETINDIRECTDRAW");
+	using FPermutationDomain = TShaderPermutationDomain<FSetIndirectDraw>;
+
+	BEGIN_SHADER_PARAMETER_STRUCT(FParameters, )
+		SHADER_PARAMETER_RDG_BUFFER_UAV(RWBuffer, DispatchIndirectParametersClusterCount)
+		SHADER_PARAMETER_RDG_BUFFER_UAV(RWBuffer, DrawIndirectParameters)
+		SHADER_PARAMETER(uint32, VertexCountPerInstance)
+	END_SHADER_PARAMETER_STRUCT()
+
+public:
+	static bool ShouldCompilePermutation(const FGlobalShaderPermutationParameters& Parameters) { return IsHairStrandsSupported(EHairStrandsShaderType::Strands, Parameters.Platform); }
+	static void ModifyCompilationEnvironment(const FGlobalShaderPermutationParameters& Parameters, FShaderCompilerEnvironment& OutEnvironment)
+	{
+		FGlobalShader::ModifyCompilationEnvironment(Parameters, OutEnvironment);
+		OutEnvironment.SetDefine(TEXT("SHADER_CLUSTERCULLINGINDCLEAR"), 1);
+	}
+};
+
+IMPLEMENT_GLOBAL_SHADER(FHairIndBufferClearCS, "/Engine/Private/HairStrands/HairStrandsClusterCulling.usf", "MainClusterCullingIndClearCS", SF_Compute);
+
+///////////////////////////////////////////////////////////////////////////////////////////////////
+
+class FHairClusterCullingCS : public FGlobalShader
+{
+	DECLARE_GLOBAL_SHADER(FHairClusterCullingCS);
+	SHADER_USE_PARAMETER_STRUCT(FHairClusterCullingCS, FGlobalShader);
+
+	class FDebugAABBBuffer : SHADER_PERMUTATION_INT("PERMUTATION_DEBUGAABBBUFFER", 2);
+	using FPermutationDomain = TShaderPermutationDomain<FDebugAABBBuffer>;
+
+	BEGIN_SHADER_PARAMETER_STRUCT(FParameters, )
+		SHADER_PARAMETER(FVector3f, CameraWorldPos)
+		SHADER_PARAMETER(FMatrix44f, WorldToClipMatrix)
+		SHADER_PARAMETER(FMatrix44f, ProjectionMatrix)
+		SHADER_PARAMETER(uint32, ClusterCount)
+		SHADER_PARAMETER(float, LODForcedIndex)
+		SHADER_PARAMETER(int32, bIsHairGroupVisible)
+		SHADER_PARAMETER(uint32, NumConvexHullPlanes)
+		SHADER_PARAMETER(float, LODBias)
+		SHADER_PARAMETER_ARRAY(FVector4f, ViewFrustumConvexHull, [6])
+		SHADER_PARAMETER_RDG_BUFFER_SRV(Buffer, ClusterAABBBuffer)
+		SHADER_PARAMETER_RDG_BUFFER_SRV(StructuredBuffer, ClusterInfoBuffer)
+		SHADER_PARAMETER_RDG_BUFFER_SRV(StructuredBuffer, ClusterLODInfoBuffer)
+		SHADER_PARAMETER_RDG_BUFFER_UAV(RWBuffer, GlobalClusterIdBuffer)
+		SHADER_PARAMETER_RDG_BUFFER_UAV(RWBuffer, GlobalIndexStartBuffer)
+		SHADER_PARAMETER_RDG_BUFFER_UAV(RWBuffer, GlobalIndexCountBuffer)
+		SHADER_PARAMETER_RDG_BUFFER_UAV(RWBuffer, GlobalRadiusScaleBuffer)
+		SHADER_PARAMETER_RDG_BUFFER_UAV(RWStructuredBuffer, ClusterDebugInfoBuffer)
+		SHADER_PARAMETER_RDG_BUFFER_UAV(RWBuffer, DispatchIndirectParametersClusterCount)
+		SHADER_PARAMETER_RDG_BUFFER_UAV(RWBuffer, DrawIndirectParameters)
+	END_SHADER_PARAMETER_STRUCT()
+
+public:
+	static bool ShouldCompilePermutation(const FGlobalShaderPermutationParameters& Parameters) { return IsHairStrandsSupported(EHairStrandsShaderType::Strands, Parameters.Platform); }
+	static void ModifyCompilationEnvironment(const FGlobalShaderPermutationParameters& Parameters, FShaderCompilerEnvironment& OutEnvironment)
+	{
+		FGlobalShader::ModifyCompilationEnvironment(Parameters, OutEnvironment);
+		OutEnvironment.SetDefine(TEXT("SHADER_CLUSTERCULLING"), 1);
+	}
+};
+
+IMPLEMENT_GLOBAL_SHADER(FHairClusterCullingCS, "/Engine/Private/HairStrands/HairStrandsClusterCulling.usf", "MainClusterCullingCS", SF_Compute);
+
+///////////////////////////////////////////////////////////////////////////////////////////////////
+
+class FMainClusterCullingPrepareIndirectDrawsCS : public FGlobalShader
+{
+	DECLARE_GLOBAL_SHADER(FMainClusterCullingPrepareIndirectDrawsCS);
+	SHADER_USE_PARAMETER_STRUCT(FMainClusterCullingPrepareIndirectDrawsCS, FGlobalShader);
+
+	BEGIN_SHADER_PARAMETER_STRUCT(FParameters, )
+		SHADER_PARAMETER(uint32, GroupSize1D)
+		SHADER_PARAMETER_RDG_BUFFER_SRV(Buffer,   DispatchIndirectParametersClusterCount)
+		SHADER_PARAMETER_RDG_BUFFER_UAV(RWBuffer, DispatchIndirectParametersClusterCount1D)
+		SHADER_PARAMETER_RDG_BUFFER_UAV(RWBuffer, DispatchIndirectParametersClusterCount2D)
+		SHADER_PARAMETER_RDG_BUFFER_UAV(RWBuffer, DispatchIndirectParametersClusterCountDiv512)
+		SHADER_PARAMETER_RDG_BUFFER_UAV(RWBuffer, DispatchIndirectParametersClusterCountDiv512Div512)
+	END_SHADER_PARAMETER_STRUCT()
+
+public:
+	static bool ShouldCompilePermutation(const FGlobalShaderPermutationParameters& Parameters) { return IsHairStrandsSupported(EHairStrandsShaderType::Strands, Parameters.Platform); }
+	static void ModifyCompilationEnvironment(const FGlobalShaderPermutationParameters& Parameters, FShaderCompilerEnvironment& OutEnvironment)
+	{
+		FGlobalShader::ModifyCompilationEnvironment(Parameters, OutEnvironment);
+		OutEnvironment.SetDefine(TEXT("SHADER_PREPAREINDIRECTDRAW"), 1);
+	}
+};
+
+IMPLEMENT_GLOBAL_SHADER(FMainClusterCullingPrepareIndirectDrawsCS, "/Engine/Private/HairStrands/HairStrandsClusterCulling.usf", "MainClusterCullingPrepareIndirectDrawsCS", SF_Compute);
+
+///////////////////////////////////////////////////////////////////////////////////////////////////
+
+class FMainClusterCullingPrepareIndirectDispatchCS : public FGlobalShader
+{
+	DECLARE_GLOBAL_SHADER(FMainClusterCullingPrepareIndirectDispatchCS);
+	SHADER_USE_PARAMETER_STRUCT(FMainClusterCullingPrepareIndirectDispatchCS, FGlobalShader);
+
+	using FPermutationDomain = TShaderPermutationDomain<>;
+
+	BEGIN_SHADER_PARAMETER_STRUCT(FParameters, )
+		SHADER_PARAMETER_RDG_BUFFER_UAV(RWBuffer, DrawIndirectBuffer)
+		SHADER_PARAMETER_RDG_BUFFER_UAV(RWBuffer, DispatchIndirectBuffer)
+	END_SHADER_PARAMETER_STRUCT()
+
+public:
+	static bool ShouldCompilePermutation(const FGlobalShaderPermutationParameters& Parameters) { return IsHairStrandsSupported(EHairStrandsShaderType::Strands, Parameters.Platform); }
+	static void ModifyCompilationEnvironment(const FGlobalShaderPermutationParameters& Parameters, FShaderCompilerEnvironment& OutEnvironment)
+	{
+		FGlobalShader::ModifyCompilationEnvironment(Parameters, OutEnvironment);
+		OutEnvironment.SetDefine(TEXT("SHADER_PREPAREINDIRECTDISPATCH"), 1);
+	}
+};
+
+IMPLEMENT_GLOBAL_SHADER(FMainClusterCullingPrepareIndirectDispatchCS, "/Engine/Private/HairStrands/HairStrandsClusterCulling.usf", "MainClusterCullingPrepareIndirectDispatchCS", SF_Compute);
+
+///////////////////////////////////////////////////////////////////////////////////////////////////
+
+class FHairClusterCullingLocalBlockPreFixSumCS : public FGlobalShader
+{
+	DECLARE_GLOBAL_SHADER(FHairClusterCullingLocalBlockPreFixSumCS);
+	SHADER_USE_PARAMETER_STRUCT(FHairClusterCullingLocalBlockPreFixSumCS, FGlobalShader);
+
+	BEGIN_SHADER_PARAMETER_STRUCT(FParameters, )
+		RDG_BUFFER_ACCESS(DispatchIndirectParametersClusterCountDiv512, ERHIAccess::IndirectArgs)
+		SHADER_PARAMETER_RDG_BUFFER_SRV(Buffer, DispatchIndirectParametersClusterCount)
+		SHADER_PARAMETER_RDG_BUFFER_SRV(Buffer, GlobalIndexCountBuffer)
+		SHADER_PARAMETER_RDG_BUFFER_UAV(RWBuffer, PerBlocklIndexCountPreFixSumBuffer)
+		SHADER_PARAMETER_RDG_BUFFER_UAV(RWBuffer, PerBlocklTotalIndexCountBuffer)
+	END_SHADER_PARAMETER_STRUCT()
+
+public:
+	static bool ShouldCompilePermutation(const FGlobalShaderPermutationParameters& Parameters) { return IsHairStrandsSupported(EHairStrandsShaderType::Strands, Parameters.Platform); }
+	static void ModifyCompilationEnvironment(const FGlobalShaderPermutationParameters& Parameters, FShaderCompilerEnvironment& OutEnvironment)
+	{
+		FGlobalShader::ModifyCompilationEnvironment(Parameters, OutEnvironment);
+		OutEnvironment.SetDefine(TEXT("SHADER_PERBLOCKPREFIXSUM"), 1);
+	}
+};
+
+IMPLEMENT_GLOBAL_SHADER(FHairClusterCullingLocalBlockPreFixSumCS, "/Engine/Private/HairStrands/HairStrandsClusterCulling.usf", "MainPerBlockPreFixSumCS", SF_Compute);
+
+///////////////////////////////////////////////////////////////////////////////////////////////////
+
+class FHairClusterCullingCompactVertexIdsLocalBlockCS : public FGlobalShader
+{
+	DECLARE_GLOBAL_SHADER(FHairClusterCullingCompactVertexIdsLocalBlockCS);
+	SHADER_USE_PARAMETER_STRUCT(FHairClusterCullingCompactVertexIdsLocalBlockCS, FGlobalShader);
+
+	BEGIN_SHADER_PARAMETER_STRUCT(FParameters, )
+		RDG_BUFFER_ACCESS(DispatchIndirectParametersBuffer, ERHIAccess::IndirectArgs)
+		SHADER_PARAMETER_RDG_BUFFER_SRV(Buffer, DispatchIndirectParametersClusterCount)
+		SHADER_PARAMETER_RDG_BUFFER_SRV(Buffer, DispatchIndirectParametersClusterCount2D)
+		SHADER_PARAMETER_RDG_BUFFER_SRV(Buffer, DispatchIndirectParametersClusterCountDiv512)
+		SHADER_PARAMETER_RDG_BUFFER_SRV(Buffer, PerBlocklIndexCountPreFixSumBuffer)
+		SHADER_PARAMETER_RDG_BUFFER_SRV(Buffer, PerBlocklTotalIndexCountPreFixSumBuffer)
+		SHADER_PARAMETER_RDG_BUFFER_SRV(Buffer, GlobalIndexStartBuffer)
+		SHADER_PARAMETER_RDG_BUFFER_SRV(Buffer, GlobalIndexCountBuffer)
+		SHADER_PARAMETER_RDG_BUFFER_SRV(Buffer, GlobalRadiusScaleBuffer)
+		SHADER_PARAMETER_RDG_BUFFER_SRV(Buffer, ClusterVertexIdBuffer)
+		SHADER_PARAMETER_RDG_BUFFER_UAV(RWBuffer, CulledCompactedIndexBuffer)
+		SHADER_PARAMETER_RDG_BUFFER_UAV(RWBuffer, CulledCompactedRadiusScaleBuffer)
+	END_SHADER_PARAMETER_STRUCT()
+
+public:
+	static bool ShouldCompilePermutation(const FGlobalShaderPermutationParameters& Parameters) { return IsHairStrandsSupported(EHairStrandsShaderType::Strands, Parameters.Platform); }
+	static void ModifyCompilationEnvironment(const FGlobalShaderPermutationParameters& Parameters, FShaderCompilerEnvironment& OutEnvironment)
+	{
+		FGlobalShader::ModifyCompilationEnvironment(Parameters, OutEnvironment);
+		OutEnvironment.SetDefine(TEXT("SHADER_CLUSTERCULLINGCOMPACTVERTEXIDLOCALBLOCK"), 1);
+	}
+};
+
+IMPLEMENT_GLOBAL_SHADER(FHairClusterCullingCompactVertexIdsLocalBlockCS, "/Engine/Private/HairStrands/HairStrandsClusterCulling.usf", "MainClusterCullingCompactVertexIdsLocalBlockCS", SF_Compute);
+
+///////////////////////////////////////////////////////////////////////////////////////////////////
+
+class FHairClusterCullingPreFixSumCS : public FGlobalShader
+{
+	DECLARE_GLOBAL_SHADER(FHairClusterCullingPreFixSumCS);
+	SHADER_USE_PARAMETER_STRUCT(FHairClusterCullingPreFixSumCS, FGlobalShader);
+
+	BEGIN_SHADER_PARAMETER_STRUCT(FParameters, )
+		RDG_BUFFER_ACCESS(DispatchIndirectParameters, ERHIAccess::IndirectArgs)
+		SHADER_PARAMETER_RDG_BUFFER_SRV(Buffer, DispatchIndirectParametersClusterCount)
+		SHADER_PARAMETER_RDG_BUFFER_SRV(Buffer, GlobalIndexCountBuffer)
+		SHADER_PARAMETER_RDG_BUFFER_UAV(RWBuffer, GlobalIndexCountPreFixSumBuffer)
+	END_SHADER_PARAMETER_STRUCT()
+
+public:
+	static bool ShouldCompilePermutation(const FGlobalShaderPermutationParameters& Parameters) { return IsHairStrandsSupported(EHairStrandsShaderType::Strands, Parameters.Platform); }
+	static void ModifyCompilationEnvironment(const FGlobalShaderPermutationParameters& Parameters, FShaderCompilerEnvironment& OutEnvironment)
+	{
+		FGlobalShader::ModifyCompilationEnvironment(Parameters, OutEnvironment);
+		OutEnvironment.SetDefine(TEXT("SHADER_CLUSTERCULLINGPREFIXSUM"), 1);
+	}
+};
+
+IMPLEMENT_GLOBAL_SHADER(FHairClusterCullingPreFixSumCS, "/Engine/Private/HairStrands/HairStrandsClusterCulling.usf", "MainClusterCullingPreFixSumCS", SF_Compute);
+
+///////////////////////////////////////////////////////////////////////////////////////////////////
+
+static FVector CapturedCameraWorldPos;
+static FMatrix CapturedWorldToClipMatrix;
+static FMatrix CapturedProjMatrix;
+
+static void AddClusterCullingPass(
+	FRDGBuilder& GraphBuilder,
+	FGlobalShaderMap* ShaderMap,
+	const FSceneView& View,
+	FHairStrandClusterData::FHairGroup& ClusterData)
+{
+	FRDGBufferRef DispatchIndirectParametersClusterCount = GraphBuilder.CreateBuffer(FRDGBufferDesc::CreateIndirectDesc<FRHIDispatchIndirectParameters>(), TEXT("Hair.DispatchIndirectParametersClusterCount"));
+	FRDGBufferRef DispatchIndirectParametersClusterCount1D = GraphBuilder.CreateBuffer(FRDGBufferDesc::CreateIndirectDesc<FRHIDispatchIndirectParameters>(), TEXT("Hair.DispatchIndirectParametersClusterCount1D"));
+	FRDGBufferRef DispatchIndirectParametersClusterCount2D = GraphBuilder.CreateBuffer(FRDGBufferDesc::CreateIndirectDesc<FRHIDispatchIndirectParameters>(), TEXT("Hair.DispatchIndirectParametersClusterCount2D"));
+	FRDGBufferRef DispatchIndirectParametersClusterCountDiv512 = GraphBuilder.CreateBuffer(FRDGBufferDesc::CreateIndirectDesc<FRHIDispatchIndirectParameters>(), TEXT("Hair.DispatchIndirectParametersClusterCountDiv512"));
+	FRDGBufferRef DispatchIndirectParametersClusterCountDiv512Div512 = GraphBuilder.CreateBuffer(FRDGBufferDesc::CreateIndirectDesc<FRHIDispatchIndirectParameters>(), TEXT("Hair.DispatchIndirectParametersClusterCountDiv512Div512"));
+	
+	FRDGBufferRef GlobalClusterIdBuffer = GraphBuilder.CreateBuffer(FRDGBufferDesc::CreateBufferDesc(sizeof(uint32), ClusterData.ClusterCount), TEXT("Hair.GlobalClusterIdBuffer"));
+	FRDGBufferRef GlobalIndexStartBuffer = GraphBuilder.CreateBuffer(FRDGBufferDesc::CreateBufferDesc(sizeof(uint32), ClusterData.ClusterCount), TEXT("Hair.GlobalIndexStartBuffer"));
+	FRDGBufferRef GlobalIndexCountBuffer = GraphBuilder.CreateBuffer(FRDGBufferDesc::CreateBufferDesc(sizeof(uint32), ClusterData.ClusterCount), TEXT("Hair.GlobalIndexCountBuffer"));
+	FRDGBufferRef GlobalRadiusScaleBuffer = GraphBuilder.CreateBuffer(FRDGBufferDesc::CreateBufferDesc(sizeof(float), ClusterData.ClusterCount), TEXT("Hair.GlobalRadiusScaleBuffer"));
+
+	const uint32 PrefixGroupSize = 512;
+	const uint32 ClusterCountRoundUpToGroupSize = FMath::DivideAndRoundUp(ClusterData.ClusterCount, PrefixGroupSize) * PrefixGroupSize;
+	FRDGBufferRef PerBlocklTotalIndexCountBuffer = GraphBuilder.CreateBuffer(FRDGBufferDesc::CreateBufferDesc(sizeof(uint32), ClusterCountRoundUpToGroupSize), TEXT("Hair.PerBlocklTotalIndexCountBuffer"));
+	FRDGBufferRef PerBlocklTotalIndexCountPreFixSumBuffer = GraphBuilder.CreateBuffer(FRDGBufferDesc::CreateBufferDesc(sizeof(uint32) * 2, ClusterCountRoundUpToGroupSize), TEXT("Hair.PerBlocklTotalIndexCountPreFixSumBuffer"));
+	FRDGBufferRef PerBlocklIndexCountPreFixSumBuffer = GraphBuilder.CreateBuffer(FRDGBufferDesc::CreateBufferDesc(sizeof(uint32) * 2, ClusterCountRoundUpToGroupSize), TEXT("Hair.PerBlocklIndexCountPreFixSumBuffer"));
+
+	FRDGImportedBuffer DrawIndirectParametersBuffer = Register(GraphBuilder, ClusterData.HairGroupPublicPtr->GetDrawIndirectBuffer(), ERDGImportedBufferFlags::CreateViews);
+	FRDGImportedBuffer DrawIndirectParametersRasterComputeBuffer = Register(GraphBuilder, ClusterData.HairGroupPublicPtr->GetDrawIndirectRasterComputeBuffer(), ERDGImportedBufferFlags::CreateViews);
+	
+	bool bClusterDebugAABBBuffer = false;
+	bool bClusterDebug = false;
+#if WITH_EDITOR
+	FRDGBufferRef ClusterDebugInfoBuffer = nullptr;
+	{
+		// Defined in HairStrandsClusterCommon.ush
+		struct FHairClusterDebugInfo
+		{
+			uint32 GroupIndex;
+			float LOD;
+			float VertexCount;
+			float CurveCount;
+		};
+
+		static IConsoleVariable* CVarClusterDebug = IConsoleManager::Get().FindConsoleVariable(TEXT("r.HairStrands.Cluster.Debug"));
+		bClusterDebugAABBBuffer = GetGroomViewMode(View) == EGroomViewMode::ClusterAABB;
+		bClusterDebug = GetGroomViewMode(View) == EGroomViewMode::Cluster;
+		ClusterDebugInfoBuffer = GraphBuilder.CreateBuffer(FRDGBufferDesc::CreateStructuredDesc(sizeof(FHairClusterDebugInfo), ClusterData.ClusterCount), TEXT("Hair.CulledCompactedIndexBuffer"));
+	}
+#endif
+
+	/// Initialise indirect buffers to be setup during the culling process
+	{
+		FHairIndBufferClearCS::FParameters* Parameters = GraphBuilder.AllocParameters<FHairIndBufferClearCS::FParameters>();
+		Parameters->DispatchIndirectParametersClusterCount = GraphBuilder.CreateUAV(DispatchIndirectParametersClusterCount);
+		Parameters->DrawIndirectParameters = DrawIndirectParametersBuffer.UAV;
+
+		FHairIndBufferClearCS::FPermutationDomain Permutation;
+		Permutation.Set<FHairIndBufferClearCS::FSetIndirectDraw>(false);
+		TShaderMapRef<FHairIndBufferClearCS> ComputeShader(ShaderMap, Permutation);
+		FComputeShaderUtils::AddPass(
+			GraphBuilder,
+			RDG_EVENT_NAME("BufferClearCS"),
+			ComputeShader,
+			Parameters,
+			FIntVector(1, 1, 1));
+	}
+
+	/// Cull cluster, generate indirect dispatch and prepare data to expand index buffer
+	{
+		const bool bClusterCullingFrozenCamera = GHairStrandsClusterCullingFreezeCamera > 0;
+		if (!bClusterCullingFrozenCamera)
+		{
+			CapturedCameraWorldPos = View.ViewMatrices.GetViewOrigin();
+			CapturedWorldToClipMatrix = View.ViewMatrices.GetViewProjectionMatrix();
+			CapturedProjMatrix = View.ViewMatrices.GetProjectionMatrix();
+		}
+
+		float ForceLOD = -1;
+		bool bIsVisible = true;
+		if (GHairStrandsClusterForceLOD >= 0)
+		{
+			ForceLOD = GHairStrandsClusterForceLOD;
+		}
+		else if (ClusterData.LODIndex >= 0) // CPU-driven LOD selection
+		{
+			ForceLOD = ClusterData.LODIndex;
+			bIsVisible = ClusterData.bVisible;
+		}
+
+		FHairClusterCullingCS::FParameters* Parameters = GraphBuilder.AllocParameters<FHairClusterCullingCS::FParameters>();
+		Parameters->ProjectionMatrix = FMatrix44f(CapturedProjMatrix);			// LWC_TODO: Precision loss
+		Parameters->CameraWorldPos = FVector3f(CapturedCameraWorldPos);
+		Parameters->WorldToClipMatrix = FMatrix44f(CapturedWorldToClipMatrix);
+		Parameters->ClusterCount = ClusterData.ClusterCount;
+		Parameters->LODForcedIndex = ForceLOD;
+		Parameters->LODBias = ClusterData.LODBias;
+		Parameters->bIsHairGroupVisible = bIsVisible ? 1 : 0;
+
+		Parameters->NumConvexHullPlanes = View.ViewFrustum.Planes.Num();
+		check(Parameters->NumConvexHullPlanes <= 6);
+		for (uint32 i = 0; i < Parameters->NumConvexHullPlanes; ++i)
+		{
+			const FPlane& Plane = View.ViewFrustum.Planes[i];
+			Parameters->ViewFrustumConvexHull[i] = FVector4f(Plane.X, Plane.Y, Plane.Z, Plane.W); // LWC_TODO: precision loss
+		}
+
+		Parameters->ClusterAABBBuffer = RegisterAsSRV(GraphBuilder, *ClusterData.ClusterAABBBuffer);
+		Parameters->ClusterInfoBuffer = RegisterAsSRV(GraphBuilder, *ClusterData.ClusterInfoBuffer);
+		Parameters->ClusterLODInfoBuffer = RegisterAsSRV(GraphBuilder, *ClusterData.ClusterLODInfoBuffer);
+
+		Parameters->GlobalClusterIdBuffer = GraphBuilder.CreateUAV(GlobalClusterIdBuffer, PF_R32_UINT);
+		Parameters->GlobalIndexStartBuffer = GraphBuilder.CreateUAV(GlobalIndexStartBuffer, PF_R32_UINT);
+		Parameters->GlobalIndexCountBuffer = GraphBuilder.CreateUAV(GlobalIndexCountBuffer, PF_R32_UINT);
+		Parameters->GlobalRadiusScaleBuffer = GraphBuilder.CreateUAV(GlobalRadiusScaleBuffer, PF_R32_FLOAT);
+
+		Parameters->DispatchIndirectParametersClusterCount = GraphBuilder.CreateUAV(DispatchIndirectParametersClusterCount);
+		Parameters->DrawIndirectParameters = DrawIndirectParametersBuffer.UAV;
+
+#if WITH_EDITOR
+		Parameters->ClusterDebugInfoBuffer = GraphBuilder.CreateUAV(ClusterDebugInfoBuffer, PF_R32_SINT);
+#endif
+
+		FHairClusterCullingCS::FPermutationDomain Permutation;
+		Permutation.Set<FHairClusterCullingCS::FDebugAABBBuffer>(bClusterDebugAABBBuffer ? 1 : 0);
+		TShaderMapRef<FHairClusterCullingCS> ComputeShader(ShaderMap, Permutation);
+		const FIntVector DispatchCount = DispatchCount.DivideAndRoundUp(FIntVector(ClusterData.ClusterCount, 1, 1), FIntVector(64, 1, 1));
+		FComputeShaderUtils::AddPass(
+			GraphBuilder,
+			RDG_EVENT_NAME("ClusterCullingCS"),
+			ComputeShader,
+			Parameters,
+			DispatchCount);
+	}
+
+	/// Prepare some indirect draw buffers for specific compute group size
+	const uint32 GroupSize1D = GetVendorOptimalGroupSize1D();
+	{
+		FMainClusterCullingPrepareIndirectDrawsCS::FParameters* Parameters = GraphBuilder.AllocParameters<FMainClusterCullingPrepareIndirectDrawsCS::FParameters>();
+		Parameters->GroupSize1D = GroupSize1D;
+		Parameters->DispatchIndirectParametersClusterCount = GraphBuilder.CreateSRV(DispatchIndirectParametersClusterCount, PF_R32_UINT);
+		Parameters->DispatchIndirectParametersClusterCount1D = GraphBuilder.CreateUAV(DispatchIndirectParametersClusterCount1D, PF_R32_UINT);
+		Parameters->DispatchIndirectParametersClusterCount2D = GraphBuilder.CreateUAV(DispatchIndirectParametersClusterCount2D, PF_R32_UINT);
+		Parameters->DispatchIndirectParametersClusterCountDiv512 = GraphBuilder.CreateUAV(DispatchIndirectParametersClusterCountDiv512, PF_R32_UINT);
+		Parameters->DispatchIndirectParametersClusterCountDiv512Div512 = GraphBuilder.CreateUAV(DispatchIndirectParametersClusterCountDiv512Div512, PF_R32_UINT);
+
+		TShaderMapRef<FMainClusterCullingPrepareIndirectDrawsCS> ComputeShader(ShaderMap);
+		FComputeShaderUtils::AddPass(
+			GraphBuilder,
+			RDG_EVENT_NAME("PrepareIndirectDrawsCS"),
+			ComputeShader,
+			Parameters,
+			FIntVector(2, 1, 1));
+	}
+
+	/// local prefix sum per 512 block
+	{
+		FHairClusterCullingLocalBlockPreFixSumCS::FParameters* Parameters = GraphBuilder.AllocParameters<FHairClusterCullingLocalBlockPreFixSumCS::FParameters>();
+		Parameters->DispatchIndirectParametersClusterCountDiv512 = DispatchIndirectParametersClusterCountDiv512;
+		Parameters->DispatchIndirectParametersClusterCount = GraphBuilder.CreateSRV(DispatchIndirectParametersClusterCount);
+		Parameters->GlobalIndexCountBuffer = GraphBuilder.CreateSRV(GlobalIndexCountBuffer, PF_R32_UINT);
+		Parameters->PerBlocklIndexCountPreFixSumBuffer = GraphBuilder.CreateUAV(PerBlocklIndexCountPreFixSumBuffer, PF_R32G32_UINT);
+		Parameters->PerBlocklTotalIndexCountBuffer = GraphBuilder.CreateUAV(PerBlocklTotalIndexCountBuffer, PF_R32_UINT);
+
+		TShaderMapRef<FHairClusterCullingLocalBlockPreFixSumCS> ComputeShader(ShaderMap);
+		FComputeShaderUtils::AddPass(
+			GraphBuilder,
+			RDG_EVENT_NAME("WithinBlockIndexCountPreFixSumCS"),
+			ComputeShader,
+			Parameters,
+			DispatchIndirectParametersClusterCountDiv512, 0); // FIX ME, this could get over 65535
+		check(ClusterData.ClusterCount / 512u <= uint32(GRHIMaxDispatchThreadGroupsPerDimension.X));
+	}
+
+	/// Prefix sum on the total index count per block of 512
+	{
+		FHairClusterCullingPreFixSumCS::FParameters* Parameters = GraphBuilder.AllocParameters<FHairClusterCullingPreFixSumCS::FParameters>();
+		Parameters->DispatchIndirectParameters = DispatchIndirectParametersClusterCountDiv512Div512;
+		Parameters->DispatchIndirectParametersClusterCount = GraphBuilder.CreateSRV(DispatchIndirectParametersClusterCountDiv512, PF_R32_UINT);
+		Parameters->GlobalIndexCountBuffer = GraphBuilder.CreateSRV(PerBlocklTotalIndexCountBuffer, PF_R32_UINT);
+		Parameters->GlobalIndexCountPreFixSumBuffer = GraphBuilder.CreateUAV(PerBlocklTotalIndexCountPreFixSumBuffer, PF_R32G32_UINT);
+
+		TShaderMapRef<FHairClusterCullingPreFixSumCS> ComputeShader(ShaderMap);
+		FComputeShaderUtils::AddPass(
+			GraphBuilder,
+			RDG_EVENT_NAME("BlockIndexCountPreFixSumCS"),
+			ComputeShader,
+			Parameters,
+			DispatchIndirectParametersClusterCountDiv512Div512, 0); // FIX ME, this could get over 65535
+		check(ClusterData.ClusterCount / (512u*512u) <= uint32(GRHIMaxDispatchThreadGroupsPerDimension.X));
+	}
+
+	/// Compact to VertexId buffer using hierarchical binary search or splatting
+	{
+		check(ClusterData.GetCulledVertexIdBuffer());
+		check(ClusterData.GetCulledVertexRadiusScaleBuffer());
+		
+		FRDGImportedBuffer ClusterVertexIdBuffer			= Register(GraphBuilder, *ClusterData.ClusterVertexIdBuffer, ERDGImportedBufferFlags::CreateSRV);
+		FRDGImportedBuffer CulledCompactedIndexBuffer		= Register(GraphBuilder, *ClusterData.GetCulledVertexIdBuffer(), ERDGImportedBufferFlags::CreateUAV);
+		FRDGImportedBuffer CulledCompactedRadiusScaleBuffer	= Register(GraphBuilder, *ClusterData.GetCulledVertexRadiusScaleBuffer(), ERDGImportedBufferFlags::CreateUAV);
+
+		FHairClusterCullingCompactVertexIdsLocalBlockCS::FParameters* Parameters = GraphBuilder.AllocParameters<FHairClusterCullingCompactVertexIdsLocalBlockCS::FParameters>();
+
+		Parameters->DispatchIndirectParametersClusterCount = GraphBuilder.CreateSRV(DispatchIndirectParametersClusterCount, PF_R32_UINT);
+		Parameters->DispatchIndirectParametersClusterCount2D = GraphBuilder.CreateSRV(DispatchIndirectParametersClusterCount2D, PF_R32_UINT);
+		Parameters->DispatchIndirectParametersClusterCountDiv512 = GraphBuilder.CreateSRV(DispatchIndirectParametersClusterCountDiv512, PF_R32_UINT);
+
+		Parameters->PerBlocklIndexCountPreFixSumBuffer = GraphBuilder.CreateSRV(PerBlocklIndexCountPreFixSumBuffer, PF_R32G32_UINT);
+		Parameters->PerBlocklTotalIndexCountPreFixSumBuffer = GraphBuilder.CreateSRV(PerBlocklTotalIndexCountPreFixSumBuffer, PF_R32G32_UINT);
+
+		Parameters->GlobalIndexStartBuffer = GraphBuilder.CreateSRV(GlobalIndexStartBuffer, PF_R32_UINT);
+		Parameters->GlobalIndexCountBuffer = GraphBuilder.CreateSRV(GlobalIndexCountBuffer, PF_R32_UINT);
+		Parameters->GlobalRadiusScaleBuffer = GraphBuilder.CreateSRV(GlobalRadiusScaleBuffer, PF_R32_FLOAT);
+		Parameters->ClusterVertexIdBuffer = ClusterVertexIdBuffer.SRV;
+
+		
+		Parameters->CulledCompactedIndexBuffer = CulledCompactedIndexBuffer.UAV;
+		Parameters->CulledCompactedRadiusScaleBuffer = CulledCompactedRadiusScaleBuffer.UAV;
+
+		Parameters->DispatchIndirectParametersBuffer = DispatchIndirectParametersClusterCount2D;
+		TShaderMapRef<FHairClusterCullingCompactVertexIdsLocalBlockCS> ComputeShader(ShaderMap);
+		FComputeShaderUtils::AddPass(
+			GraphBuilder,
+			RDG_EVENT_NAME("SplatCompactVertexIdsCS"),
+			ComputeShader,
+			Parameters,
+			DispatchIndirectParametersClusterCount2D, 0); // DispatchIndirectParametersClusterCount2D is used to avoid having any dispatch dimension going above 65535.
+
+		GraphBuilder.SetBufferAccessFinal(CulledCompactedIndexBuffer.Buffer, ERHIAccess::SRVMask);
+		GraphBuilder.SetBufferAccessFinal(CulledCompactedRadiusScaleBuffer.Buffer, ERHIAccess::SRVMask);
+	}
+
+	{
+		ClusterData.ClusterIdBuffer = GraphBuilder.ConvertToExternalBuffer(GlobalClusterIdBuffer);
+		ClusterData.ClusterIndexOffsetBuffer = GraphBuilder.ConvertToExternalBuffer(GlobalIndexStartBuffer);
+		ClusterData.ClusterIndexCountBuffer = GraphBuilder.ConvertToExternalBuffer(GlobalIndexCountBuffer);
+		ClusterData.CulledCluster2DIndirectArgsBuffer = DispatchIndirectParametersClusterCount2D;
+		ClusterData.CulledCluster1DIndirectArgsBuffer = DispatchIndirectParametersClusterCount1D;
+		ClusterData.CulledClusterCountBuffer = DispatchIndirectParametersClusterCount;
+		ClusterData.GroupSize1D = GroupSize1D;
+	}
+
+	/// Prepare some indirect dispatch for compute raster visibility buffers
+	{
+		FMainClusterCullingPrepareIndirectDispatchCS::FParameters* Parameters = GraphBuilder.AllocParameters<FMainClusterCullingPrepareIndirectDispatchCS::FParameters>();
+		Parameters->DrawIndirectBuffer = DrawIndirectParametersBuffer.UAV;
+		Parameters->DispatchIndirectBuffer = DrawIndirectParametersRasterComputeBuffer.UAV;
+		
+		TShaderMapRef<FMainClusterCullingPrepareIndirectDispatchCS> ComputeShader(ShaderMap);
+		FComputeShaderUtils::AddPass(
+			GraphBuilder,
+			RDG_EVENT_NAME("PrepareIndirectDispatchCS"),
+			ComputeShader,
+			Parameters,
+			FIntVector(1, 1, 1));
+	}
+
+	// Should this be move onto the culling result?
+#if WITH_EDITOR
+	if (bClusterDebugAABBBuffer)
+	{
+		ClusterData.ClusterDebugInfoBuffer = GraphBuilder.ConvertToExternalBuffer(ClusterDebugInfoBuffer);
+	}
+#endif
+
+	GraphBuilder.SetBufferAccessFinal(DrawIndirectParametersBuffer.Buffer, ERHIAccess::IndirectArgs | ERHIAccess::SRVMask);
+	GraphBuilder.SetBufferAccessFinal(DrawIndirectParametersRasterComputeBuffer.Buffer, ERHIAccess::IndirectArgs | ERHIAccess::SRVMask);
+
+	ClusterData.SetCullingResultAvailable(true);
+}
+
+///////////////////////////////////////////////////////////////////////////////////////////////////
+
+static void AddClusterResetLod0(
+	FRDGBuilder& GraphBuilder,
+	FGlobalShaderMap* ShaderMap,
+	FHairStrandClusterData::FHairGroup& ClusterData)
+{
+	// Set as culling result not available
+	ClusterData.SetCullingResultAvailable(false);
+
+	FRDGImportedBuffer IndirectBuffer = Register(GraphBuilder, ClusterData.HairGroupPublicPtr->GetDrawIndirectBuffer(), ERDGImportedBufferFlags::CreateViews);
+
+	// Initialise indirect buffers to entire lod 0 dispatch
+	FHairIndBufferClearCS::FParameters* Parameters = GraphBuilder.AllocParameters<FHairIndBufferClearCS::FParameters>();
+	Parameters->DrawIndirectParameters = IndirectBuffer.UAV;
+	Parameters->VertexCountPerInstance = ClusterData.HairGroupPublicPtr->GetActiveStrandsPointCount() * HAIR_POINT_TO_VERTEX;
+
+	FHairIndBufferClearCS::FPermutationDomain Permutation;
+	Permutation.Set<FHairIndBufferClearCS::FSetIndirectDraw>(true);
+	TShaderMapRef<FHairIndBufferClearCS> ComputeShader(ShaderMap, Permutation);
+	FComputeShaderUtils::AddPass(
+		GraphBuilder,
+		RDG_EVENT_NAME("BufferClearCS"),
+		ComputeShader,
+		Parameters,
+		FIntVector(1, 1, 1));
+
+	GraphBuilder.SetBufferAccessFinal(IndirectBuffer.Buffer, ERHIAccess::IndirectArgs);
+}
+
+///////////////////////////////////////////////////////////////////////////////////////////////////
+
+void ComputeHairStrandsClustersCulling(
+	FRDGBuilder& GraphBuilder,
+	FGlobalShaderMap& ShaderMap,
+	const TArray<const FSceneView*>& Views, 
+	FHairStrandClusterData& ClusterDatas)
+{
+	DECLARE_GPU_STAT(HairStrandsClusterCulling);
+	RDG_EVENT_SCOPE(GraphBuilder, "HairStrandsClusterCulling");
+	TRACE_CPUPROFILER_EVENT_SCOPE(ComputeHairStrandsClustersCulling);
+	RDG_GPU_STAT_SCOPE(GraphBuilder, HairStrandsClusterCulling);
+
+	for (const FSceneView* View : Views)
+	{
+		// TODO use compute overlap (will need to split AddClusterCullingPass)
+		for (FHairStrandClusterData::FHairGroup& ClusterData : ClusterDatas.HairGroups)
+		{
+			AddClusterResetLod0(GraphBuilder, &ShaderMap, ClusterData);
+		}
+
+		for (FHairStrandClusterData::FHairGroup& ClusterData : ClusterDatas.HairGroups)		
+		{
+			AddClusterCullingPass(
+				GraphBuilder,
+				&ShaderMap,
+				*View,
+				ClusterData);
+		}
+	}
+}
+
+///////////////////////////////////////////////////////////////////////////////////////////////////
+
+void AddInstanceToClusterData(FHairGroupInstance* In, FHairStrandClusterData& Out)
+{
+	// Initialize group cluster data for culling by the renderer
+	const int32 ClusterDataGroupIndex = Out.HairGroups.Num();
+	FHairStrandClusterData::FHairGroup& HairGroupCluster = Out.HairGroups.Emplace_GetRef();
+	HairGroupCluster.ClusterCount = In->HairGroupPublicData->GetClusterCount();
+	HairGroupCluster.VertexCount = In->HairGroupPublicData->RestPointCount * HAIR_POINT_TO_VERTEX; // Instance->HairGroupPublicData->GetActiveStrandsPointCount() ?
+	HairGroupCluster.GroupAABBBuffer = &In->HairGroupPublicData->GetGroupAABBBuffer();
+	HairGroupCluster.ClusterAABBBuffer = &In->HairGroupPublicData->GetClusterAABBBuffer();
+
+	HairGroupCluster.ClusterInfoBuffer = &In->Strands.ClusterCullingResource->ClusterInfoBuffer;
+	HairGroupCluster.ClusterLODInfoBuffer = &In->Strands.ClusterCullingResource->ClusterLODInfoBuffer;
+	HairGroupCluster.CurveToClusterIdBuffer = &In->Strands.ClusterCullingResource->CurveToClusterIdBuffer;
+	HairGroupCluster.ClusterVertexIdBuffer = &In->Strands.ClusterCullingResource->ClusterVertexIdBuffer;
+
+	HairGroupCluster.HairGroupPublicPtr = In->HairGroupPublicData;
+	HairGroupCluster.LODBias  = In->HairGroupPublicData->GetLODBias();
+	HairGroupCluster.LODIndex = In->HairGroupPublicData->GetLODIndex();
+	HairGroupCluster.bVisible = In->HairGroupPublicData->GetLODVisibility();
+
+	// These buffer are create during the culling pass
+	// HairGroupCluster.ClusterIdBuffer = nullptr;
+	// HairGroupCluster.ClusterIndexOffsetBuffer = nullptr;
+	// HairGroupCluster.ClusterIndexCountBuffer = nullptr;
+
+	HairGroupCluster.HairGroupPublicPtr->ClusterDataIndex = ClusterDataGroupIndex;
+}
\ No newline at end of file
diff --git a/Engine/Plugins/Runtime/HairStrands/Source/HairStrandsCore/Private/HairStrandsClusterCulling.h b/Engine/Plugins/Runtime/HairStrands/Source/HairStrandsCore/Private/HairStrandsClusterCulling.h
new file mode 100644
--- /dev/null
+++ b/Engine/Plugins/Runtime/HairStrands/Source/HairStrandsCore/Private/HairStrandsClusterCulling.h
@@ -0,0 +1,28 @@
+// Copyright Epic Games, Inc. All Rights Reserved.
+
+#pragma once
+
+#include "CoreMinimal.h"
+#include "RendererInterface.h"
+#include "HairStrandsInterface.h"
+
+struct FHairCullingParams
+{
+	bool bCullingProcessSkipped	= false;
+};
+
+struct FHairGroupInstance;
+struct FHairStrandClusterData;
+class FViewInfo;
+class FGlobalShaderMap;
+class FRDGBuilder;
+
+void AddInstanceToClusterData(
+	FHairGroupInstance* In,
+	FHairStrandClusterData& Out);
+
+void ComputeHairStrandsClustersCulling(
+	FRDGBuilder& GraphBuilder,
+	FGlobalShaderMap& ShaderMap,
+	const TArray<const FSceneView*>& Views,
+	FHairStrandClusterData& ClusterDatas);
diff --git a/Engine/Plugins/Runtime/HairStrands/Source/HairStrandsCore/Private/HairStrandsInterpolation.cpp b/Engine/Plugins/Runtime/HairStrands/Source/HairStrandsCore/Private/HairStrandsInterpolation.cpp
--- a/Engine/Plugins/Runtime/HairStrands/Source/HairStrandsCore/Private/HairStrandsInterpolation.cpp
+++ b/Engine/Plugins/Runtime/HairStrands/Source/HairStrandsCore/Private/HairStrandsInterpolation.cpp
@@ -1532,35 +1532,6 @@ static void BuildHairAccelerationStructure_Meshes(FRHICommandList& RHICmdList,
 }
 #endif // RHI_RAYTRACING
 
-
-void RegisterClusterData(FHairGroupInstance* Instance, FHairStrandClusterData* InClusterData)
-{
-	// Initialize group cluster data for culling by the renderer
-	const int32 ClusterDataGroupIndex = InClusterData->HairGroups.Num();
-	FHairStrandClusterData::FHairGroup& HairGroupCluster = InClusterData->HairGroups.Emplace_GetRef();
-	HairGroupCluster.ClusterCount = Instance->HairGroupPublicData->GetClusterCount();
-	HairGroupCluster.VertexCount = Instance->HairGroupPublicData->RestPointCount * HAIR_POINT_TO_VERTEX; // Instance->HairGroupPublicData->GetActiveStrandsPointCount() ?
-	HairGroupCluster.GroupAABBBuffer = &Instance->HairGroupPublicData->GetGroupAABBBuffer();
-	HairGroupCluster.ClusterAABBBuffer = &Instance->HairGroupPublicData->GetClusterAABBBuffer();
-
-	HairGroupCluster.ClusterInfoBuffer = &Instance->Strands.ClusterCullingResource->ClusterInfoBuffer;
-	HairGroupCluster.ClusterLODInfoBuffer = &Instance->Strands.ClusterCullingResource->ClusterLODInfoBuffer;
-	HairGroupCluster.CurveToClusterIdBuffer = &Instance->Strands.ClusterCullingResource->CurveToClusterIdBuffer;
-	HairGroupCluster.ClusterVertexIdBuffer = &Instance->Strands.ClusterCullingResource->ClusterVertexIdBuffer;
-
-	HairGroupCluster.HairGroupPublicPtr = Instance->HairGroupPublicData;
-	HairGroupCluster.LODBias  = Instance->HairGroupPublicData->GetLODBias();
-	HairGroupCluster.LODIndex = Instance->HairGroupPublicData->GetLODIndex();
-	HairGroupCluster.bVisible = Instance->HairGroupPublicData->GetLODVisibility();
-
-	// These buffer are create during the culling pass
-	// HairGroupCluster.ClusterIdBuffer = nullptr;
-	// HairGroupCluster.ClusterIndexOffsetBuffer = nullptr;
-	// HairGroupCluster.ClusterIndexCountBuffer = nullptr;
-
-	HairGroupCluster.HairGroupPublicPtr->ClusterDataIndex = ClusterDataGroupIndex;
-}
-
 static void ConvertHairStrandsVFParameters(
 	FRDGBuilder* GraphBuilder,
 	FRDGImportedBuffer& OutBuffer,
diff --git a/Engine/Plugins/Runtime/HairStrands/Source/HairStrandsCore/Public/HairStrandsInterpolation.h b/Engine/Plugins/Runtime/HairStrands/Source/HairStrandsCore/Public/HairStrandsInterpolation.h
--- a/Engine/Plugins/Runtime/HairStrands/Source/HairStrandsCore/Public/HairStrandsInterpolation.h
+++ b/Engine/Plugins/Runtime/HairStrands/Source/HairStrandsCore/Public/HairStrandsInterpolation.h
@@ -51,8 +51,4 @@ void ComputeHairStrandsInterpolation(
 	int32 LODIndex,
 	FHairStrandClusterData* ClusterData);
 
-void RegisterClusterData(
-	FHairGroupInstance* Instance,
-	FHairStrandClusterData* InClusterData);
-
 HAIRSTRANDSCORE_API void ComputeInterpolationWeights(UGroomBindingAsset* BindingAsset, FSkeletalMeshRenderData* TargetRenderData, TArray<FRWBuffer>& TransferedPositions);
diff --git a/Engine/Source/Runtime/Renderer/Private/DeferredShadingRenderer.cpp b/Engine/Source/Runtime/Renderer/Private/DeferredShadingRenderer.cpp
--- a/Engine/Source/Runtime/Renderer/Private/DeferredShadingRenderer.cpp
+++ b/Engine/Source/Runtime/Renderer/Private/DeferredShadingRenderer.cpp
@@ -3012,7 +3012,7 @@ void FDeferredShadingSceneRenderer::Render(FRDGBuilder& GraphBuilder)
 	FHairStrandsBookmarkParameters& HairStrandsBookmarkParameters = *GraphBuilder.AllocObject<FHairStrandsBookmarkParameters>();
 	if (IsHairStrandsEnabled(EHairStrandsShaderType::All, Scene->GetShaderPlatform()) && RendererOutput == ERendererOutput::FinalSceneColor)
 	{
-		HairStrandsBookmarkParameters = CreateHairStrandsBookmarkParameters(Scene, Views[0]);
+		HairStrandsBookmarkParameters = CreateHairStrandsBookmarkParameters(Scene, Views, AllFamilyViews);
 		RunHairStrandsBookmark(GraphBuilder, EHairStrandsBookmark::ProcessTasks, HairStrandsBookmarkParameters);
 
 		// Interpolation needs to happen after the skin cache run as there is a dependency 
@@ -3020,15 +3020,6 @@ void FDeferredShadingSceneRenderer::Render(FRDGBuilder& GraphBuilder)
 		const bool bRunHairStrands = HairStrandsBookmarkParameters.HasInstances() && (Views.Num() > 0);
 		if (bRunHairStrands)
 		{
-			if (IsHairStrandsEnabled(EHairStrandsShaderType::Strands, Scene->GetShaderPlatform()))
-			{
-				RunHairStrandsBookmark(GraphBuilder, EHairStrandsBookmark::ProcessGatherCluster, HairStrandsBookmarkParameters);
-
-				FHairCullingParams CullingParams;
-				CullingParams.bCullingProcessSkipped = false;
-				ComputeHairStrandsClustersCulling(GraphBuilder, *HairStrandsBookmarkParameters.ShaderMap, Views, CullingParams, HairStrandsBookmarkParameters.HairClusterData);
-			}
-
 			RunHairStrandsBookmark(GraphBuilder, EHairStrandsBookmark::ProcessStrandsInterpolation, HairStrandsBookmarkParameters);
 		}
 		else
@@ -4186,7 +4177,7 @@ void FDeferredShadingSceneRenderer::Render(FRDGBuilder& GraphBuilder)
 		{
 			if (HairStrandsBookmarkParameters.HasInstances())
 			{
-				RenderHairStrandsDebugInfo(GraphBuilder, Scene, Views, HairStrandsBookmarkParameters.HairClusterData, SceneTextures.Color.Target, SceneTextures.Depth.Target);
+				RenderHairStrandsDebugInfo(GraphBuilder, Scene, Views, SceneTextures.Color.Target, SceneTextures.Depth.Target);
 			}
 		}
 
diff --git a/Engine/Source/Runtime/Renderer/Private/HairStrands/HairStrandsClusters.cpp b/Engine/Source/Runtime/Renderer/Private/HairStrands/HairStrandsClusters.cpp
deleted file mode 100644
diff --git a/Engine/Source/Runtime/Renderer/Private/HairStrands/HairStrandsClusters.h b/Engine/Source/Runtime/Renderer/Private/HairStrands/HairStrandsClusters.h
deleted file mode 100644
diff --git a/Engine/Source/Runtime/Renderer/Private/HairStrands/HairStrandsDebug.cpp b/Engine/Source/Runtime/Renderer/Private/HairStrands/HairStrandsDebug.cpp
--- a/Engine/Source/Runtime/Renderer/Private/HairStrands/HairStrandsDebug.cpp
+++ b/Engine/Source/Runtime/Renderer/Private/HairStrands/HairStrandsDebug.cpp
@@ -1063,126 +1063,12 @@ class FHairVisibilityDebugPPLLCS : public FGlobalShader
 };
 IMPLEMENT_GLOBAL_SHADER(FHairVisibilityDebugPPLLCS, "/Engine/Private/HairStrands/HairStrandsDebug.usf", "VisibilityDebugPPLLCS", SF_Compute);
 
-///////////////////////////////////////////////////////////////////////////////////////////////////
-
-class FDrawDebugClusterAABBCS : public FGlobalShader
-{
-	DECLARE_GLOBAL_SHADER(FDrawDebugClusterAABBCS);
-	SHADER_USE_PARAMETER_STRUCT(FDrawDebugClusterAABBCS, FGlobalShader);
-
-	class FDebugAABBBuffer : SHADER_PERMUTATION_INT("PERMUTATION_DEBUGAABBBUFFER", 2);
-	using FPermutationDomain = TShaderPermutationDomain<FDebugAABBBuffer>;
-
-	BEGIN_SHADER_PARAMETER_STRUCT(FParameters, )
-		SHADER_PARAMETER_STRUCT_REF(FViewUniformShaderParameters, ViewUniformBuffer)
-		SHADER_PARAMETER_SRV(Buffer, ClusterAABBBuffer)
-		SHADER_PARAMETER_SRV(Buffer, GroupAABBBuffer)
-		SHADER_PARAMETER_RDG_BUFFER_SRV(Buffer, CulledDispatchIndirectParametersClusterCountBuffer)
-		SHADER_PARAMETER_RDG_BUFFER_SRV(StructuredBuffer, ClusterDebugInfoBuffer)
-		SHADER_PARAMETER_SRV(Buffer, CulledDrawIndirectParameters)
-		SHADER_PARAMETER(uint32, ClusterCount)
-		SHADER_PARAMETER(uint32, TriangleCount)
-		SHADER_PARAMETER(uint32, HairGroupId)
-		SHADER_PARAMETER(int32, ClusterDebugMode)
-		SHADER_PARAMETER_STRUCT_INCLUDE(ShaderPrint::FShaderParameters, ShaderPrintParameters)
-	END_SHADER_PARAMETER_STRUCT()
-
-public:
-	static bool ShouldCompilePermutation(const FGlobalShaderPermutationParameters& Parameters) { return IsHairStrandsSupported(EHairStrandsShaderType::Tool, Parameters.Platform); }
-	static void ModifyCompilationEnvironment(const FGlobalShaderPermutationParameters& Parameters, FShaderCompilerEnvironment& OutEnvironment)
-	{
-		FGlobalShader::ModifyCompilationEnvironment(Parameters, OutEnvironment);
-		OutEnvironment.SetDefine(TEXT("SHADER_DRAWDEBUGAABB"), 1);
-
-		// Skip optimization for avoiding long compilation time due to large UAV writes
-		OutEnvironment.CompilerFlags.Add(CFLAG_Debug);
-	}
-};
-
-IMPLEMENT_GLOBAL_SHADER(FDrawDebugClusterAABBCS, "/Engine/Private/HairStrands/HairStrandsClusterCulling.usf", "MainDrawDebugAABBCS", SF_Compute);
-
-static void AddDrawDebugClusterPass(
-	FRDGBuilder& GraphBuilder,
-	const FHairStrandClusterData& HairClusterData,
-	const FViewInfo& View)
-{
-	const EGroomViewMode ViewMode = GetGroomViewMode(View);
-	const bool bDebugEnable = ViewMode == EGroomViewMode::ClusterAABB || ViewMode == EGroomViewMode::Cluster;
-	if (!bDebugEnable)
-	{
-		return;
-	}
-	
-	if (!TryEnableShaderDrawAndShaderPrint(View, 64000u, 2000))
-	{
-		return;
-	}
-
-	for (const FHairStrandsMacroGroupData& MacroGroupData : View.HairStrandsViewData.MacroGroupDatas)
-	{
-		const bool bDebugAABB = GetGroomViewMode(View) == EGroomViewMode::ClusterAABB;
-
-		for (const FHairStrandsMacroGroupData::PrimitiveInfo& PrimitiveInfo : MacroGroupData.PrimitivesInfos)
-		{
-			if (!PrimitiveInfo.Mesh || PrimitiveInfo.Mesh->Elements.Num() == 0)
-			{
-				continue;
-			}
-
-			for (int DataIndex = 0; DataIndex < HairClusterData.HairGroups.Num(); ++DataIndex)
-			{
-				const FHairStrandClusterData::FHairGroup& HairGroupClusters = HairClusterData.HairGroups[DataIndex];
-
-				// Find a better/less hacky way
-				if (PrimitiveInfo.PublicDataPtr != HairGroupClusters.HairGroupPublicPtr)
-				{
-					continue;
-				}
-
-				if (ShaderPrint::IsEnabled(View.ShaderPrintData) && HairGroupClusters.CulledClusterCountBuffer)
-				{
-					FRDGExternalBuffer& DrawIndirectBuffer = HairGroupClusters.HairGroupPublicPtr->GetDrawIndirectBuffer();
-
-					FDrawDebugClusterAABBCS::FPermutationDomain Permutation;
-					Permutation.Set<FDrawDebugClusterAABBCS::FDebugAABBBuffer>(bDebugAABB ? 1 : 0);
-					TShaderMapRef<FDrawDebugClusterAABBCS> ComputeShader(View.ShaderMap, Permutation);
-
-					FDrawDebugClusterAABBCS::FParameters* Parameters = GraphBuilder.AllocParameters<FDrawDebugClusterAABBCS::FParameters>();
-					Parameters->ViewUniformBuffer = View.ViewUniformBuffer;
-					Parameters->ClusterCount = HairGroupClusters.ClusterCount;
-					Parameters->TriangleCount = HairGroupClusters.VertexCount * 2; // VertexCount is actually the number of control points
-					Parameters->HairGroupId = DataIndex;
-					Parameters->ClusterDebugMode = ViewMode == EGroomViewMode::Cluster ? 1 : (ViewMode == EGroomViewMode::ClusterAABB ? 2 : 0);
-					Parameters->ClusterAABBBuffer = HairGroupClusters.ClusterAABBBuffer->SRV;
-					Parameters->CulledDispatchIndirectParametersClusterCountBuffer = GraphBuilder.CreateSRV(HairGroupClusters.CulledClusterCountBuffer, EPixelFormat::PF_R32_UINT);
-					Parameters->CulledDrawIndirectParameters = DrawIndirectBuffer.SRV;
-					Parameters->GroupAABBBuffer = HairGroupClusters.GroupAABBBuffer->SRV;
-
-					if (HairGroupClusters.ClusterDebugInfoBuffer && bDebugAABB)
-					{
-						FRDGBufferRef ClusterDebugInfoBuffer = GraphBuilder.RegisterExternalBuffer(HairGroupClusters.ClusterDebugInfoBuffer);
-						Parameters->ClusterDebugInfoBuffer = GraphBuilder.CreateSRV(ClusterDebugInfoBuffer);
-					}
-					ShaderPrint::SetParameters(GraphBuilder, View.ShaderPrintData, Parameters->ShaderPrintParameters);
-
-					check(Parameters->ClusterCount / 64u <= uint32(GRHIMaxDispatchThreadGroupsPerDimension.X));
-					const FIntVector DispatchCount = DispatchCount.DivideAndRoundUp(FIntVector(Parameters->ClusterCount, 1, 1), FIntVector(64, 1, 1));// FIX ME, this could get over 65535
-					FComputeShaderUtils::AddPass(
-						GraphBuilder, RDG_EVENT_NAME("DrawDebugClusterAABB"),
-						ComputeShader, Parameters, DispatchCount);
-				}
-			}
-		}
-	}
-}
-
 ///////////////////////////////////////////////////////////////////////////////////////////////////
 uint32 GetHairStrandsMeanSamplePerPixel(EShaderPlatform In);
 static void InternalRenderHairStrandsDebugInfo(
 	FRDGBuilder& GraphBuilder,
 	FScene* Scene,
 	FViewInfo& View,
-	const struct FHairStrandClusterData& HairClusterData,
 	FRDGTextureRef SceneColorTexture,
 	FRDGTextureRef SceneDepthTexture)
 {
@@ -1365,11 +1251,6 @@ static void InternalRenderHairStrandsDebugInfo(
 		FComputeShaderUtils::AddPass(GraphBuilder, RDG_EVENT_NAME("HairStrands::PPLLDebug"), ComputeShader, PassParameters, FIntVector::DivideAndRoundUp(TextureSize, FIntVector(8, 8, 1)));
 	}
 
-	if (ViewMode == EGroomViewMode::Cluster || ViewMode == EGroomViewMode::ClusterAABB)
-	{
-		AddDrawDebugClusterPass(GraphBuilder, HairClusterData, View);
-	}
-
 	if (ViewMode != EGroomViewMode::None)
 	{
 		AddDrawCanvasPass(GraphBuilder, {}, View, SceneColor, [&View, YStep, ViewMode](FCanvas& Canvas)
@@ -1391,13 +1272,12 @@ void RenderHairStrandsDebugInfo(
 	FRDGBuilder& GraphBuilder,
 	FScene* Scene,
 	TArrayView<FViewInfo> Views,
-	const struct FHairStrandClusterData& HairClusterData,
 	FRDGTextureRef SceneColorTexture,
 	FRDGTextureRef SceneDepthTexture)
 {
 	bool bHasHairData = false;
 	for (FViewInfo& View : Views)
 	{
-		InternalRenderHairStrandsDebugInfo(GraphBuilder, Scene, View, HairClusterData, SceneColorTexture, SceneDepthTexture);
+		InternalRenderHairStrandsDebugInfo(GraphBuilder, Scene, View, SceneColorTexture, SceneDepthTexture);
 	}
 }
diff --git a/Engine/Source/Runtime/Renderer/Private/HairStrands/HairStrandsDebug.h b/Engine/Source/Runtime/Renderer/Private/HairStrands/HairStrandsDebug.h
--- a/Engine/Source/Runtime/Renderer/Private/HairStrands/HairStrandsDebug.h
+++ b/Engine/Source/Runtime/Renderer/Private/HairStrands/HairStrandsDebug.h
@@ -13,6 +13,5 @@ void RenderHairStrandsDebugInfo(
 	FRDGBuilder& GraphBuilder,
 	FScene* Scene, 
 	TArrayView<FViewInfo> Views,
-	const struct FHairStrandClusterData& HairClusterData,
 	FRDGTextureRef SceneColorTexture, 
 	FRDGTextureRef SceneDepthTexture);
diff --git a/Engine/Source/Runtime/Renderer/Private/HairStrands/HairStrandsRendering.cpp b/Engine/Source/Runtime/Renderer/Private/HairStrands/HairStrandsRendering.cpp
--- a/Engine/Source/Runtime/Renderer/Private/HairStrands/HairStrandsRendering.cpp
+++ b/Engine/Source/Runtime/Renderer/Private/HairStrands/HairStrandsRendering.cpp
@@ -6,6 +6,7 @@
 #include "ScenePrivate.h"
 #include "SystemTextures.h"
 #include "RenderGraphUtils.h"
+#include "HairStrands/HairStrandsCluster.h"
 
 static TRDGUniformBufferRef<FHairStrandsViewUniformParameters> InternalCreateHairStrandsViewUniformBuffer(
 	FRDGBuilder& GraphBuilder, 
diff --git a/Engine/Source/Runtime/Renderer/Private/HairStrands/HairStrandsRendering.h b/Engine/Source/Runtime/Renderer/Private/HairStrands/HairStrandsRendering.h
--- a/Engine/Source/Runtime/Renderer/Private/HairStrands/HairStrandsRendering.h
+++ b/Engine/Source/Runtime/Renderer/Private/HairStrands/HairStrandsRendering.h
@@ -10,8 +10,6 @@
 #include "RendererInterface.h"
 #include "Shader.h"
 #include "HairStrandsUtils.h"
-#include "HairStrandsCluster.h"
-#include "HairStrandsClusters.h"
 #include "HairStrandsLUT.h"
 #include "HairStrandsDeepShadow.h"
 #include "HairStrandsVoxelization.h"
diff --git a/Engine/Source/Runtime/Renderer/Public/HairStrandsInterface.h b/Engine/Source/Runtime/Renderer/Public/HairStrandsInterface.h
--- a/Engine/Source/Runtime/Renderer/Public/HairStrandsInterface.h
+++ b/Engine/Source/Runtime/Renderer/Public/HairStrandsInterface.h
@@ -546,9 +546,6 @@ struct FHairStrandsBookmarkParameters
 	bool bHzbRequest = false;
 	uint32 FrameIndex = ~0;
 
-	// Temporary
-	FHairStrandClusterData HairClusterData;
-
 	inline bool HasInstances() const { return Instances != nullptr && Instances->Num() > 0; }
 };
 
