* Change cluster culling to be done in a single pass. This saves about 20% cost of the interpolation work.
* Cluster are not computed based on k first curve to avoid expensive cluster bound computation.

diff --git a/Engine/Plugins/Runtime/HairStrands/Source/HairStrandsCore/Private/GroomBuilder.cpp b/Engine/Plugins/Runtime/HairStrands/Source/HairStrandsCore/Private/GroomBuilder.cpp
--- a/Engine/Plugins/Runtime/HairStrands/Source/HairStrandsCore/Private/GroomBuilder.cpp
+++ b/Engine/Plugins/Runtime/HairStrands/Source/HairStrandsCore/Private/GroomBuilder.cpp
@@ -40,7 +40,7 @@ static FAutoConsoleVariableRef CVarHairGroupIndexBuilder_MaxVoxelResolution(TEXT
 
 FString FGroomBuilder::GetVersion()
 {
-	return TEXT("v8r30");
+	return TEXT("v8r32");
 }
 
 namespace FHairStrandsDecimation
@@ -2990,6 +2990,7 @@ static void BuildClusterData(
 	Out.ClusterInfos.Init(FHairClusterInfo(), Out.ClusterCount);
 	Out.CurveToClusterIds.SetNum(Out.CurveCount);
 	Out.LODInfos.SetNum(LODCount);
+	Out.PointLODs.Init(0, Out.PointCount);
 
 	// Conservative allocation for inserting vertex indices for the various curves LOD
 	uint32* RawClusterVertexIds = new uint32[LODCount * InRenStrandsData.GetNumPoints()];
@@ -3004,6 +3005,7 @@ static void BuildClusterData(
 	TArray<FHairClusterInfo>& LocalClusterInfos = Out.ClusterInfos;
 	TArray<FHairClusterLODInfo>& LocalClusterLODInfos = Out.ClusterLODInfos;
 	TArray<uint32>& LocalCurveToClusterIds = Out.CurveToClusterIds;
+	TArray<uint32>& LocalPointLODs = Out.PointLODs;
 #define USE_PARALLE_FOR 0
 #if USE_PARALLE_FOR
 	ParallelFor(Out.ClusterCount,
@@ -3191,6 +3193,9 @@ static void BuildClusterData(
 						}
 
 						LocalClusterVertexIds.Add(GlobalPointIndex);
+
+						// Store the minimum/coarser LOD at which the current point is visible
+						LocalPointLODs[GlobalPointIndex] = LODIt;
 					}
 				}
 			}
@@ -3311,6 +3316,25 @@ static void BuildClusterBulkData(
 	HairStrandsBuilder::CopyToBulkData<FHairClusterIndexFormat>(Out.Data.CurveToClusterIds, In.CurveToClusterIds);
 	HairStrandsBuilder::CopyToBulkData<FHairClusterIndexFormat>(Out.Data.ClusterVertexIds, In.ClusterVertexIds);
 
+	// Pack point LOD data
+	// |                          32bits                               |
+	// |      4bits    |      4bits    |      4bits    |      4bits    |
+	// | Point0 MinLOD | Point1 MinLOD | Point2 MinLOD | Point3 MinLOD |
+	static_assert(HAIR_POINT_LOD_BIT_COUNT == 4);
+	TArray<uint32> PackedPointLODs;
+	PackedPointLODs.SetNum(FMath::DivideAndRoundUp(uint32(In.PointLODs.Num()), HAIR_POINT_LOD_COUNT_PER_UINT));
+	for (uint32 PackedIt = 0, PackedCount = PackedPointLODs.Num(); PackedIt<PackedCount; ++PackedIt)
+	{
+		uint32 Packed = 0;
+		for (uint32 It = 0; It < HAIR_POINT_LOD_COUNT_PER_UINT; ++It)
+		{
+			const uint32 InIndex = FMath::Min(PackedIt * HAIR_POINT_LOD_COUNT_PER_UINT + It, uint32(In.PointLODs.Num()-1));
+			Packed = Packed | ((In.PointLODs[InIndex] & 0xF) << (It * HAIR_POINT_LOD_BIT_COUNT));
+		}
+		PackedPointLODs[PackedIt] = Packed;
+	}
+	HairStrandsBuilder::CopyToBulkData<FHairClusterIndexFormat>(Out.Data.PointLODs, PackedPointLODs);
+
 	// Pack LODInfo into GPU format
 	{
 		check(uint32(In.ClusterInfos.Num()) == Out.Header.ClusterCount);
diff --git a/Engine/Plugins/Runtime/HairStrands/Source/HairStrandsCore/Private/GroomManager.cpp b/Engine/Plugins/Runtime/HairStrands/Source/HairStrandsCore/Private/GroomManager.cpp
--- a/Engine/Plugins/Runtime/HairStrands/Source/HairStrandsCore/Private/GroomManager.cpp
+++ b/Engine/Plugins/Runtime/HairStrands/Source/HairStrandsCore/Private/GroomManager.cpp
@@ -501,7 +501,7 @@ static void RunHairStrandsInterpolation_Strands(
 				}
 			}
 		}
-		ComputeHairStrandsClustersCulling(GraphBuilder, *ShaderMap, Views, ClusterData);
+		ComputeHairStrandsClustersCulling(GraphBuilder, *ShaderMap, Views, ShaderPrintData, ClusterData);
 
 		// Run cluster debug view here (instead of GroomDebug.h/.cpp, as we need to have the (transient) cluster data 
 		const EGroomViewMode ViewMode = GetGroomViewMode(*View);
diff --git a/Engine/Plugins/Runtime/HairStrands/Source/HairStrandsCore/Private/GroomResources.cpp b/Engine/Plugins/Runtime/HairStrands/Source/HairStrandsCore/Private/GroomResources.cpp
--- a/Engine/Plugins/Runtime/HairStrands/Source/HairStrandsCore/Private/GroomResources.cpp
+++ b/Engine/Plugins/Runtime/HairStrands/Source/HairStrandsCore/Private/GroomResources.cpp
@@ -1200,6 +1200,9 @@ void FHairStrandsClusterCullingResource::InternalAllocate(FRDGBuilder& GraphBuil
 
 	InternalCreateVertexBufferRDG_FromHairBulkData<FHairClusterIndexFormat>(GraphBuilder, BulkData.Data.CurveToClusterIds, BulkData.Header.CurveCount, CurveToClusterIdBuffer, ToHairResourceDebugName(TEXT("Hair.StrandsClusterCulling_CurveToClusterIds"), ResourceName), OwnerName, EHairResourceUsageType::Static);
 	InternalCreateVertexBufferRDG_FromHairBulkData<FHairClusterIndexFormat>(GraphBuilder, BulkData.Data.ClusterVertexIds, BulkData.Header.VertexLODCount, ClusterVertexIdBuffer, ToHairResourceDebugName(TEXT("Hair.StrandsClusterCulling_ClusterVertexIds"), ResourceName), OwnerName, EHairResourceUsageType::Static);
+	
+	const uint32 EntryCount = FMath::DivideAndRoundUp(BulkData.Header.PointCount, HAIR_POINT_LOD_COUNT_PER_UINT);
+	InternalCreateVertexBufferRDG_FromHairBulkData<FHairClusterIndexFormat>(GraphBuilder, BulkData.Data.PointLODs, EntryCount, PointLODBuffer, ToHairResourceDebugName(TEXT("Hair.StrandsClusterCulling_PointLOD"), ResourceName), OwnerName, EHairResourceUsageType::Static);
 }
 
 void FHairStrandsClusterCullingResource::InternalRelease()
@@ -1208,6 +1211,7 @@ void FHairStrandsClusterCullingResource::InternalRelease()
 	ClusterLODInfoBuffer.Release();
 	ClusterVertexIdBuffer.Release();
 	CurveToClusterIdBuffer.Release();
+	PointLODBuffer.Release();
 	MaxAvailableCurveCount = 0;
 }
 
diff --git a/Engine/Plugins/Runtime/HairStrands/Source/HairStrandsCore/Private/HairStrandsClusterCulling.cpp b/Engine/Plugins/Runtime/HairStrands/Source/HairStrandsCore/Private/HairStrandsClusterCulling.cpp
--- a/Engine/Plugins/Runtime/HairStrands/Source/HairStrandsCore/Private/HairStrandsClusterCulling.cpp
+++ b/Engine/Plugins/Runtime/HairStrands/Source/HairStrandsCore/Private/HairStrandsClusterCulling.cpp
@@ -13,47 +13,6 @@
 
 ///////////////////////////////////////////////////////////////////////////////////////////////////
 
-static int32 GHairStrandsClusterForceLOD = -1;
-static FAutoConsoleVariableRef CVarHairClusterCullingLodMode(TEXT("r.HairStrands.Cluster.ForceLOD"), GHairStrandsClusterForceLOD, TEXT("Force a specific hair LOD."));
-
-static int32 GHairStrandsClusterCullingFreezeCamera = 0;
-static FAutoConsoleVariableRef CVarHairStrandsClusterCullingFreezeCamera(TEXT("r.HairStrands.Cluster.CullingFreezeCamera"), GHairStrandsClusterCullingFreezeCamera, TEXT("Freeze camera when enabled. It will disable HZB culling because hzb buffer is not frozen."));
-
-bool IsHairStrandsClusterCullingEnable()
-{
-	// At the moment it is not possible to disable cluster culling, as this pass is in charge of LOD selection, 
-	// and preparing the buffer which will be need for the cluster AABB pass (used later on by the voxelization pass)
-	return true;
-}
-
-// Temp: to remove
-enum EHairVisibilityVendor
-{
-	HairVisibilityVendor_AMD,
-	HairVisibilityVendor_NVIDIA,
-	HairVisibilityVendor_INTEL,
-	HairVisibilityVendorCount
-};
-
-// Temp: to remove
-inline EHairVisibilityVendor GetVendor()
-{
-	return IsRHIDeviceAMD() ? HairVisibilityVendor_AMD : (IsRHIDeviceNVIDIA() ? HairVisibilityVendor_NVIDIA : HairVisibilityVendor_INTEL);
-}
-// Temp: to remove
-static uint32 GetVendorOptimalGroupSize1D()
-{
-	switch (GetVendor())
-	{
-		case HairVisibilityVendor_AMD:		return 64;
-		case HairVisibilityVendor_NVIDIA:	return 32;
-		case HairVisibilityVendor_INTEL:	return 64;
-		default:							return 64;
-	}
-}
-
-///////////////////////////////////////////////////////////////////////////////////////////////////
-
 class FHairIndBufferClearCS : public FGlobalShader
 {
 	DECLARE_GLOBAL_SHADER(FHairIndBufferClearCS);
@@ -81,452 +40,158 @@ IMPLEMENT_GLOBAL_SHADER(FHairIndBufferClearCS, "/Engine/Private/HairStrands/Hair
 
 ///////////////////////////////////////////////////////////////////////////////////////////////////
 
-class FHairClusterCullingCS : public FGlobalShader
+class FHairClusterCullCS: public FGlobalShader
 {
-	DECLARE_GLOBAL_SHADER(FHairClusterCullingCS);
-	SHADER_USE_PARAMETER_STRUCT(FHairClusterCullingCS, FGlobalShader);
+	DECLARE_GLOBAL_SHADER(FHairClusterCullCS);
+	SHADER_USE_PARAMETER_STRUCT(FHairClusterCullCS, FGlobalShader);
 
-	class FDebugAABBBuffer : SHADER_PERMUTATION_INT("PERMUTATION_DEBUGAABBBUFFER", 2);
-	using FPermutationDomain = TShaderPermutationDomain<FDebugAABBBuffer>;
+	class FPointPerCurve : SHADER_PERMUTATION_SPARSE_INT("PERMUTATION_POINT_PER_CURVE", 4, 8, 16, 32, 64);
+	using FPermutationDomain = TShaderPermutationDomain<FPointPerCurve>;
 
-	BEGIN_SHADER_PARAMETER_STRUCT(FParameters, )
-		SHADER_PARAMETER(FVector3f, CameraWorldPos)
-		SHADER_PARAMETER(FMatrix44f, WorldToClipMatrix)
-		SHADER_PARAMETER(FMatrix44f, ProjectionMatrix)
-		SHADER_PARAMETER(uint32, ClusterCount)
-		SHADER_PARAMETER(float, LODForcedIndex)
-		SHADER_PARAMETER(int32, bIsHairGroupVisible)
-		SHADER_PARAMETER(uint32, NumConvexHullPlanes)
+	BEGIN_SHADER_PARAMETER_STRUCT(FParameters,)
+		SHADER_PARAMETER(float, LODIndex)
 		SHADER_PARAMETER(float, LODBias)
-		SHADER_PARAMETER_ARRAY(FVector4f, ViewFrustumConvexHull, [6])
-		SHADER_PARAMETER_RDG_BUFFER_SRV(Buffer, ClusterAABBBuffer)
-		SHADER_PARAMETER_RDG_BUFFER_SRV(StructuredBuffer, ClusterInfoBuffer)
-		SHADER_PARAMETER_RDG_BUFFER_SRV(StructuredBuffer, ClusterLODInfoBuffer)
-		SHADER_PARAMETER_RDG_BUFFER_UAV(RWBuffer, GlobalClusterIdBuffer)
-		SHADER_PARAMETER_RDG_BUFFER_UAV(RWBuffer, GlobalIndexStartBuffer)
-		SHADER_PARAMETER_RDG_BUFFER_UAV(RWBuffer, GlobalIndexCountBuffer)
-		SHADER_PARAMETER_RDG_BUFFER_UAV(RWBuffer, GlobalRadiusScaleBuffer)
-		SHADER_PARAMETER_RDG_BUFFER_UAV(RWStructuredBuffer, ClusterDebugInfoBuffer)
-		SHADER_PARAMETER_RDG_BUFFER_UAV(RWBuffer, DispatchIndirectParametersClusterCount)
-		SHADER_PARAMETER_RDG_BUFFER_UAV(RWBuffer, DrawIndirectParameters)
-	END_SHADER_PARAMETER_STRUCT()
-
-public:
-	static bool ShouldCompilePermutation(const FGlobalShaderPermutationParameters& Parameters) { return IsHairStrandsSupported(EHairStrandsShaderType::Strands, Parameters.Platform); }
-	static void ModifyCompilationEnvironment(const FGlobalShaderPermutationParameters& Parameters, FShaderCompilerEnvironment& OutEnvironment)
-	{
-		FGlobalShader::ModifyCompilationEnvironment(Parameters, OutEnvironment);
-		OutEnvironment.SetDefine(TEXT("SHADER_CLUSTERCULLING"), 1);
-	}
-};
-
-IMPLEMENT_GLOBAL_SHADER(FHairClusterCullingCS, "/Engine/Private/HairStrands/HairStrandsClusterCulling.usf", "MainClusterCullingCS", SF_Compute);
-
-///////////////////////////////////////////////////////////////////////////////////////////////////
-
-class FMainClusterCullingPrepareIndirectDrawsCS : public FGlobalShader
-{
-	DECLARE_GLOBAL_SHADER(FMainClusterCullingPrepareIndirectDrawsCS);
-	SHADER_USE_PARAMETER_STRUCT(FMainClusterCullingPrepareIndirectDrawsCS, FGlobalShader);
-
-	BEGIN_SHADER_PARAMETER_STRUCT(FParameters, )
-		SHADER_PARAMETER(uint32, GroupSize1D)
-		SHADER_PARAMETER_RDG_BUFFER_SRV(Buffer,   DispatchIndirectParametersClusterCount)
-		SHADER_PARAMETER_RDG_BUFFER_UAV(RWBuffer, DispatchIndirectParametersClusterCount1D)
-		SHADER_PARAMETER_RDG_BUFFER_UAV(RWBuffer, DispatchIndirectParametersClusterCount2D)
-		SHADER_PARAMETER_RDG_BUFFER_UAV(RWBuffer, DispatchIndirectParametersClusterCountDiv512)
-		SHADER_PARAMETER_RDG_BUFFER_UAV(RWBuffer, DispatchIndirectParametersClusterCountDiv512Div512)
-	END_SHADER_PARAMETER_STRUCT()
-
-public:
-	static bool ShouldCompilePermutation(const FGlobalShaderPermutationParameters& Parameters) { return IsHairStrandsSupported(EHairStrandsShaderType::Strands, Parameters.Platform); }
-	static void ModifyCompilationEnvironment(const FGlobalShaderPermutationParameters& Parameters, FShaderCompilerEnvironment& OutEnvironment)
-	{
-		FGlobalShader::ModifyCompilationEnvironment(Parameters, OutEnvironment);
-		OutEnvironment.SetDefine(TEXT("SHADER_PREPAREINDIRECTDRAW"), 1);
-	}
-};
-
-IMPLEMENT_GLOBAL_SHADER(FMainClusterCullingPrepareIndirectDrawsCS, "/Engine/Private/HairStrands/HairStrandsClusterCulling.usf", "MainClusterCullingPrepareIndirectDrawsCS", SF_Compute);
-
-///////////////////////////////////////////////////////////////////////////////////////////////////
-
-class FMainClusterCullingPrepareIndirectDispatchCS : public FGlobalShader
-{
-	DECLARE_GLOBAL_SHADER(FMainClusterCullingPrepareIndirectDispatchCS);
-	SHADER_USE_PARAMETER_STRUCT(FMainClusterCullingPrepareIndirectDispatchCS, FGlobalShader);
-
-	using FPermutationDomain = TShaderPermutationDomain<>;
-
-	BEGIN_SHADER_PARAMETER_STRUCT(FParameters, )
-		SHADER_PARAMETER_RDG_BUFFER_UAV(RWBuffer, DrawIndirectBuffer)
-		SHADER_PARAMETER_RDG_BUFFER_UAV(RWBuffer, DispatchIndirectBuffer)
-	END_SHADER_PARAMETER_STRUCT()
-
-public:
-	static bool ShouldCompilePermutation(const FGlobalShaderPermutationParameters& Parameters) { return IsHairStrandsSupported(EHairStrandsShaderType::Strands, Parameters.Platform); }
-	static void ModifyCompilationEnvironment(const FGlobalShaderPermutationParameters& Parameters, FShaderCompilerEnvironment& OutEnvironment)
-	{
-		FGlobalShader::ModifyCompilationEnvironment(Parameters, OutEnvironment);
-		OutEnvironment.SetDefine(TEXT("SHADER_PREPAREINDIRECTDISPATCH"), 1);
-	}
-};
-
-IMPLEMENT_GLOBAL_SHADER(FMainClusterCullingPrepareIndirectDispatchCS, "/Engine/Private/HairStrands/HairStrandsClusterCulling.usf", "MainClusterCullingPrepareIndirectDispatchCS", SF_Compute);
-
-///////////////////////////////////////////////////////////////////////////////////////////////////
-
-class FHairClusterCullingLocalBlockPreFixSumCS : public FGlobalShader
-{
-	DECLARE_GLOBAL_SHADER(FHairClusterCullingLocalBlockPreFixSumCS);
-	SHADER_USE_PARAMETER_STRUCT(FHairClusterCullingLocalBlockPreFixSumCS, FGlobalShader);
-
-	BEGIN_SHADER_PARAMETER_STRUCT(FParameters, )
-		RDG_BUFFER_ACCESS(DispatchIndirectParametersClusterCountDiv512, ERHIAccess::IndirectArgs)
-		SHADER_PARAMETER_RDG_BUFFER_SRV(Buffer, DispatchIndirectParametersClusterCount)
-		SHADER_PARAMETER_RDG_BUFFER_SRV(Buffer, GlobalIndexCountBuffer)
-		SHADER_PARAMETER_RDG_BUFFER_UAV(RWBuffer, PerBlocklIndexCountPreFixSumBuffer)
-		SHADER_PARAMETER_RDG_BUFFER_UAV(RWBuffer, PerBlocklTotalIndexCountBuffer)
-	END_SHADER_PARAMETER_STRUCT()
-
-public:
-	static bool ShouldCompilePermutation(const FGlobalShaderPermutationParameters& Parameters) { return IsHairStrandsSupported(EHairStrandsShaderType::Strands, Parameters.Platform); }
-	static void ModifyCompilationEnvironment(const FGlobalShaderPermutationParameters& Parameters, FShaderCompilerEnvironment& OutEnvironment)
-	{
-		FGlobalShader::ModifyCompilationEnvironment(Parameters, OutEnvironment);
-		OutEnvironment.SetDefine(TEXT("SHADER_PERBLOCKPREFIXSUM"), 1);
-	}
-};
-
-IMPLEMENT_GLOBAL_SHADER(FHairClusterCullingLocalBlockPreFixSumCS, "/Engine/Private/HairStrands/HairStrandsClusterCulling.usf", "MainPerBlockPreFixSumCS", SF_Compute);
-
-///////////////////////////////////////////////////////////////////////////////////////////////////
-
-class FHairClusterCullingCompactVertexIdsLocalBlockCS : public FGlobalShader
-{
-	DECLARE_GLOBAL_SHADER(FHairClusterCullingCompactVertexIdsLocalBlockCS);
-	SHADER_USE_PARAMETER_STRUCT(FHairClusterCullingCompactVertexIdsLocalBlockCS, FGlobalShader);
+		SHADER_PARAMETER(uint32, CurveCount)
+		SHADER_PARAMETER(uint32, DispatchCountX)
+
+		SHADER_PARAMETER_RDG_BUFFER_SRV(Buffer<uint>, CurveBuffer)
+		SHADER_PARAMETER_RDG_BUFFER_SRV(Buffer<uint>, CurveToClusterIdBuffer)
+		SHADER_PARAMETER_RDG_BUFFER_SRV(Buffer<uint>, PointLODBuffer)
+		SHADER_PARAMETER_RDG_BUFFER_SRV(StructuredBuffer<FPackedHairClusterInfo>, ClusterInfoBuffer)
+		SHADER_PARAMETER_RDG_BUFFER_SRV(StructuredBuffer<FPackedHairClusterLODInfo>, ClusterLODInfoBuffer)
+	
+		SHADER_PARAMETER_RDG_BUFFER_UAV(RWBuffer<uint>, OutPointCounter)
+		SHADER_PARAMETER_RDG_BUFFER_UAV(RWBuffer<uint>, OutIndexBuffer)
+		SHADER_PARAMETER_RDG_BUFFER_UAV(RWBuffer<float>, OutRadiusScaleBuffer)
+		SHADER_PARAMETER_RDG_BUFFER_UAV(RWBuffer<uint>, OutCulledCurveBuffer)
 
-	BEGIN_SHADER_PARAMETER_STRUCT(FParameters, )
-		RDG_BUFFER_ACCESS(DispatchIndirectParametersBuffer, ERHIAccess::IndirectArgs)
-		SHADER_PARAMETER_RDG_BUFFER_SRV(Buffer, DispatchIndirectParametersClusterCount)
-		SHADER_PARAMETER_RDG_BUFFER_SRV(Buffer, DispatchIndirectParametersClusterCount2D)
-		SHADER_PARAMETER_RDG_BUFFER_SRV(Buffer, DispatchIndirectParametersClusterCountDiv512)
-		SHADER_PARAMETER_RDG_BUFFER_SRV(Buffer, PerBlocklIndexCountPreFixSumBuffer)
-		SHADER_PARAMETER_RDG_BUFFER_SRV(Buffer, PerBlocklTotalIndexCountPreFixSumBuffer)
-		SHADER_PARAMETER_RDG_BUFFER_SRV(Buffer, GlobalIndexStartBuffer)
-		SHADER_PARAMETER_RDG_BUFFER_SRV(Buffer, GlobalIndexCountBuffer)
-		SHADER_PARAMETER_RDG_BUFFER_SRV(Buffer, GlobalRadiusScaleBuffer)
-		SHADER_PARAMETER_RDG_BUFFER_SRV(Buffer, ClusterVertexIdBuffer)
-		SHADER_PARAMETER_RDG_BUFFER_UAV(RWBuffer, CulledCompactedIndexBuffer)
-		SHADER_PARAMETER_RDG_BUFFER_UAV(RWBuffer, CulledCompactedRadiusScaleBuffer)
-	END_SHADER_PARAMETER_STRUCT()
+		SHADER_PARAMETER_STRUCT_INCLUDE(ShaderPrint::FShaderParameters, ShaderPrintUniformBuffer)
+		END_SHADER_PARAMETER_STRUCT()
 
 public:
+		static uint32 GetGroupSize() { return 64u; }
 	static bool ShouldCompilePermutation(const FGlobalShaderPermutationParameters& Parameters) { return IsHairStrandsSupported(EHairStrandsShaderType::Strands, Parameters.Platform); }
 	static void ModifyCompilationEnvironment(const FGlobalShaderPermutationParameters& Parameters, FShaderCompilerEnvironment& OutEnvironment)
 	{
 		FGlobalShader::ModifyCompilationEnvironment(Parameters, OutEnvironment);
-		OutEnvironment.SetDefine(TEXT("SHADER_CLUSTERCULLINGCOMPACTVERTEXIDLOCALBLOCK"), 1);
+		OutEnvironment.SetDefine(TEXT("SHADER_CLUSTER_CULL"), 1);
+		OutEnvironment.SetDefine(TEXT("GROUP_SIZE"), GetGroupSize());
 	}
 };
 
-IMPLEMENT_GLOBAL_SHADER(FHairClusterCullingCompactVertexIdsLocalBlockCS, "/Engine/Private/HairStrands/HairStrandsClusterCulling.usf", "MainClusterCullingCompactVertexIdsLocalBlockCS", SF_Compute);
+IMPLEMENT_GLOBAL_SHADER(FHairClusterCullCS, "/Engine/Private/HairStrands/HairStrandsClusterCulling.usf", "Main", SF_Compute);
 
 ///////////////////////////////////////////////////////////////////////////////////////////////////
 
-class FHairClusterCullingPreFixSumCS : public FGlobalShader
+class FHairClusterCullArgsCS: public FGlobalShader
 {
-	DECLARE_GLOBAL_SHADER(FHairClusterCullingPreFixSumCS);
-	SHADER_USE_PARAMETER_STRUCT(FHairClusterCullingPreFixSumCS, FGlobalShader);
+	DECLARE_GLOBAL_SHADER(FHairClusterCullArgsCS);
+	SHADER_USE_PARAMETER_STRUCT(FHairClusterCullArgsCS, FGlobalShader);
 
-	BEGIN_SHADER_PARAMETER_STRUCT(FParameters, )
-		RDG_BUFFER_ACCESS(DispatchIndirectParameters, ERHIAccess::IndirectArgs)
-		SHADER_PARAMETER_RDG_BUFFER_SRV(Buffer, DispatchIndirectParametersClusterCount)
-		SHADER_PARAMETER_RDG_BUFFER_SRV(Buffer, GlobalIndexCountBuffer)
-		SHADER_PARAMETER_RDG_BUFFER_UAV(RWBuffer, GlobalIndexCountPreFixSumBuffer)
-	END_SHADER_PARAMETER_STRUCT()
+	BEGIN_SHADER_PARAMETER_STRUCT(FParameters,)
+		SHADER_PARAMETER_RDG_BUFFER_SRV(Buffer<uint>, PointCounterBuffer)
+		SHADER_PARAMETER_RDG_BUFFER_UAV(RWBuffer<uint>, RWIndirectDrawArgsBuffer)
+		SHADER_PARAMETER_RDG_BUFFER_UAV(RWBuffer<uint>, RWIndirectDispatchArgsBuffer)
+		END_SHADER_PARAMETER_STRUCT()
 
 public:
-	static bool ShouldCompilePermutation(const FGlobalShaderPermutationParameters& Parameters) { return IsHairStrandsSupported(EHairStrandsShaderType::Strands, Parameters.Platform); }
+		static bool ShouldCompilePermutation(const FGlobalShaderPermutationParameters& Parameters) { return IsHairStrandsSupported(EHairStrandsShaderType::Strands, Parameters.Platform); }
 	static void ModifyCompilationEnvironment(const FGlobalShaderPermutationParameters& Parameters, FShaderCompilerEnvironment& OutEnvironment)
 	{
 		FGlobalShader::ModifyCompilationEnvironment(Parameters, OutEnvironment);
-		OutEnvironment.SetDefine(TEXT("SHADER_CLUSTERCULLINGPREFIXSUM"), 1);
+		OutEnvironment.SetDefine(TEXT("SHADER_CLUSTER_CULL_ARGS"), 1);
 	}
 };
 
-IMPLEMENT_GLOBAL_SHADER(FHairClusterCullingPreFixSumCS, "/Engine/Private/HairStrands/HairStrandsClusterCulling.usf", "MainClusterCullingPreFixSumCS", SF_Compute);
+IMPLEMENT_GLOBAL_SHADER(FHairClusterCullArgsCS, "/Engine/Private/HairStrands/HairStrandsClusterCulling.usf", "Main", SF_Compute);
 
 ///////////////////////////////////////////////////////////////////////////////////////////////////
-
-static FVector CapturedCameraWorldPos;
-static FMatrix CapturedWorldToClipMatrix;
-static FMatrix CapturedProjMatrix;
-
 static void AddClusterCullingPass(
 	FRDGBuilder& GraphBuilder,
 	FGlobalShaderMap* ShaderMap,
-	const FSceneView& View,
+	const FSceneView* View,
+	const FShaderPrintData* ShaderPrintData,
 	FHairStrandClusterData::FHairGroup& ClusterData)
 {
-	FRDGBufferRef DispatchIndirectParametersClusterCount = GraphBuilder.CreateBuffer(FRDGBufferDesc::CreateIndirectDesc<FRHIDispatchIndirectParameters>(), TEXT("Hair.DispatchIndirectParametersClusterCount"));
-	FRDGBufferRef DispatchIndirectParametersClusterCount1D = GraphBuilder.CreateBuffer(FRDGBufferDesc::CreateIndirectDesc<FRHIDispatchIndirectParameters>(), TEXT("Hair.DispatchIndirectParametersClusterCount1D"));
-	FRDGBufferRef DispatchIndirectParametersClusterCount2D = GraphBuilder.CreateBuffer(FRDGBufferDesc::CreateIndirectDesc<FRHIDispatchIndirectParameters>(), TEXT("Hair.DispatchIndirectParametersClusterCount2D"));
-	FRDGBufferRef DispatchIndirectParametersClusterCountDiv512 = GraphBuilder.CreateBuffer(FRDGBufferDesc::CreateIndirectDesc<FRHIDispatchIndirectParameters>(), TEXT("Hair.DispatchIndirectParametersClusterCountDiv512"));
-	FRDGBufferRef DispatchIndirectParametersClusterCountDiv512Div512 = GraphBuilder.CreateBuffer(FRDGBufferDesc::CreateIndirectDesc<FRHIDispatchIndirectParameters>(), TEXT("Hair.DispatchIndirectParametersClusterCountDiv512Div512"));
-	
-	FRDGBufferRef GlobalClusterIdBuffer = GraphBuilder.CreateBuffer(FRDGBufferDesc::CreateBufferDesc(sizeof(uint32), ClusterData.ClusterCount), TEXT("Hair.GlobalClusterIdBuffer"));
-	FRDGBufferRef GlobalIndexStartBuffer = GraphBuilder.CreateBuffer(FRDGBufferDesc::CreateBufferDesc(sizeof(uint32), ClusterData.ClusterCount), TEXT("Hair.GlobalIndexStartBuffer"));
-	FRDGBufferRef GlobalIndexCountBuffer = GraphBuilder.CreateBuffer(FRDGBufferDesc::CreateBufferDesc(sizeof(uint32), ClusterData.ClusterCount), TEXT("Hair.GlobalIndexCountBuffer"));
-	FRDGBufferRef GlobalRadiusScaleBuffer = GraphBuilder.CreateBuffer(FRDGBufferDesc::CreateBufferDesc(sizeof(float), ClusterData.ClusterCount), TEXT("Hair.GlobalRadiusScaleBuffer"));
-
-	const uint32 PrefixGroupSize = 512;
-	const uint32 ClusterCountRoundUpToGroupSize = FMath::DivideAndRoundUp(ClusterData.ClusterCount, PrefixGroupSize) * PrefixGroupSize;
-	FRDGBufferRef PerBlocklTotalIndexCountBuffer = GraphBuilder.CreateBuffer(FRDGBufferDesc::CreateBufferDesc(sizeof(uint32), ClusterCountRoundUpToGroupSize), TEXT("Hair.PerBlocklTotalIndexCountBuffer"));
-	FRDGBufferRef PerBlocklTotalIndexCountPreFixSumBuffer = GraphBuilder.CreateBuffer(FRDGBufferDesc::CreateBufferDesc(sizeof(uint32) * 2, ClusterCountRoundUpToGroupSize), TEXT("Hair.PerBlocklTotalIndexCountPreFixSumBuffer"));
-	FRDGBufferRef PerBlocklIndexCountPreFixSumBuffer = GraphBuilder.CreateBuffer(FRDGBufferDesc::CreateBufferDesc(sizeof(uint32) * 2, ClusterCountRoundUpToGroupSize), TEXT("Hair.PerBlocklIndexCountPreFixSumBuffer"));
-
-	FRDGImportedBuffer DrawIndirectParametersBuffer = Register(GraphBuilder, ClusterData.HairGroupPublicPtr->GetDrawIndirectBuffer(), ERDGImportedBufferFlags::CreateViews);
-	FRDGImportedBuffer DrawIndirectParametersRasterComputeBuffer = Register(GraphBuilder, ClusterData.HairGroupPublicPtr->GetDrawIndirectRasterComputeBuffer(), ERDGImportedBufferFlags::CreateViews);
-	
-	bool bClusterDebugAABBBuffer = false;
-	bool bClusterDebug = false;
-#if WITH_EDITOR
-	FRDGBufferRef ClusterDebugInfoBuffer = nullptr;
-	{
-		// Defined in HairStrandsClusterCommon.ush
-		struct FHairClusterDebugInfo
-		{
-			uint32 GroupIndex;
-			float LOD;
-			float VertexCount;
-			float CurveCount;
-		};
-
-		static IConsoleVariable* CVarClusterDebug = IConsoleManager::Get().FindConsoleVariable(TEXT("r.HairStrands.Cluster.Debug"));
-		bClusterDebugAABBBuffer = GetGroomViewMode(View) == EGroomViewMode::ClusterAABB;
-		bClusterDebug = GetGroomViewMode(View) == EGroomViewMode::Cluster;
-		ClusterDebugInfoBuffer = GraphBuilder.CreateBuffer(FRDGBufferDesc::CreateStructuredDesc(sizeof(FHairClusterDebugInfo), ClusterData.ClusterCount), TEXT("Hair.CulledCompactedIndexBuffer"));
-	}
-#endif
 
-	/// Initialise indirect buffers to be setup during the culling process
-	{
-		FHairIndBufferClearCS::FParameters* Parameters = GraphBuilder.AllocParameters<FHairIndBufferClearCS::FParameters>();
-		Parameters->DispatchIndirectParametersClusterCount = GraphBuilder.CreateUAV(DispatchIndirectParametersClusterCount);
-		Parameters->DrawIndirectParameters = DrawIndirectParametersBuffer.UAV;
+	ShaderPrint::SetEnabled(true);
+	ShaderPrint::RequestSpaceForCharacters(2048);
+	ShaderPrint::RequestSpaceForLines(2048);
 
-		FHairIndBufferClearCS::FPermutationDomain Permutation;
-		Permutation.Set<FHairIndBufferClearCS::FSetIndirectDraw>(false);
-		TShaderMapRef<FHairIndBufferClearCS> ComputeShader(ShaderMap, Permutation);
-		FComputeShaderUtils::AddPass(
-			GraphBuilder,
-			RDG_EVENT_NAME("BufferClearCS"),
-			ComputeShader,
-			Parameters,
-			FIntVector(1, 1, 1));
-	}
+	// 0. Glogal counter for visible points
+	FRDGBufferRef PointCounter 	= GraphBuilder.CreateBuffer(FRDGBufferDesc::CreateByteAddressDesc(4), TEXT("Hair.ClusterPointCounter"));
+	FRDGBufferUAVRef PointCounterUAV = GraphBuilder.CreateUAV(PointCounter);
+	AddClearUAVPass(GraphBuilder, PointCounterUAV, 0);
 
-	/// Cull cluster, generate indirect dispatch and prepare data to expand index buffer
+	// 1. Build culled index buffer
 	{
-		const bool bClusterCullingFrozenCamera = GHairStrandsClusterCullingFreezeCamera > 0;
-		if (!bClusterCullingFrozenCamera)
-		{
-			CapturedCameraWorldPos = View.ViewMatrices.GetViewOrigin();
-			CapturedWorldToClipMatrix = View.ViewMatrices.GetViewProjectionMatrix();
-			CapturedProjMatrix = View.ViewMatrices.GetProjectionMatrix();
-		}
-
-		float ForceLOD = -1;
-		bool bIsVisible = true;
-		if (GHairStrandsClusterForceLOD >= 0)
-		{
-			ForceLOD = GHairStrandsClusterForceLOD;
-		}
-		else if (ClusterData.LODIndex >= 0) // CPU-driven LOD selection
-		{
-			ForceLOD = ClusterData.LODIndex;
-			bIsVisible = ClusterData.bVisible;
-		}
-
-		FHairClusterCullingCS::FParameters* Parameters = GraphBuilder.AllocParameters<FHairClusterCullingCS::FParameters>();
-		Parameters->ProjectionMatrix = FMatrix44f(CapturedProjMatrix);			// LWC_TODO: Precision loss
-		Parameters->CameraWorldPos = FVector3f(CapturedCameraWorldPos);
-		Parameters->WorldToClipMatrix = FMatrix44f(CapturedWorldToClipMatrix);
-		Parameters->ClusterCount = ClusterData.ClusterCount;
-		Parameters->LODForcedIndex = ForceLOD;
-		Parameters->LODBias = ClusterData.LODBias;
-		Parameters->bIsHairGroupVisible = bIsVisible ? 1 : 0;
-
-		Parameters->NumConvexHullPlanes = View.ViewFrustum.Planes.Num();
-		check(Parameters->NumConvexHullPlanes <= 6);
-		for (uint32 i = 0; i < Parameters->NumConvexHullPlanes; ++i)
-		{
-			const FPlane& Plane = View.ViewFrustum.Planes[i];
-			Parameters->ViewFrustumConvexHull[i] = FVector4f(Plane.X, Plane.Y, Plane.Z, Plane.W); // LWC_TODO: precision loss
-		}
-
-		Parameters->ClusterAABBBuffer = RegisterAsSRV(GraphBuilder, *ClusterData.ClusterAABBBuffer);
-		Parameters->ClusterInfoBuffer = RegisterAsSRV(GraphBuilder, *ClusterData.ClusterInfoBuffer);
+		FHairClusterCullCS::FParameters* Parameters = GraphBuilder.AllocParameters<FHairClusterCullCS::FParameters>();
+		Parameters->LODIndex 			 = ClusterData.LODIndex;
+		Parameters->LODBias 			 = ClusterData.LODBias;
+		Parameters->CurveCount 			 = ClusterData.HairGroupPublicPtr->GetActiveStrandsCurveCount();	
+		Parameters->CurveBuffer 		 = RegisterAsSRV(GraphBuilder, *ClusterData.CurveBuffer);
+		Parameters->CurveToClusterIdBuffer=RegisterAsSRV(GraphBuilder, *ClusterData.CurveToClusterIdBuffer);
+		Parameters->PointLODBuffer 		 = RegisterAsSRV(GraphBuilder, *ClusterData.PointLODBuffer);
+		Parameters->ClusterInfoBuffer 	 = RegisterAsSRV(GraphBuilder, *ClusterData.ClusterInfoBuffer);
 		Parameters->ClusterLODInfoBuffer = RegisterAsSRV(GraphBuilder, *ClusterData.ClusterLODInfoBuffer);
 
-		Parameters->GlobalClusterIdBuffer = GraphBuilder.CreateUAV(GlobalClusterIdBuffer, PF_R32_UINT);
-		Parameters->GlobalIndexStartBuffer = GraphBuilder.CreateUAV(GlobalIndexStartBuffer, PF_R32_UINT);
-		Parameters->GlobalIndexCountBuffer = GraphBuilder.CreateUAV(GlobalIndexCountBuffer, PF_R32_UINT);
-		Parameters->GlobalRadiusScaleBuffer = GraphBuilder.CreateUAV(GlobalRadiusScaleBuffer, PF_R32_FLOAT);
+		Parameters->OutPointCounter 	 = PointCounterUAV;
+		Parameters->OutIndexBuffer 		 = RegisterAsUAV(GraphBuilder, *ClusterData.GetCulledVertexIdBuffer());
+		Parameters->OutRadiusScaleBuffer = RegisterAsUAV(GraphBuilder, *ClusterData.GetCulledVertexRadiusScaleBuffer());
+		Parameters->OutCulledCurveBuffer = RegisterAsUAV(GraphBuilder, *ClusterData.GetCulledCurveBuffer()); // TODO: this could be changed to be transient buffer
 
-		Parameters->DispatchIndirectParametersClusterCount = GraphBuilder.CreateUAV(DispatchIndirectParametersClusterCount);
-		Parameters->DrawIndirectParameters = DrawIndirectParametersBuffer.UAV;
-
-#if WITH_EDITOR
-		Parameters->ClusterDebugInfoBuffer = GraphBuilder.CreateUAV(ClusterDebugInfoBuffer, PF_R32_SINT);
-#endif
-
-		FHairClusterCullingCS::FPermutationDomain Permutation;
-		Permutation.Set<FHairClusterCullingCS::FDebugAABBBuffer>(bClusterDebugAABBBuffer ? 1 : 0);
-		TShaderMapRef<FHairClusterCullingCS> ComputeShader(ShaderMap, Permutation);
-		const FIntVector DispatchCount = DispatchCount.DivideAndRoundUp(FIntVector(ClusterData.ClusterCount, 1, 1), FIntVector(64, 1, 1));
-		FComputeShaderUtils::AddPass(
-			GraphBuilder,
-			RDG_EVENT_NAME("ClusterCullingCS"),
-			ComputeShader,
-			Parameters,
-			DispatchCount);
-	}
+		ShaderPrint::SetParameters(GraphBuilder, *ShaderPrintData, Parameters->ShaderPrintUniformBuffer);
 
-	/// Prepare some indirect draw buffers for specific compute group size
-	const uint32 GroupSize1D = GetVendorOptimalGroupSize1D();
-	{
-		FMainClusterCullingPrepareIndirectDrawsCS::FParameters* Parameters = GraphBuilder.AllocParameters<FMainClusterCullingPrepareIndirectDrawsCS::FParameters>();
-		Parameters->GroupSize1D = GroupSize1D;
-		Parameters->DispatchIndirectParametersClusterCount = GraphBuilder.CreateSRV(DispatchIndirectParametersClusterCount, PF_R32_UINT);
-		Parameters->DispatchIndirectParametersClusterCount1D = GraphBuilder.CreateUAV(DispatchIndirectParametersClusterCount1D, PF_R32_UINT);
-		Parameters->DispatchIndirectParametersClusterCount2D = GraphBuilder.CreateUAV(DispatchIndirectParametersClusterCount2D, PF_R32_UINT);
-		Parameters->DispatchIndirectParametersClusterCountDiv512 = GraphBuilder.CreateUAV(DispatchIndirectParametersClusterCountDiv512, PF_R32_UINT);
-		Parameters->DispatchIndirectParametersClusterCountDiv512Div512 = GraphBuilder.CreateUAV(DispatchIndirectParametersClusterCountDiv512Div512, PF_R32_UINT);
-
-		TShaderMapRef<FMainClusterCullingPrepareIndirectDrawsCS> ComputeShader(ShaderMap);
-		FComputeShaderUtils::AddPass(
-			GraphBuilder,
-			RDG_EVENT_NAME("PrepareIndirectDrawsCS"),
-			ComputeShader,
-			Parameters,
-			FIntVector(2, 1, 1));
-	}
+		// Pick permutation based on groom max number of point per curve
+		//const uint32 PointPerCurve = FMath::Clamp(uint32(FMath::Pow( FMath::RoundFromZero(FMath::Log2(float(ClusterData.MaxPointPerCurve))),2u)), 4u, 32u);
+		const uint32 PointPerCurve = FMath::Clamp(uint32(FMath::Pow(2u, FMath::RoundFromZero(FMath::Log2(float(ClusterData.MaxPointPerCurve))))), 4u, 64u);
+		check(FMath::IsPowerOfTwo(PointPerCurve));
 
-	/// local prefix sum per 512 block
-	{
-		FHairClusterCullingLocalBlockPreFixSumCS::FParameters* Parameters = GraphBuilder.AllocParameters<FHairClusterCullingLocalBlockPreFixSumCS::FParameters>();
-		Parameters->DispatchIndirectParametersClusterCountDiv512 = DispatchIndirectParametersClusterCountDiv512;
-		Parameters->DispatchIndirectParametersClusterCount = GraphBuilder.CreateSRV(DispatchIndirectParametersClusterCount);
-		Parameters->GlobalIndexCountBuffer = GraphBuilder.CreateSRV(GlobalIndexCountBuffer, PF_R32_UINT);
-		Parameters->PerBlocklIndexCountPreFixSumBuffer = GraphBuilder.CreateUAV(PerBlocklIndexCountPreFixSumBuffer, PF_R32G32_UINT);
-		Parameters->PerBlocklTotalIndexCountBuffer = GraphBuilder.CreateUAV(PerBlocklTotalIndexCountBuffer, PF_R32_UINT);
-
-		TShaderMapRef<FHairClusterCullingLocalBlockPreFixSumCS> ComputeShader(ShaderMap);
-		FComputeShaderUtils::AddPass(
-			GraphBuilder,
-			RDG_EVENT_NAME("WithinBlockIndexCountPreFixSumCS"),
-			ComputeShader,
-			Parameters,
-			DispatchIndirectParametersClusterCountDiv512, 0); // FIX ME, this could get over 65535
-		check(ClusterData.ClusterCount / 512u <= uint32(GRHIMaxDispatchThreadGroupsPerDimension.X));
-	}
+		const uint32 CurvePerGroup = FHairClusterCullCS::GetGroupSize() / PointPerCurve;
+		const uint32 LinearGroupCount = FMath::DivideAndRoundUp(Parameters->CurveCount, CurvePerGroup);
 
-	/// Prefix sum on the total index count per block of 512
-	{
-		FHairClusterCullingPreFixSumCS::FParameters* Parameters = GraphBuilder.AllocParameters<FHairClusterCullingPreFixSumCS::FParameters>();
-		Parameters->DispatchIndirectParameters = DispatchIndirectParametersClusterCountDiv512Div512;
-		Parameters->DispatchIndirectParametersClusterCount = GraphBuilder.CreateSRV(DispatchIndirectParametersClusterCountDiv512, PF_R32_UINT);
-		Parameters->GlobalIndexCountBuffer = GraphBuilder.CreateSRV(PerBlocklTotalIndexCountBuffer, PF_R32_UINT);
-		Parameters->GlobalIndexCountPreFixSumBuffer = GraphBuilder.CreateUAV(PerBlocklTotalIndexCountPreFixSumBuffer, PF_R32G32_UINT);
+		FIntVector DispatchCount(LinearGroupCount, 1, 1);
+		if (DispatchCount.X > 0xFFFFu)
+		{
+			DispatchCount.X = 64;
+			DispatchCount.Y = FMath::DivideAndRoundUp(LinearGroupCount, uint32(DispatchCount.X));
+		}
+		Parameters->DispatchCountX = DispatchCount.X;
+		check(DispatchCount.X <= 0xFFFFu);
+		check(DispatchCount.Y <= 0xFFFFu);
 
-		TShaderMapRef<FHairClusterCullingPreFixSumCS> ComputeShader(ShaderMap);
+		FHairClusterCullCS::FPermutationDomain Permutation;
+		Permutation.Set<FHairClusterCullCS::FPointPerCurve>(PointPerCurve);
+		TShaderMapRef<FHairClusterCullCS> ComputeShader(ShaderMap, Permutation);
 		FComputeShaderUtils::AddPass(
 			GraphBuilder,
-			RDG_EVENT_NAME("BlockIndexCountPreFixSumCS"),
+			RDG_EVENT_NAME("HairStrands::ClusterCullPass2"),
 			ComputeShader,
 			Parameters,
-			DispatchIndirectParametersClusterCountDiv512Div512, 0); // FIX ME, this could get over 65535
-		check(ClusterData.ClusterCount / (512u*512u) <= uint32(GRHIMaxDispatchThreadGroupsPerDimension.X));
+			DispatchCount);
 	}
 
-	/// Compact to VertexId buffer using hierarchical binary search or splatting
+	// 2. Prepare indirect draw/dispatch args buffers
 	{
-		check(ClusterData.GetCulledVertexIdBuffer());
-		check(ClusterData.GetCulledVertexRadiusScaleBuffer());
-		
-		FRDGImportedBuffer ClusterVertexIdBuffer			= Register(GraphBuilder, *ClusterData.ClusterVertexIdBuffer, ERDGImportedBufferFlags::CreateSRV);
-		FRDGImportedBuffer CulledCompactedIndexBuffer		= Register(GraphBuilder, *ClusterData.GetCulledVertexIdBuffer(), ERDGImportedBufferFlags::CreateUAV);
-		FRDGImportedBuffer CulledCompactedRadiusScaleBuffer	= Register(GraphBuilder, *ClusterData.GetCulledVertexRadiusScaleBuffer(), ERDGImportedBufferFlags::CreateUAV);
+		FRDGImportedBuffer DrawIndirectParametersBuffer = Register(GraphBuilder, ClusterData.HairGroupPublicPtr->GetDrawIndirectBuffer(), ERDGImportedBufferFlags::CreateViews);
+		FRDGImportedBuffer DrawIndirectParametersRasterComputeBuffer = Register(GraphBuilder, ClusterData.HairGroupPublicPtr->GetDrawIndirectRasterComputeBuffer(), ERDGImportedBufferFlags::CreateViews);
 
-		FHairClusterCullingCompactVertexIdsLocalBlockCS::FParameters* Parameters = GraphBuilder.AllocParameters<FHairClusterCullingCompactVertexIdsLocalBlockCS::FParameters>();
+		FHairClusterCullArgsCS::FParameters* Parameters = GraphBuilder.AllocParameters<FHairClusterCullArgsCS::FParameters>();
+		Parameters->PointCounterBuffer = GraphBuilder.CreateSRV(PointCounter);
+		Parameters->RWIndirectDrawArgsBuffer = DrawIndirectParametersBuffer.UAV;
+		Parameters->RWIndirectDispatchArgsBuffer = DrawIndirectParametersRasterComputeBuffer.UAV;
 
-		Parameters->DispatchIndirectParametersClusterCount = GraphBuilder.CreateSRV(DispatchIndirectParametersClusterCount, PF_R32_UINT);
-		Parameters->DispatchIndirectParametersClusterCount2D = GraphBuilder.CreateSRV(DispatchIndirectParametersClusterCount2D, PF_R32_UINT);
-		Parameters->DispatchIndirectParametersClusterCountDiv512 = GraphBuilder.CreateSRV(DispatchIndirectParametersClusterCountDiv512, PF_R32_UINT);
-
-		Parameters->PerBlocklIndexCountPreFixSumBuffer = GraphBuilder.CreateSRV(PerBlocklIndexCountPreFixSumBuffer, PF_R32G32_UINT);
-		Parameters->PerBlocklTotalIndexCountPreFixSumBuffer = GraphBuilder.CreateSRV(PerBlocklTotalIndexCountPreFixSumBuffer, PF_R32G32_UINT);
-
-		Parameters->GlobalIndexStartBuffer = GraphBuilder.CreateSRV(GlobalIndexStartBuffer, PF_R32_UINT);
-		Parameters->GlobalIndexCountBuffer = GraphBuilder.CreateSRV(GlobalIndexCountBuffer, PF_R32_UINT);
-		Parameters->GlobalRadiusScaleBuffer = GraphBuilder.CreateSRV(GlobalRadiusScaleBuffer, PF_R32_FLOAT);
-		Parameters->ClusterVertexIdBuffer = ClusterVertexIdBuffer.SRV;
-
-		
-		Parameters->CulledCompactedIndexBuffer = CulledCompactedIndexBuffer.UAV;
-		Parameters->CulledCompactedRadiusScaleBuffer = CulledCompactedRadiusScaleBuffer.UAV;
-
-		Parameters->DispatchIndirectParametersBuffer = DispatchIndirectParametersClusterCount2D;
-		TShaderMapRef<FHairClusterCullingCompactVertexIdsLocalBlockCS> ComputeShader(ShaderMap);
+		TShaderMapRef<FHairClusterCullArgsCS> ComputeShader(ShaderMap);
 		FComputeShaderUtils::AddPass(
 			GraphBuilder,
-			RDG_EVENT_NAME("SplatCompactVertexIdsCS"),
-			ComputeShader,
-			Parameters,
-			DispatchIndirectParametersClusterCount2D, 0); // DispatchIndirectParametersClusterCount2D is used to avoid having any dispatch dimension going above 65535.
-
-		GraphBuilder.SetBufferAccessFinal(CulledCompactedIndexBuffer.Buffer, ERHIAccess::SRVMask);
-		GraphBuilder.SetBufferAccessFinal(CulledCompactedRadiusScaleBuffer.Buffer, ERHIAccess::SRVMask);
-	}
-
-	{
-		ClusterData.ClusterIdBuffer = GraphBuilder.ConvertToExternalBuffer(GlobalClusterIdBuffer);
-		ClusterData.ClusterIndexOffsetBuffer = GraphBuilder.ConvertToExternalBuffer(GlobalIndexStartBuffer);
-		ClusterData.ClusterIndexCountBuffer = GraphBuilder.ConvertToExternalBuffer(GlobalIndexCountBuffer);
-		ClusterData.CulledCluster2DIndirectArgsBuffer = DispatchIndirectParametersClusterCount2D;
-		ClusterData.CulledCluster1DIndirectArgsBuffer = DispatchIndirectParametersClusterCount1D;
-		ClusterData.CulledClusterCountBuffer = DispatchIndirectParametersClusterCount;
-		ClusterData.GroupSize1D = GroupSize1D;
-	}
-
-	/// Prepare some indirect dispatch for compute raster visibility buffers
-	{
-		FMainClusterCullingPrepareIndirectDispatchCS::FParameters* Parameters = GraphBuilder.AllocParameters<FMainClusterCullingPrepareIndirectDispatchCS::FParameters>();
-		Parameters->DrawIndirectBuffer = DrawIndirectParametersBuffer.UAV;
-		Parameters->DispatchIndirectBuffer = DrawIndirectParametersRasterComputeBuffer.UAV;
-		
-		TShaderMapRef<FMainClusterCullingPrepareIndirectDispatchCS> ComputeShader(ShaderMap);
-		FComputeShaderUtils::AddPass(
-			GraphBuilder,
-			RDG_EVENT_NAME("PrepareIndirectDispatchCS"),
+			RDG_EVENT_NAME("HairStrands::ClusterCullPrepareArgs"),
 			ComputeShader,
 			Parameters,
 			FIntVector(1, 1, 1));
-	}
 
-	// Should this be move onto the culling result?
-#if WITH_EDITOR
-	if (bClusterDebugAABBBuffer)
-	{
-		ClusterData.ClusterDebugInfoBuffer = GraphBuilder.ConvertToExternalBuffer(ClusterDebugInfoBuffer);
+		GraphBuilder.SetBufferAccessFinal(DrawIndirectParametersBuffer.Buffer, ERHIAccess::IndirectArgs | ERHIAccess::SRVMask);
+		GraphBuilder.SetBufferAccessFinal(DrawIndirectParametersRasterComputeBuffer.Buffer, ERHIAccess::IndirectArgs | ERHIAccess::SRVMask);
 	}
-#endif
-
-	GraphBuilder.SetBufferAccessFinal(DrawIndirectParametersBuffer.Buffer, ERHIAccess::IndirectArgs | ERHIAccess::SRVMask);
-	GraphBuilder.SetBufferAccessFinal(DrawIndirectParametersRasterComputeBuffer.Buffer, ERHIAccess::IndirectArgs | ERHIAccess::SRVMask);
 
 	ClusterData.SetCullingResultAvailable(true);
 }
@@ -567,6 +232,7 @@ void ComputeHairStrandsClustersCulling(
 	FRDGBuilder& GraphBuilder,
 	FGlobalShaderMap& ShaderMap,
 	const TArray<const FSceneView*>& Views, 
+	const FShaderPrintData* ShaderPrintData,
 	FHairStrandClusterData& ClusterDatas)
 {
 	DECLARE_GPU_STAT(HairStrandsClusterCulling);
@@ -587,7 +253,8 @@ void ComputeHairStrandsClustersCulling(
 			AddClusterCullingPass(
 				GraphBuilder,
 				&ShaderMap,
-				*View,
+				View,
+				ShaderPrintData,
 				ClusterData);
 		}
 	}
@@ -600,11 +267,15 @@ void AddInstanceToClusterData(FHairGroupInstance* In, FHairStrandClusterData& Ou
 	// Initialize group cluster data for culling by the renderer
 	const int32 ClusterDataGroupIndex = Out.HairGroups.Num();
 	FHairStrandClusterData::FHairGroup& HairGroupCluster = Out.HairGroups.Emplace_GetRef();
+	HairGroupCluster.MaxPointPerCurve = In->Strands.Data ? In->Strands.Data->Header.MaxPointPerCurve : 0;
 	HairGroupCluster.ClusterCount = In->HairGroupPublicData->GetClusterCount();
 	HairGroupCluster.VertexCount = In->HairGroupPublicData->RestPointCount * HAIR_POINT_TO_VERTEX; // Instance->HairGroupPublicData->GetActiveStrandsPointCount() ?
 	HairGroupCluster.GroupAABBBuffer = &In->HairGroupPublicData->GetGroupAABBBuffer();
 	HairGroupCluster.ClusterAABBBuffer = &In->HairGroupPublicData->GetClusterAABBBuffer();
 
+	HairGroupCluster.CurveBuffer = &In->Strands.RestResource->CurveBuffer;
+	HairGroupCluster.PointLODBuffer = &In->Strands.ClusterCullingResource->PointLODBuffer;
+
 	HairGroupCluster.ClusterInfoBuffer = &In->Strands.ClusterCullingResource->ClusterInfoBuffer;
 	HairGroupCluster.ClusterLODInfoBuffer = &In->Strands.ClusterCullingResource->ClusterLODInfoBuffer;
 	HairGroupCluster.CurveToClusterIdBuffer = &In->Strands.ClusterCullingResource->CurveToClusterIdBuffer;
diff --git a/Engine/Plugins/Runtime/HairStrands/Source/HairStrandsCore/Private/HairStrandsClusterCulling.h b/Engine/Plugins/Runtime/HairStrands/Source/HairStrandsCore/Private/HairStrandsClusterCulling.h
--- a/Engine/Plugins/Runtime/HairStrands/Source/HairStrandsCore/Private/HairStrandsClusterCulling.h
+++ b/Engine/Plugins/Runtime/HairStrands/Source/HairStrandsCore/Private/HairStrandsClusterCulling.h
@@ -25,4 +25,5 @@ void ComputeHairStrandsClustersCulling(
 	FRDGBuilder& GraphBuilder,
 	FGlobalShaderMap& ShaderMap,
 	const TArray<const FSceneView*>& Views,
+	const FShaderPrintData* ShaderPrintData,
 	FHairStrandClusterData& ClusterDatas);
diff --git a/Engine/Plugins/Runtime/HairStrands/Source/HairStrandsCore/Private/HairStrandsDatas.cpp b/Engine/Plugins/Runtime/HairStrands/Source/HairStrandsCore/Private/HairStrandsDatas.cpp
--- a/Engine/Plugins/Runtime/HairStrands/Source/HairStrandsCore/Private/HairStrandsDatas.cpp
+++ b/Engine/Plugins/Runtime/HairStrands/Source/HairStrandsCore/Private/HairStrandsDatas.cpp
@@ -667,12 +667,14 @@ void FHairStrandsClusterCullingBulkData::Reset()
 	Data.CurveToClusterIds.RemoveBulkData();
 	Data.ClusterVertexIds.RemoveBulkData();
 	Data.PackedClusterInfos.RemoveBulkData();
+	Data.PointLODs.RemoveBulkData();
 
 	// Reset the bulk byte buffer to ensure the (serialize) data size is reset to 0
 	Data.ClusterLODInfos 	= FHairBulkContainer();
 	Data.CurveToClusterIds  = FHairBulkContainer();
 	Data.ClusterVertexIds 	= FHairBulkContainer();
 	Data.PackedClusterInfos = FHairBulkContainer();
+	Data.PointLODs			= FHairBulkContainer();
 }
 
 void FHairStrandsClusterCullingBulkData::ResetLoadedSize()
@@ -707,7 +709,7 @@ void FHairStrandsClusterCullingBulkData::SerializeHeader(FArchive& Ar, UObject*
 
 uint32 FHairStrandsClusterCullingBulkData::GetResourceCount() const
 {
-	return 4;
+	return 5;
 }
 
 bool ValidateHairBulkData();
@@ -734,6 +736,10 @@ void FHairStrandsClusterCullingBulkData::GetResources(FHairStrandsBulkCommon::FQ
 		Out.Add(Data.PackedClusterInfos, TEXT("_PackedClusterInfos"), 0, 0); // Load all data
 	}
 
+	{
+		Out.Add(Data.PointLODs, TEXT("_PointLODs"), 0, 0); // Load all data. TODO load data based on active curve/point count
+	}
+
 	if (ValidateHairBulkData() && (Out.Type == FHairStrandsBulkCommon::FQuery::WriteDDC || Out.Type == FHairStrandsBulkCommon::FQuery::ReadWriteIO))
 	{
 		Validate(true);
diff --git a/Engine/Plugins/Runtime/HairStrands/Source/HairStrandsCore/Private/HairStrandsInterpolation.cpp b/Engine/Plugins/Runtime/HairStrands/Source/HairStrandsCore/Private/HairStrandsInterpolation.cpp
--- a/Engine/Plugins/Runtime/HairStrands/Source/HairStrandsCore/Private/HairStrandsInterpolation.cpp
+++ b/Engine/Plugins/Runtime/HairStrands/Source/HairStrandsCore/Private/HairStrandsInterpolation.cpp
@@ -821,39 +821,31 @@ class FHairClusterAABBCS : public FGlobalShader
 	DECLARE_GLOBAL_SHADER(FHairClusterAABBCS);
 	SHADER_USE_PARAMETER_STRUCT(FHairClusterAABBCS, FGlobalShader);
 
-	class FGroupSize : SHADER_PERMUTATION_SPARSE_INT("GROUP_SIZE", 32, 64);
 	class FCPUAABB : SHADER_PERMUTATION_BOOL("PERMUTATION_USE_CPU_AABB");
-	using FPermutationDomain = TShaderPermutationDomain<FGroupSize, FCPUAABB>;
+	using FPermutationDomain = TShaderPermutationDomain<FCPUAABB>;
 
 	BEGIN_SHADER_PARAMETER_STRUCT(FParameters, )
 		SHADER_PARAMETER_STRUCT_INCLUDE(ShaderPrint::FShaderParameters, ShaderDrawParameters)
-		SHADER_PARAMETER(uint32, TotalClusterCount)
-		SHADER_PARAMETER(uint32, PointCount)
-
+		SHADER_PARAMETER(uint32, ClusterCount)
+		SHADER_PARAMETER(uint32, CurveCount)
 		SHADER_PARAMETER(FVector3f, CPUBoundMin)
 		SHADER_PARAMETER(FVector3f, CPUBoundMax)
 		SHADER_PARAMETER(FMatrix44f, LocalToTranslatedWorldMatrix)
-
+		SHADER_PARAMETER_RDG_BUFFER_SRV(Buffer<uint>, RenCurveBuffer)
 		SHADER_PARAMETER_RDG_BUFFER_SRV(Buffer, RenderDeformedPositionBuffer)
 		SHADER_PARAMETER_RDG_BUFFER_SRV(Buffer, RenderDeformedOffsetBuffer)
-		SHADER_PARAMETER_RDG_BUFFER_SRV(Buffer, ClusterVertexIdBuffer)
-		SHADER_PARAMETER_RDG_BUFFER_SRV(Buffer, ClusterIdBuffer)
-		SHADER_PARAMETER_RDG_BUFFER_SRV(Buffer, ClusterIndexOffsetBuffer)
-		SHADER_PARAMETER_RDG_BUFFER_SRV(Buffer, ClusterIndexCountBuffer)
-		SHADER_PARAMETER_RDG_BUFFER_SRV(Buffer, CulledClusterCountBuffer)
-		SHADER_PARAMETER_RDG_BUFFER_SRV(Buffer<uint>, HairStrandsVFTODO_CullingIndirectBuffer)
-		SHADER_PARAMETER_RDG_BUFFER_SRV(Buffer<uint>, IndirectBuffer)
 		SHADER_PARAMETER_RDG_BUFFER_UAV(RWBuffer, OutClusterAABBBuffer)
 		SHADER_PARAMETER_RDG_BUFFER_UAV(RWBuffer, OutGroupAABBBuffer)
-		RDG_BUFFER_ACCESS(IndirectArgsBuffer, ERHIAccess::IndirectArgs)
 	END_SHADER_PARAMETER_STRUCT()
 
 public:
+	static uint32 GetGroupSize() { return 64; }
 	static bool ShouldCompilePermutation(const FGlobalShaderPermutationParameters& Parameters) { return IsHairStrandsSupported(EHairStrandsShaderType::Strands, Parameters.Platform); }
 	static void ModifyCompilationEnvironment(const FGlobalShaderPermutationParameters& Parameters, FShaderCompilerEnvironment& OutEnvironment)
 	{
 		FGlobalShader::ModifyCompilationEnvironment(Parameters, OutEnvironment);
 		OutEnvironment.SetDefine(TEXT("SHADER_CLUSTERAABB"), 1);
+		OutEnvironment.SetDefine(TEXT("GROUP_SIZE"), GetGroupSize());
 	}
 };
 
@@ -868,6 +860,7 @@ enum class EHairAABBUpdateType
 static void AddHairClusterAABBPass(
 	FRDGBuilder& GraphBuilder,
 	FGlobalShaderMap* ShaderMap,
+	const uint32 ActiveCurveCount,
 	const FVector& TranslatedWorldOffset,
 	const FShaderPrintData* ShaderPrintData,
 	const EHairAABBUpdateType UpdateType,
@@ -885,78 +878,42 @@ static void AddHairClusterAABBPass(
 	const FBoxSphereBounds TransformedBounds = Bounds.TransformBy(InRenLocalToTranslatedWorld);
 	Instance->HairGroupPublicData->bClusterAABBValid = UpdateType == EHairAABBUpdateType::UpdateClusterAABB;
 
-	// bNeedDeformation: 
-	// * If the instance has deformation (simulation, global interpolation, skinning, then we update all clusters (for voxelization purpose) and the instance AABB
-	// * If the instance is static, we only update the instance AABB
-	const uint32 GroupSize = ComputeGroupSize();
-
 	FHairClusterAABBCS::FParameters* Parameters = GraphBuilder.AllocParameters<FHairClusterAABBCS::FParameters>();
 	Parameters->CPUBoundMin = (FVector3f)TransformedBounds.GetBox().Min;
 	Parameters->CPUBoundMax = (FVector3f)TransformedBounds.GetBox().Max;
 	Parameters->LocalToTranslatedWorldMatrix = FMatrix44f(InRenLocalToTranslatedWorld.ToMatrixWithScale());
 	Parameters->RenderDeformedPositionBuffer = RenderPositionBufferSRV;
 	Parameters->RenderDeformedOffsetBuffer = RenderDeformedOffsetBuffer;
-	Parameters->TotalClusterCount = 1;
-
-	Parameters->PointCount = Instance->HairGroupPublicData->GetActiveStrandsPointCount();
+	Parameters->CurveCount = ActiveCurveCount;
+	Parameters->ClusterCount = ClusterData ? ClusterData->ClusterCount : 1;
+	Parameters->RenCurveBuffer = RegisterAsSRV(GraphBuilder, Instance->Strands.RestResource->CurveBuffer);
+	Parameters->OutClusterAABBBuffer = ClusterAABBData.ClusterAABBBuffer.UAV;
+	Parameters->OutGroupAABBBuffer = ClusterAABBData.GroupAABBBuffer.UAV;
 
 	if (ShaderPrintData)
 	{
-		ShaderPrint::SetParameters(GraphBuilder, *ShaderPrintData, Parameters->ShaderDrawParameters);
-		
-	}
-
-	if (UpdateType == EHairAABBUpdateType::UpdateClusterAABB)
-	{
-		check(ClusterData);
-		check(ClusterData->GroupSize1D == GroupSize);
-		check(ClusterData->CulledCluster2DIndirectArgsBuffer);
-		check(ClusterData->CulledClusterCountBuffer);
-
-		FRDGBufferRef ClusterIdBuffer = GraphBuilder.RegisterExternalBuffer(ClusterData->ClusterIdBuffer);
-		FRDGBufferRef ClusterIndexOffsetBuffer = GraphBuilder.RegisterExternalBuffer(ClusterData->ClusterIndexOffsetBuffer);
-		FRDGBufferRef ClusterIndexCountBuffer  = GraphBuilder.RegisterExternalBuffer(ClusterData->ClusterIndexCountBuffer);
-
-		Parameters->TotalClusterCount = ClusterData->ClusterCount;
-		Parameters->ClusterVertexIdBuffer = RegisterAsSRV(GraphBuilder, *ClusterData->ClusterVertexIdBuffer);
-		Parameters->ClusterIdBuffer = GraphBuilder.CreateSRV(ClusterIdBuffer, PF_R32_UINT);
-		Parameters->ClusterIndexOffsetBuffer = GraphBuilder.CreateSRV(ClusterIndexOffsetBuffer, PF_R32_UINT);
-		Parameters->ClusterIndexCountBuffer = GraphBuilder.CreateSRV(ClusterIndexCountBuffer, PF_R32_UINT);
-		Parameters->CulledClusterCountBuffer = GraphBuilder.CreateSRV(ClusterData->CulledClusterCountBuffer, PF_R32_UINT);
-		Parameters->IndirectBuffer = GraphBuilder.CreateSRV(ClusterData->CulledCluster2DIndirectArgsBuffer, PF_R32_UINT);
-		Parameters->IndirectArgsBuffer = ClusterData->CulledCluster2DIndirectArgsBuffer;
-
-		// Sanity check
-		check(ClusterData->ClusterCount == ClusterIdBuffer->Desc.NumElements);
-		check(ClusterData->ClusterCount == ClusterIndexOffsetBuffer->Desc.NumElements);
-		check(ClusterData->ClusterCount == ClusterIndexCountBuffer->Desc.NumElements);
-		// Currently disabled this sanity check as the RDG allocation can return a larger buffer (for reuse purpose)
-		//check(ClusterData->ClusterCount * 6 == ClusterAABBData.ClusterAABBBuffer.Buffer->Desc.NumElements);
+		ShaderPrint::SetParameters(GraphBuilder, *ShaderPrintData, Parameters->ShaderDrawParameters);		
 	}
-	Parameters->HairStrandsVFTODO_CullingIndirectBuffer = DrawIndirectRasterComputeBuffer; // Used for checking max vertex count
-	Parameters->OutClusterAABBBuffer = ClusterAABBData.ClusterAABBBuffer.UAV;
-	Parameters->OutGroupAABBBuffer = ClusterAABBData.GroupAABBBuffer.UAV;
-
 
 	FHairClusterAABBCS::FPermutationDomain PermutationVector;
-	PermutationVector.Set<FHairClusterAABBCS::FGroupSize>(GroupSize);
 	PermutationVector.Set<FHairClusterAABBCS::FCPUAABB>(UpdateType == EHairAABBUpdateType::UpdateClusterAABB ? 0u : 1u);
 	TShaderMapRef<FHairClusterAABBCS> ComputeShader(ShaderMap, PermutationVector);
 
 	if (UpdateType == EHairAABBUpdateType::UpdateClusterAABB)
 	{
+		FIntVector DispatchCount(ClusterData->ClusterCount, 1, 1);
 		FComputeShaderUtils::AddPass(
 			GraphBuilder,
-			RDG_EVENT_NAME("HairStrands::ClusterAABB(%s)", TEXT("Cluster,Group"), TEXT("Group Only")),
+			RDG_EVENT_NAME("HairStrands::ClusterAABB(Cluster,Group)"),
 			ComputeShader,
 			Parameters,
-			Parameters->IndirectArgsBuffer, 0);
+			DispatchCount);
 	}
 	else
 	{
 		FComputeShaderUtils::AddPass(
 			GraphBuilder,
-			RDG_EVENT_NAME("HairStrands::ClusterAABB(%s)", TEXT("Cluster,Group"), TEXT("Group Only")),
+			RDG_EVENT_NAME("HairStrands::ClusterAABB(Group Only - CPU bound)"),
 			ComputeShader,
 			Parameters,
 			FIntVector(1,1,1));
@@ -2124,6 +2081,7 @@ void ComputeHairStrandsInterpolation(
 					AddHairClusterAABBPass(
 						GraphBuilder,
 						ShaderMap,
+						ActiveCurveCount,
 						TranslatedWorldOffset,
 						ShaderPrintData,
 						UpdateType,
diff --git a/Engine/Plugins/Runtime/HairStrands/Source/HairStrandsCore/Public/GroomResources.h b/Engine/Plugins/Runtime/HairStrands/Source/HairStrandsCore/Public/GroomResources.h
--- a/Engine/Plugins/Runtime/HairStrands/Source/HairStrandsCore/Public/GroomResources.h
+++ b/Engine/Plugins/Runtime/HairStrands/Source/HairStrandsCore/Public/GroomResources.h
@@ -332,6 +332,7 @@ struct FHairStrandsClusterCullingResource : public FHairCommonResource
 		Total += GetBufferTotalNumBytes(ClusterLODInfoBuffer);
 		Total += GetBufferTotalNumBytes(CurveToClusterIdBuffer);
 		Total += GetBufferTotalNumBytes(ClusterVertexIdBuffer);
+		Total += GetBufferTotalNumBytes(PointLODBuffer);
 		return Total;
 	}
 
@@ -345,6 +346,9 @@ struct FHairStrandsClusterCullingResource : public FHairCommonResource
 	/* Concatenated data for each cluster: list of VertexId pointed to by ClusterInfoBuffer */
 	FRDGExternalBuffer ClusterVertexIdBuffer;
 
+	/* Min. LOD a which a Point becomes visible */
+	FRDGExternalBuffer PointLODBuffer;
+
 	FHairStrandsClusterCullingBulkData& BulkData;
 };
 
diff --git a/Engine/Plugins/Runtime/HairStrands/Source/HairStrandsCore/Public/HairStrandsDatas.h b/Engine/Plugins/Runtime/HairStrands/Source/HairStrandsCore/Public/HairStrandsDatas.h
--- a/Engine/Plugins/Runtime/HairStrands/Source/HairStrandsCore/Public/HairStrandsDatas.h
+++ b/Engine/Plugins/Runtime/HairStrands/Source/HairStrandsCore/Public/HairStrandsDatas.h
@@ -736,6 +736,7 @@ struct HAIRSTRANDSCORE_API FHairStrandsClusterCullingData
 	TArray<uint32>				CurveToClusterIds;
 	TArray<uint32>				ClusterVertexIds;
 	TArray<FHairLODInfo>		LODInfos;
+	TArray<uint32>				PointLODs;
 
 	uint32 ClusterCount = 0;		// Number of clusters
 	uint32 PointCount = 0;			// Number of points
@@ -778,7 +779,8 @@ struct HAIRSTRANDSCORE_API FHairStrandsClusterCullingBulkData : FHairStrandsBulk
 		FHairBulkContainer	PackedClusterInfos;		// Size - ClusterCount
 		FHairBulkContainer	ClusterLODInfos;		// Size - ClusterLODCount
 		FHairBulkContainer	CurveToClusterIds;		// Size - CurveCount
-		FHairBulkContainer	ClusterVertexIds;		// Size - VertexLODCount
+		FHairBulkContainer	ClusterVertexIds;		// Size - VertexLODCount // TODO REMOVE
+		FHairBulkContainer	PointLODs;				// Size - PointCount / HAIR_POINT_LOD_COUNT_PER_UINT
 	} Data;
 };
 
diff --git a/Engine/Shaders/Private/HairStrands/HairStrandsClusterCulling.usf b/Engine/Shaders/Private/HairStrands/HairStrandsClusterCulling.usf
--- a/Engine/Shaders/Private/HairStrands/HairStrandsClusterCulling.usf
+++ b/Engine/Shaders/Private/HairStrands/HairStrandsClusterCulling.usf
@@ -285,452 +285,364 @@ void MainClusterCullingIndClearCS(uint3 DispatchThreadId : SV_DispatchThreadID)
 #endif
 
 /////////////////////////////////////////////////////////////////////////////////////////////////////////////////////////////
-/// 0. Cluster culling  and LOD selection
-#ifdef SHADER_CLUSTERCULLING
-
-float4x4 ProjectionMatrix;
-float4x4 WorldToClipMatrix;
-float3 CameraWorldPos;
-uint NumConvexHullPlanes;
-float LODForcedIndex;
+// Write out culled index buffex & radius scale buffer
+
+#if SHADER_CLUSTER_CULL
+
+#include "../ShaderPrint.ush"  
+
+#define NUM_THREADS_PER_GROUP GROUP_SIZE
+#include "../ThreadGroupPrefixSum.ush"
+
+//uint ForceLOD;
+float LODIndex;
 float LODBias;
-int bIsHairGroupVisible;
-float4 ViewFrustumConvexHull[6];
+uint CurveCount;
+uint CurvePerGroup;
+uint DispatchCountX;
 
-Buffer<int> ClusterAABBBuffer;
+Buffer<uint> CurveBuffer;
+Buffer<uint> CurveToClusterIdBuffer;
+Buffer<uint> PointLODBuffer;
 StructuredBuffer<FPackedHairClusterInfo> ClusterInfoBuffer;
 StructuredBuffer<FPackedHairClusterLODInfo> ClusterLODInfoBuffer; 
 
-RWBuffer<uint> GlobalClusterIdBuffer;
-RWBuffer<uint> GlobalIndexStartBuffer;
-RWBuffer<uint> GlobalIndexCountBuffer;
-RWBuffer<float>GlobalRadiusScaleBuffer;
-RWBuffer<uint> DispatchIndirectParametersClusterCount;
-RWBuffer<uint> DispatchIndirectParametersVertexCount;
-RWBuffer<uint> DrawIndirectParameters; 
+RWBuffer<uint>  OutPointCounter;
+RWBuffer<uint>  OutIndexBuffer;
+RWBuffer<float> OutRadiusScaleBuffer;
+RWBuffer<uint>  OutCulledCurveBuffer;
 
-#if PERMUTATION_DEBUGAABBBUFFER
-RWStructuredBuffer<FHairClusterDebugInfo> ClusterDebugInfoBuffer;
+#ifndef GROUP_SIZE
+#error GROUP_SIZE needs to be Undefined
 #endif
 
-bool ViewFrustumIntersectSphere(float3 SphereOrigin, float SphereRadius)
-{
-	for (uint PlaneIndex = 0; PlaneIndex < NumConvexHullPlanes; PlaneIndex++)
-	{
-		float4 PlaneData = ViewFrustumConvexHull[PlaneIndex];
-		float PlaneDistance = dot(PlaneData.xyz, SphereOrigin) - PlaneData.w;
 
-		if (PlaneDistance > SphereRadius)
-		{
-			return false;
-		}
-	}
-	return true;
-}
+#define CURVE_PER_GROUP (GROUP_SIZE / PERMUTATION_POINT_PER_CURVE)
+#define MAX_CURVE CURVE_PER_GROUP
+#define MAX_POINT PERMUTATION_POINT_PER_CURVE
 
-[numthreads(64, 1, 1)]
-void MainClusterCullingCS(uint3 DispatchThreadId : SV_DispatchThreadID)
+uint UnpackPointLOD(uint InPacked, uint InIndex)
 {
-	const uint ClusterId = DispatchThreadId.x;
-	if (ClusterId < ClusterCount)
-	{
-		// Note on the cluster culling & LOD selection
-		// Run:
-		// 1. Frustum culling
-		// 2. LOD selection
-		//
-		// When a cluster is culled:
-		// . Cluster's strands position still updated (HairInterpolation), for the lowest LOD
-		// . Cluster's AABB is still updated, to be valid next frame (HairClusterAABB), with lowest LOD position
-		// . Cluster's strands are not rendered.
-		// However only the strands position of the selected LOD are updated
-
-		const uint ClusterId  = DispatchThreadId.x;
-		const uint ClusterId6 = ClusterId * 6;
-		checkBufferAccessSlow(ClusterAABBBuffer, ClusterId6 + 5);
-		const float3 Min = float3(int3(ClusterAABBBuffer[ClusterId6 + 0], ClusterAABBBuffer[ClusterId6 + 1], ClusterAABBBuffer[ClusterId6 + 2]));
-		const float3 Max = float3(int3(ClusterAABBBuffer[ClusterId6 + 3], ClusterAABBBuffer[ClusterId6 + 4], ClusterAABBBuffer[ClusterId6 + 5]));
-
-		const float3 SphereOrigin = (Min + Max)*0.5f;
-		const float  SphereRadius = length(Max - SphereOrigin);
-		const float  ScreenSize = ComputeBoundsScreenSize(SphereOrigin, SphereRadius, CameraWorldPos, ProjectionMatrix);
-		
-		const float3 CamToSphere = SphereOrigin - CameraWorldPos.xyz;
-		const float CamToSphereDistance = length(CamToSphere);
-		bool bVisible = (bIsHairGroupVisible == 1) && ViewFrustumIntersectSphere(CamToSphere, SphereRadius);
-		FHairClusterLOD ClusterLOD = (FHairClusterLOD)0;
-
-		// Always force clusters to be flagged as (forced) visible. If a cluster is not visible, its lowest LOD will be selected. 
-		// This ensures that clusters position are updated which is used for two things:
-		// * Update the cluster AABB, which are used for voxelization allocation purpose
-		// * Update vertex positions used for voxelization 
-		bool bForceVisible = bIsHairGroupVisible == 1;
-
-		if (bVisible || bForceVisible) 
-		{
-			float4 NotClampedRect = float4(0.0f, 1.0f, 0.0f, 1.0f);
-			if (bVisible || bForceVisible)
-			{
-				const FHairClusterInfo ClusterInfo = UnpackHairClusterInfo(ClusterInfoBuffer[ClusterId]);
-
-				const float ForceLOD = LODForcedIndex >= 0 ? LODForcedIndex : (bForceVisible && !bVisible ? float(ClusterInfo.LODCount-1) : -1.f);
-				if (ForceLOD < 0)
-				{
-					ClusterLOD = GetHairClusterLOD(ClusterInfo, ClusterLODInfoBuffer, ScreenSize, LODBias);
-				}
-				else
-				{
-					ClusterLOD = GetHairClusterLOD(ClusterInfo, ClusterLODInfoBuffer, ForceLOD);
-				}
-
-				// Update visibility based on the cluster visibility flag set up by users
-				bVisible = IsLODVisible(ClusterInfo, ClusterLOD.LOD);
-
-				if ((bVisible || bForceVisible) && ClusterLOD.VertexCount > 0)
-				{
-					uint VisibleClusterId = 0;
-					InterlockedAdd(DispatchIndirectParametersClusterCount[0], 1, VisibleClusterId);
-
-					//uint PreviousVertexCount = 0;
-					InterlockedAdd(DrawIndirectParameters[0], ClusterLOD.VertexCount * HAIR_POINT_TO_VERTEX/*, PreviousVertexCount*/);
-
-					checkBufferAccessSlow(GlobalClusterIdBuffer, VisibleClusterId);
-					GlobalClusterIdBuffer[VisibleClusterId]	  = ClusterId;
-					checkBufferAccessSlow(GlobalIndexStartBuffer, VisibleClusterId);
-					GlobalIndexStartBuffer[VisibleClusterId]  = ClusterLOD.VertexOffset;
-					checkBufferAccessSlow(GlobalIndexCountBuffer, VisibleClusterId);
-					GlobalIndexCountBuffer[VisibleClusterId]  = ClusterLOD.VertexCount;
-					checkBufferAccessSlow(GlobalRadiusScaleBuffer, VisibleClusterId);
-					GlobalRadiusScaleBuffer[VisibleClusterId] = ClusterLOD.RadiusScale;
-				}
-			}
-		}
-
-		#if PERMUTATION_DEBUGAABBBUFFER
-		{
-			FHairClusterDebugInfo ClusterDebug;			
-			ClusterDebug.GroupIndex = 0;
-			ClusterDebug.LOD = bVisible ? ClusterLOD.LOD : -1;
-			ClusterDebug.VertexCount = ClusterLOD.VertexCount;
-			ClusterDebug.CurveCount=0;
-
-			ClusterDebugInfoBuffer[ClusterId] = ClusterDebug;
-		}
-		#endif
-	}
+	return (InPacked >> (InIndex * HAIR_POINT_LOD_BIT_COUNT)) & 0xF;
 }
 
+uint GetGlobalCurveIndex(uint ThreadIndex, uint2 GroupIndex)
+{
+	const uint LinearGroupIndex = GroupIndex.x + GroupIndex.y * DispatchCountX;
+#if PERMUTATION_POINT_PER_CURVE == 64
+	return LinearGroupIndex;
+#else
+	return ThreadIndex + LinearGroupIndex * CURVE_PER_GROUP;
 #endif
+}
 
-/////////////////////////////////////////////////////////////////////////////////////////////////////////////////////////////
-/// 0. Indirect dispatch buffer preparation for compute raster visibility/voxel pass
-#ifdef SHADER_PREPAREINDIRECTDISPATCH
-
-RWBuffer<uint> DrawIndirectBuffer; 
-RWBuffer<uint> DispatchIndirectBuffer; 
+//-------------------------------------------------------------------------------------------------
+#if PERMUTATION_POINT_PER_CURVE < 64
 
-[numthreads(1, 1, 1)]
-void MainClusterCullingPrepareIndirectDispatchCS(
-	uint3 DispatchThreadId : SV_DispatchThreadID, 
-	uint3 GroupId : SV_GroupID, 
-	uint3 GroupThreadId : SV_GroupThreadID)
-{
-	if (GroupId.x == 0 && GroupThreadId.x == 0)
-	{
-		// DrawIndirectBuffer contains the total number of vertex count needed for drawing quad triangle strips, so 6 vertices per quad / control point
-		const uint VertexCount    = DrawIndirectBuffer[0];
-		const uint ControlPoint   = VertexCount / (HAIR_POINT_TO_VERTEX);
-		const uint DispatchCount  = DivideAndRoundUp(ControlPoint, HAIR_VERTEXCOUNT_GROUP_SIZE);
-
-		DispatchIndirectBuffer[0] = DispatchCount;
-		DispatchIndirectBuffer[1] = 1;
-		DispatchIndirectBuffer[2] = 1;
-		DispatchIndirectBuffer[3] = ControlPoint;
-	}
-}
+// Permutation for curves having less that 32 control points
+// In this permutation we process several curves with the same thread group
 
-#endif
+groupshared uint CurvePointIndex[MAX_CURVE];
+groupshared uint CurvePointCount[MAX_CURVE];
+groupshared uint ClusterLOD[MAX_CURVE];
+groupshared float ClusterRadiusScale[MAX_CURVE];
+groupshared uint LODCount[MAX_CURVE];
 
-/////////////////////////////////////////////////////////////////////////////////////////////////////////////////////////////
-/// 1. Indirect dispatch buffer preparation for rendering and prefix sums
-#ifdef SHADER_PREPAREINDIRECTDRAW
+groupshared uint AllocationOffset;
+groupshared uint AllocationCount;
 
-uint GroupSize1D;
-Buffer<uint> DispatchIndirectParametersClusterCount;
+// Store locally the culled point index (encoded into 8 bits, and store 4 per uint)
+groupshared uint CulledPointIndices[MAX_CURVE * MAX_POINT];
+groupshared uint CulledPointRadiusScales[MAX_CURVE * MAX_POINT];
+groupshared float CulledPointFadeScales[MAX_CURVE * MAX_POINT];
+groupshared uint ValidPoints[MAX_CURVE][MAX_POINT];
 
-RWBuffer<uint> DispatchIndirectParametersClusterCount1D;
-RWBuffer<uint> DispatchIndirectParametersClusterCount2D;
-RWBuffer<uint> DispatchIndirectParametersClusterCountDiv512;
-RWBuffer<uint> DispatchIndirectParametersClusterCountDiv512Div512;
+groupshared uint WriteOffset0[MAX_CURVE];
+groupshared uint WriteOffsetN[MAX_CURVE];
 
-[numthreads(1, 1, 1)]
-void MainClusterCullingPrepareIndirectDrawsCS(uint3 DispatchThreadId : SV_DispatchThreadID, uint3 GroupId : SV_GroupID, uint3 GroupThreadId : SV_GroupThreadID)
+[numthreads(GROUP_SIZE, 1, 1)]
+void Main(uint3 DispatchThreadId : SV_DispatchThreadID, uint ThreadIndex : SV_GroupIndex, uint2 GroupIndex : SV_GroupID)
 {
-	if (GroupId.x == 0 && GroupThreadId.x == 0)
+	// 1. Read the Curve data
+	if (ThreadIndex <= CURVE_PER_GROUP)
 	{
-		uint ClusterCount = DispatchIndirectParametersClusterCount[0];
-		uint ClusterCountDivUp512 = (ClusterCount + PREFIXSUM_PARALLEL - 1) / PREFIXSUM_PARALLEL;
-		uint ClusterCountDivUp512DivUp512 = (ClusterCountDivUp512 + PREFIXSUM_PARALLEL - 1) / PREFIXSUM_PARALLEL;
-		uint3 DispatchGroupCount2D = ComputeDispatchGroupCount2D(ClusterCount);
-
-		WriteDispatchIndirectArgs(DispatchIndirectParametersClusterCountDiv512, 0, ClusterCountDivUp512, 1, 1);
-		WriteDispatchIndirectArgs(DispatchIndirectParametersClusterCountDiv512Div512, 0, ClusterCountDivUp512DivUp512, 1, 1);
-		WriteDispatchIndirectArgs(DispatchIndirectParametersClusterCount2D, 0, DispatchGroupCount2D.x, DispatchGroupCount2D.y, 1);	
-		WriteDispatchIndirectArgs(DispatchIndirectParametersClusterCount1D, 0, (ClusterCount + GroupSize1D - 1) / GroupSize1D, 1, 1);
-	}
-}
+		const uint LocalCurveIndex  = ThreadIndex;
+		const uint GlobalCurveIndex = GetGlobalCurveIndex(ThreadIndex, GroupIndex);
 
-#endif
-
-/////////////////////////////////////////////////////////////////////////////////////////////////////////////////////////////
-/// 2. PER BLOCK OPTIMISATION
-#ifdef SHADER_PERBLOCKPREFIXSUM
-
-Buffer<uint> DispatchIndirectParametersClusterCount;
-Buffer<uint> GlobalIndexCountBuffer;
-
-RWBuffer<uint2> PerBlocklIndexCountPreFixSumBuffer;
-RWBuffer<uint>  PerBlocklTotalIndexCountBuffer;
+		// 1.1 Cluster ID
+		const uint ClusterId = CurveToClusterIdBuffer[GlobalCurveIndex];
+	
+		// 1.2 Get cluster info (i.e. map all info per LOD for a given cluster ID)
+		const FHairClusterInfo ClusterInfo = UnpackHairClusterInfo(ClusterInfoBuffer[ClusterId]);
+	
+		// 1.3 Derive effective LOD data (radius scale)
+		const FHairClusterLOD ClusterLODInfo = GetHairClusterLOD(ClusterInfo, ClusterLODInfoBuffer, LODIndex);
+		ClusterLOD[LocalCurveIndex] 		 = ClusterLODInfo.LOD;
+		ClusterRadiusScale[LocalCurveIndex]  = ClusterLODInfo.RadiusScale;
+		LODCount[LocalCurveIndex]			 = ClusterInfo.LODCount;
 
-groupshared uint PrefixSumClusterCount[2][PREFIXSUM_PARALLEL];
+		// 1.4 Curve's point count/offset
+		const FHairCurve Curve = UnpackHairCurve(CurveBuffer[GlobalCurveIndex]);
+		CurvePointIndex[LocalCurveIndex] = Curve.PointIndex;
+		CurvePointCount[LocalCurveIndex] = Curve.PointCount;
 
-[numthreads(PREFIXSUM_PARALLEL, 1, 1)] 
-void MainPerBlockPreFixSumCS(uint3 DispatchThreadId : SV_DispatchThreadID, uint3 GroupId : SV_GroupID, uint3 GroupThreadId : SV_GroupThreadID)
-{
-	uint ItemCount = DispatchIndirectParametersClusterCount[0];
-	uint ItemCountRoundUp = PREFIXSUM_PARALLEL * ((ItemCount + PREFIXSUM_PARALLEL - 1) / PREFIXSUM_PARALLEL);
-	uint WaveStart = PREFIXSUM_PARALLEL * GroupId.x;
+	}
+	if (ThreadIndex == 0)
+	{
+		AllocationCount = 0;
+		AllocationOffset = 0;
+	}
+	GroupMemoryBarrierWithGroupSync();
 
-	if (WaveStart < ItemCountRoundUp)
+	// 2. Mark valid point and compute their final position
 	{
-		// Load current block data
-		uint GlobalIndex = WaveStart + GroupThreadId.x;
-		if (GlobalIndex < ItemCount)
-		{
-			checkBufferAccessSlow(GlobalIndexCountBuffer, GlobalIndex);
+		const uint CurveIndex        = ThreadIndex / uint(MAX_POINT);
+		const uint CurrentPointIndex = ThreadIndex % uint(MAX_POINT);
+
+		// 2.2 Mark point having their LOD visibility smaller than the requested LOD level
+		float RadiusFade = 1.0f;
+		ValidPoints[CurveIndex][CurrentPointIndex] = 0; 
+		if (CurveIndex < CurveCount)
+		{	
+			if (CurrentPointIndex < CurvePointCount[CurveIndex])
+			{
+				const uint PointIndex   = CurrentPointIndex + CurvePointIndex[CurveIndex];
+				const uint BlockIt      = PointIndex >> HAIR_POINT_LOD_COUNT_PER_UINT_DIV_AS_SHIFT;
+				const uint LocalPointIt = PointIndex - BlockIt * HAIR_POINT_LOD_COUNT_PER_UINT;
+				const uint Packed 		= PointLODBuffer[BlockIt];
+				const uint MinLOD 		= UnpackPointLOD(Packed, LocalPointIt);
 
-			const uint IndexCount = GlobalIndexCountBuffer[GlobalIndex];
-			PrefixSumClusterCount[0][GroupThreadId.x] = IndexCount;
-		}
-		else
-		{
-			PrefixSumClusterCount[0][GroupThreadId.x] = 0;
-		}
-		GroupMemoryBarrierWithGroupSync();
+				// When a curve reach its last LOD, fade its radius progressively
+				RadiusFade = ceil(ClusterLOD[CurveIndex]) >= LODCount[CurveIndex]-1 ? 1.f : saturate(MinLOD+1 - ClusterLOD[CurveIndex]);
 
-		// Pre fix sum the current block using ping-pong in shared memory
-		uint Src = 0;
-		for (uint i = 1; i <= (PREFIXSUM_PARALLEL / 2); i *= 2)
-		{
-			if (GroupThreadId.x >= i)
-			{
-				PrefixSumClusterCount[1 - Src][GroupThreadId.x] = PrefixSumClusterCount[Src][GroupThreadId.x] + PrefixSumClusterCount[Src][GroupThreadId.x - i];
-			}
-			else
-			{
-				PrefixSumClusterCount[1 - Src][GroupThreadId.x] = PrefixSumClusterCount[Src][GroupThreadId.x];
+				// Prune control point based on their MinLOD value
+				if (floor(ClusterLOD[CurveIndex]) <= MinLOD)
+				{
+					ValidPoints[CurveIndex][CurrentPointIndex] = 1;
+				}
 			}
-			Src = 1 - Src;
-			GroupMemoryBarrierWithGroupSync();
 		}
-
-		// Write out result, taking into account previous blocks sum
-		uint Result = 0;
-		if (GlobalIndex < ItemCount)
+	
+		// 2.3 Count/Compact visible points
 		{
-			Result = PrefixSumClusterCount[Src][GroupThreadId.x];
-
-			checkBufferAccessSlow(PerBlocklIndexCountPreFixSumBuffer, GlobalIndex);
-
-			if (GroupThreadId.x == 0)
+			uint LocalAllocationCount = 0;
+			const uint bIsValid = ValidPoints[CurveIndex][CurrentPointIndex];
+			const uint PrefixSum = ThreadGroupPrefixSum(bIsValid, ThreadIndex, LocalAllocationCount);
+			if (ThreadIndex == 0)
 			{
-				PerBlocklIndexCountPreFixSumBuffer[GlobalIndex] = uint2(0, Result);
+				AllocationCount = LocalAllocationCount;
 			}
-			else
+			if (bIsValid)
 			{
-				PerBlocklIndexCountPreFixSumBuffer[GlobalIndex] = uint2(PrefixSumClusterCount[Src][GroupThreadId.x - 1], Result);
+				CulledPointIndices[PrefixSum]		= CurrentPointIndex + CurvePointIndex[CurveIndex];
+				CulledPointRadiusScales[PrefixSum]	= ClusterRadiusScale[CurveIndex];
+				CulledPointFadeScales[PrefixSum] 	= RadiusFade;
 			}
-		}
 
-		if (GroupThreadId.x == PREFIXSUM_LASTTHREAD)
-		{
-			checkBufferAccessSlow(PerBlocklTotalIndexCountBuffer, GroupId.x);
-			PerBlocklTotalIndexCountBuffer[GroupId.x] = PrefixSumClusterCount[Src][GroupThreadId.x];
+			if (CurveIndex < CurveCount)
+			{
+				if (CurrentPointIndex == 0) 							WriteOffset0[CurveIndex] = PrefixSum;
+				if (CurrentPointIndex == CurvePointCount[CurveIndex]-1)	WriteOffsetN[CurveIndex] = PrefixSum;
+			}
 		}
 	}
-}
-
-#endif
 
-/////////////////////////////////////////////////////////////////////////////////////////////////////////////////////////////
-/// 3. PREFIX SIUM
-#ifdef SHADER_CLUSTERCULLINGPREFIXSUM
-
-Buffer<uint> DispatchIndirectParametersClusterCount;
-Buffer<uint> GlobalIndexCountBuffer;
-
-RWBuffer<uint2> GlobalIndexCountPreFixSumBuffer;
-
-groupshared uint PreviousSum;
-groupshared uint PrefixSumClusterCount[2][PREFIXSUM_PARALLEL];
+	// 4. Global allocate
+	if (ThreadIndex == 0)
+	{
+		InterlockedAdd(OutPointCounter[0], AllocationCount, AllocationOffset);
+	}
+	GroupMemoryBarrierWithGroupSync();
 
-[numthreads(PREFIXSUM_PARALLEL, 1, 1)]
-void MainClusterCullingPreFixSumCS(uint3 DispatchThreadId : SV_DispatchThreadID, uint3 GroupId : SV_GroupID)
-{
-#if 0
+	// 5. Write out index/scale
+	if (ThreadIndex < AllocationCount)
+	{
+		OutIndexBuffer[AllocationOffset + ThreadIndex]       = CulledPointIndices[ThreadIndex];
+		OutRadiusScaleBuffer[AllocationOffset + ThreadIndex] = CulledPointRadiusScales[ThreadIndex] * CulledPointFadeScales[ThreadIndex];
+	}
 
-	// Slow reference
-	if (DispatchThreadId.x == 0)
+	// 6. Write out curve data with visible points for subsequent passes
+	if (ThreadIndex < CURVE_PER_GROUP)
 	{
-		// ugly linear 1 thread....
-		uint Sum = 0;
-		uint ItemCount = DispatchIndirectParametersClusterCount[0];
-		for (uint i = 0; i < ItemCount; ++i)
+		const uint LocalCurveIndex  = ThreadIndex;
+		const uint GlobalCurveIndex = GetGlobalCurveIndex(ThreadIndex, GroupIndex);;
+		if (GlobalCurveIndex < CurveCount)
 		{
-			uint IndexCount = GlobalIndexCountBuffer[i];
-			GlobalIndexCountPreFixSumBuffer[i] = uint2(Sum, Sum + IndexCount);
-			Sum += IndexCount;
+			FHairCurve CulledCurve;
+			CulledCurve.PointIndex = WriteOffset0[LocalCurveIndex];
+			CulledCurve.PointCount =(WriteOffsetN[LocalCurveIndex] - WriteOffset0[LocalCurveIndex]) + 1;
+
+			OutCulledCurveBuffer[GlobalCurveIndex] = PackHairCurve(CulledCurve);
 		}
 	}
-
-#else
-
-	if (DispatchThreadId.x == 0)
+}
+#endif // PERMUTATION_POINT_PER_CURVE < 64
+//-------------------------------------------------------------------------------------------------
+#if PERMUTATION_POINT_PER_CURVE == 64
+
+// Permutation is for curve having more than 32 control point per curve
+// In this permutation we process 1 curve per thread group
+
+groupshared uint CurvePointIndex;
+groupshared uint CurvePointCount;
+groupshared float ClusterLOD;
+groupshared float ClusterRadiusScale;
+groupshared uint LODCount;
+
+groupshared uint AllocationOffset;
+groupshared uint AllocationCount;
+
+// Store locally the culled point index (encoded into 8 bits, and store 4 per uint)
+groupshared uint CulledPointIndices[HAIR_MAX_NUM_POINT_PER_CURVE]; 
+groupshared uint ValidPoints[HAIR_MAX_NUM_POINT_PER_CURVE];
+groupshared float CulledPointFadeScales[HAIR_MAX_NUM_POINT_PER_CURVE];
+
+// Version for curve with >32 points
+[numthreads(GROUP_SIZE, 1, 1)]
+void Main(uint3 DispatchThreadId : SV_DispatchThreadID, uint ThreadIndex : SV_GroupIndex, uint2 GroupIndex : SV_GroupID)
+{
+	// 1. Read the Curve data
+	const uint CurveIndex = GetGlobalCurveIndex(ThreadIndex, GroupIndex);
+	
+	if (ThreadIndex == 0)
 	{
-		// Reset previous block sum
-		PreviousSum = 0;
+		// 1.1 Cluster ID
+		const uint ClusterId = CurveToClusterIdBuffer[CurveIndex];
+	
+		// 1.2 Get cluster info (i.e. map all info per LOD for a given cluster ID)
+		const FHairClusterInfo ClusterInfo = UnpackHairClusterInfo(ClusterInfoBuffer[ClusterId]);
+	
+		// 1.3 Derive effective LOD data (radius scale)
+		const FHairClusterLOD ClusterLODInfo = GetHairClusterLOD(ClusterInfo, ClusterLODInfoBuffer, LODIndex);
+		ClusterLOD = ClusterLODInfo.LOD;
+		ClusterRadiusScale = ClusterLODInfo.RadiusScale;
+		LODCount = ClusterInfo.LODCount;
+
+		// 1.4 Curve's point count/offset
+		const FHairCurve Curve = UnpackHairCurve(CurveBuffer[CurveIndex]);
+		CurvePointIndex = Curve.PointIndex;
+		CurvePointCount = Curve.PointCount;
+
+		AllocationCount = 0;
+		AllocationOffset = 0;
 	}
 	GroupMemoryBarrierWithGroupSync();
 
-	uint ItemCount = DispatchIndirectParametersClusterCount[0];
-	uint ItemCountRoundUp = PREFIXSUM_PARALLEL * ((ItemCount + PREFIXSUM_PARALLEL - 1) / PREFIXSUM_PARALLEL);
-	uint WaveStart = 0;
 
-	while (WaveStart < ItemCountRoundUp)
+	// 2. Mark valid point and compute their final position
 	{
-		// Load current block data
-		uint GlobalIndex = WaveStart + DispatchThreadId.x;
-		if (GlobalIndex < ItemCount)
-		{
-			const uint IndexCount = GlobalIndexCountBuffer[GlobalIndex];
-			PrefixSumClusterCount[0][DispatchThreadId.x] = IndexCount;
-		}
-		else
-		{
-			PrefixSumClusterCount[0][DispatchThreadId.x] = 0;
-		}
-		GroupMemoryBarrierWithGroupSync();
-
-		// Pre fix sum the current block using ping-pong in shared memory
-		uint Src = 0;
-		for (uint i = 1; i <= (PREFIXSUM_PARALLEL/2); i *= 2)
-		{
-			if (DispatchThreadId.x >= i)
-			{
-				PrefixSumClusterCount[1 - Src][DispatchThreadId.x] = PrefixSumClusterCount[Src][DispatchThreadId.x] + PrefixSumClusterCount[Src][DispatchThreadId.x - i];
+		const uint LoopCount = DivideAndRoundUp(CurvePointCount, GROUP_SIZE);
+		for (uint LoopIt = 0; LoopIt < LoopCount; ++LoopIt)
+		{		
+			const uint CurrentPointIndex = ThreadIndex + LoopIt * GROUP_SIZE;
+	
+			// 2.2 Mark point having their LOD visibility smaller than the requested LOD level
+			float RadiusFade = 1.0f;
+			ValidPoints[CurrentPointIndex] = 0; 
+			if (CurveIndex < CurveCount)
+			{	
+				if (CurrentPointIndex < CurvePointCount)
+				{
+					const uint PointIndex   = CurrentPointIndex + CurvePointIndex;
+					const uint BlockIt      = PointIndex >> HAIR_POINT_LOD_COUNT_PER_UINT_DIV_AS_SHIFT;
+					const uint LocalPointIt = PointIndex - BlockIt * HAIR_POINT_LOD_COUNT_PER_UINT;
+					const uint Packed 		= PointLODBuffer[BlockIt];
+					const uint MinLOD 		= UnpackPointLOD(Packed, LocalPointIt);
+
+					// When a curve reach its last LOD, fade its radius progressively
+					RadiusFade = ceil(ClusterLOD) >= LODCount-1 ? 1.f : saturate(MinLOD+1 - ClusterLOD);
+
+					// Prune control point based on their MinLOD value
+					if (floor(ClusterLOD) <= MinLOD)
+					{
+						ValidPoints[CurrentPointIndex] = 1;
+					}
+				}
 			}
-			else
+		
+			// 2.3 Count/Compact visible points
 			{
-				PrefixSumClusterCount[1 - Src][DispatchThreadId.x] = PrefixSumClusterCount[Src][DispatchThreadId.x];
+				uint IndexOffset = AllocationCount;
+				uint LocalAllocationCount = 0;
+				const uint bIsValid = ValidPoints[CurrentPointIndex];
+				const uint PrefixSum = ThreadGroupPrefixSum(bIsValid, ThreadIndex, LocalAllocationCount); // Has a ThreadMemoryGroupSync internally
+				if (bIsValid)
+				{
+					CulledPointIndices[IndexOffset + PrefixSum] = CurrentPointIndex;
+					CulledPointFadeScales[IndexOffset + PrefixSum] = RadiusFade;
+				}
+
+				if (ThreadIndex == 0)
+				{
+					AllocationCount += LocalAllocationCount;
+				}
 			}
-			Src = 1 - Src;
-			GroupMemoryBarrierWithGroupSync();
 		}
+	}
 
-		// Write out result, taking into account previous blocks sum
-		uint Result = 0;
-		if (GlobalIndex < ItemCount)
-		{
-			Result = PrefixSumClusterCount[Src][DispatchThreadId.x];
+	// 4. Global allocate
+	if (ThreadIndex == 0)
+	{
+		InterlockedAdd(OutPointCounter[0], AllocationCount, AllocationOffset);
+	}
+	GroupMemoryBarrierWithGroupSync();
 
-			checkBufferAccessSlow(GlobalIndexCountPreFixSumBuffer, GlobalIndex);
-			if (DispatchThreadId.x == 0)
-			{
-				GlobalIndexCountPreFixSumBuffer[GlobalIndex] = uint2(0, Result) + PreviousSum;
-			}
-			else
+	// 5. Write out index/scale
+	if (CurveIndex < CurveCount)
+	{
+		const uint LoopCount = DivideAndRoundUp(AllocationCount, GROUP_SIZE);
+		for (uint LoopIt = 0; LoopIt < LoopCount; ++LoopIt)
+		{
+			const uint CurrentPointIndex = ThreadIndex + LoopIt * GROUP_SIZE;
+			if (CurrentPointIndex < AllocationCount)
 			{
-				GlobalIndexCountPreFixSumBuffer[GlobalIndex] = uint2(PrefixSumClusterCount[Src][DispatchThreadId.x - 1], Result) + PreviousSum;
-			}
-		}
-		GroupMemoryBarrierWithGroupSync();
+				const uint CullPointIndex = CulledPointIndices[CurrentPointIndex];
+				const float FadeScale = CulledPointFadeScales[CurrentPointIndex];
 
-		// Update previous block sum
-		if (DispatchThreadId.x == PREFIXSUM_LASTTHREAD)
-		{
-			PreviousSum += Result;
+				OutIndexBuffer[AllocationOffset + CurrentPointIndex]       = CurvePointIndex + CullPointIndex;
+				OutRadiusScaleBuffer[AllocationOffset + CurrentPointIndex] = ClusterRadiusScale * FadeScale;
+			}
 		}
-
-		WaveStart += PREFIXSUM_PARALLEL;
 	}
 
-#endif
-
-	// TODO try multiple pass parallelisation, with dispatch indirect
-	// https://developer.nvidia.com/gpugems/GPUGems3/gpugems3_ch39.html
-	// http://citeseerx.ist.psu.edu/viewdoc/download?doi=10.1.1.177.8755&rep=rep1&type=pdf
+	// 6. Write out curve data with visible points for subsequent passes
+	if (ThreadIndex == 0)
+	{	
+		FHairCurve CulledCurve;
+		CulledCurve.PointCount = AllocationCount;
+		CulledCurve.PointIndex = AllocationOffset;
+		OutCulledCurveBuffer[CurveIndex] = PackHairCurve(CulledCurve);
+	}
 }
+#endif // PERMUTATION_POINT_PER_CURVE == 64
+//-------------------------------------------------------------------------------------------------
 
-#endif
+#endif // SHADER_CLUSTER_CULL
 
 /////////////////////////////////////////////////////////////////////////////////////////////////////////////////////////////
-/// 4. PER BLOCK OPTIMISATION
-/// PER BLOCK OPTIMISATION
-#if SHADER_CLUSTERCULLINGCOMPACTVERTEXIDLOCALBLOCK
-
-Buffer<uint> DispatchIndirectParametersClusterCount;
-Buffer<uint> DispatchIndirectParametersClusterCount2D;
-Buffer<uint> DispatchIndirectParametersClusterCountDiv512;
+// Build indirect Draw/Dispatch based on culled curve output
 
-Buffer<uint2> PerBlocklIndexCountPreFixSumBuffer;
-Buffer<uint2> PerBlocklTotalIndexCountPreFixSumBuffer;
+#if SHADER_CLUSTER_CULL_ARGS
 
-Buffer<uint> GlobalIndexStartBuffer;
-Buffer<uint> GlobalIndexCountBuffer;
-Buffer<float> GlobalRadiusScaleBuffer;
-Buffer<uint> ClusterVertexIdBuffer;
+ByteAddressBuffer PointCounterBuffer;
+RWBuffer<uint> RWIndirectDrawArgsBuffer;
+RWBuffer<uint> RWIndirectDispatchArgsBuffer;
 
-RWBuffer<uint> CulledCompactedIndexBuffer;
-RWBuffer<float> CulledCompactedRadiusScaleBuffer;
-
-[numthreads(COMPACTCS_NUMTHREAD, 1, 1)]
-void MainClusterCullingCompactVertexIdsLocalBlockCS(uint3 DispatchThreadId : SV_DispatchThreadID, uint3 GroupId : SV_GroupID, uint3 GroupThreadId : SV_GroupThreadID)
+[numthreads(1, 1, 1)]
+void Main(uint DispatchThreadId : SV_DispatchThreadID, uint ThreadIndex : SV_GroupIndex)
 {
-	const uint LoadedClusterCount = DispatchIndirectParametersClusterCount[0];
-	const uint DispatchCountX = DispatchIndirectParametersClusterCount2D[0];
-	const uint ClusterId = GroupId.x + GroupId.y * DispatchCountX;
-	if (ClusterId >= LoadedClusterCount)
-	{
-		return;
-	}
-
-	// For each cluster, a group is splatting the indices inside its target range (evaluated after a prefix sum)
-
-	const uint ClusterSrcStartIndex = GlobalIndexStartBuffer[ClusterId];
-	const uint ClusterSrcIndexCount = GlobalIndexCountBuffer[ClusterId];
-	const float ClusterRadiusScale  = GlobalRadiusScaleBuffer[ClusterId];
-
-	uint2 TopBlockPrefixSum = PerBlocklTotalIndexCountPreFixSumBuffer[ClusterId / PREFIXSUM_PARALLEL];
-	uint2 PerBlockPrefixSum = PerBlocklIndexCountPreFixSumBuffer[ClusterId];
-
-	uint WriteToIndex = TopBlockPrefixSum.x + PerBlockPrefixSum.x;
-	uint LoopCount = (ClusterSrcIndexCount + COMPACTCS_NUMTHREAD - 1) / COMPACTCS_NUMTHREAD;
-
-	for (uint Loop = 0; Loop < LoopCount; ++Loop)
-	{
-		uint GroupWriteOffset = Loop * COMPACTCS_NUMTHREAD;
-		uint ThreadWriteOffset = GroupWriteOffset + GroupThreadId.x;
-
-		if (ThreadWriteOffset < ClusterSrcIndexCount)
-		{
-			checkBufferAccessSlow(CulledCompactedIndexBuffer, WriteToIndex + ThreadWriteOffset);
-			CulledCompactedIndexBuffer[WriteToIndex + ThreadWriteOffset] = ClusterVertexIdBuffer[ClusterSrcStartIndex + ThreadWriteOffset];
-			checkBufferAccessSlow(CulledCompactedRadiusScaleBuffer, WriteToIndex + ThreadWriteOffset);
-			CulledCompactedRadiusScaleBuffer[WriteToIndex + ThreadWriteOffset] = D_USE_CONTINUOUSLOD ? ClusterRadiusScale : 0;
-		}
-	}
+	const uint PointCount = PointCounterBuffer.Load(0);
+	const uint DispatchCount = DivideAndRoundUp(PointCount, HAIR_VERTEXCOUNT_GROUP_SIZE);
+
+	RWIndirectDrawArgsBuffer[0] = PointCount * HAIR_POINT_TO_VERTEX;
+	RWIndirectDrawArgsBuffer[1] = 1; // 1 instance
+	RWIndirectDrawArgsBuffer[2] = 0;
+	RWIndirectDrawArgsBuffer[3] = 0;
+
+	RWIndirectDispatchArgsBuffer[0] = DispatchCount;
+	RWIndirectDispatchArgsBuffer[1] = 1;
+	RWIndirectDispatchArgsBuffer[2] = 1;
+	RWIndirectDispatchArgsBuffer[3] = PointCount;	
 }
 
-#endif
-
-/////////////////////////////////////////////////////////////////////////////////////////////////////////////////////////////
+#endif // SHADER_CLUSTER_CULL_ARGS
diff --git a/Engine/Shaders/Private/HairStrands/HairStrandsInterpolation.usf b/Engine/Shaders/Private/HairStrands/HairStrandsInterpolation.usf
--- a/Engine/Shaders/Private/HairStrands/HairStrandsInterpolation.usf
+++ b/Engine/Shaders/Private/HairStrands/HairStrandsInterpolation.usf
@@ -590,34 +590,24 @@ void MainCS(uint2 DispatchThreadId : SV_DispatchThreadID, uint GroupThread1D : S
 
 #if SHADER_CLUSTERAABB
 
-uint TotalClusterCount;
+uint CurveCount;
+uint ClusterCount;
 float3 CPUBoundMin;
 float3 CPUBoundMax;
-
-uint PointCount;
-
 float4x4 LocalToTranslatedWorldMatrix;
 
+Buffer<uint> RenCurveBuffer;
 Buffer<uint4> RenderDeformedPositionBuffer;
-Buffer<uint> ClusterVertexIdBuffer;
-Buffer<uint> ClusterIdBuffer;
-Buffer<uint> ClusterIndexOffsetBuffer;
-Buffer<uint> ClusterIndexCountBuffer;
-Buffer<uint> CulledClusterCountBuffer;
-
-Buffer<uint> IndirectBuffer;
-Buffer<uint> HairStrandsVFTODO_CullingIndirectBuffer;
 Buffer<float4> RenderDeformedOffsetBuffer;
 
 RWBuffer<int> OutClusterAABBBuffer; 	// Cluster data packed as {uint3 AABBMin, uint3 AABBMax}
 RWBuffer<int> OutGroupAABBBuffer; 		// Group data packed as {uint3 AABBMin, uint3 AABBMax}
 
-#if GROUP_SIZE!=64 && GROUP_SIZE!=32
-#error Unhandled group size
-#endif
-
 groupshared float3 SharedClusterAABBMin[GROUP_SIZE];
 groupshared float3 SharedClusterAABBMax[GROUP_SIZE];
+groupshared uint CurvePointIndex;
+groupshared uint CurvePointCount;
+groupshared float3 OutHairPositionOffset;
 
 [numthreads(GROUP_SIZE, 1, 1)]
 void ClusterAABBEvaluationCS(uint3 DispatchThreadId : SV_DispatchThreadID, uint3 GroupId : SV_GroupID, uint3 GroupThreadId : SV_GroupThreadID, uint ThreadIndex : SV_GroupIndex)
@@ -632,99 +622,98 @@ void ClusterAABBEvaluationCS(uint3 DispatchThreadId : SV_DispatchThreadID, uint3
 		WriteHairAABB(0, Bound, OutGroupAABBBuffer);
 	}
 #else // PERMUTATION_USE_CPU_AABB
-	const uint CulledClusterCount = CulledClusterCountBuffer[0];
-	const uint DispatchCountX = IndirectBuffer[0];
-	const uint ClusterFetchIndex = GroupId.x + GroupId.y * DispatchCountX;
-	if (ClusterFetchIndex >= CulledClusterCount)
+
+	// Use the first K curves to to compute cluster bound. This is an approximation, 
+	// but since curves are now shuffled, it covers most of the groom
+	const uint ClusterIndex = GroupId.x;
+	if (ThreadIndex == 0)
 	{
-		return;
+		CurvePointIndex = 0;
+		CurvePointCount = 0;
+
+		const uint CurveIndex = GroupId.x;
+		if (CurveIndex < CurveCount)
+		{
+			const FHairCurve RenCurve = UnpackHairCurve(RenCurveBuffer[CurveIndex]);
+			CurvePointIndex = RenCurve.PointIndex;
+			CurvePointCount = RenCurve.PointCount;
+		}
+		OutHairPositionOffset = RenderDeformedOffsetBuffer[0].xyz;
 	}
-	const uint ClusterId   = ClusterIdBuffer[ClusterFetchIndex];
-	const uint IndexOffset = ClusterIndexOffsetBuffer[ClusterFetchIndex];
-	const uint IndexCount = min(ClusterIndexCountBuffer[ClusterFetchIndex], GROUP_SIZE);
-	
-	const uint VertexCountAfterCulling = HairStrandsVFTODO_CullingIndirectBuffer[3];
+	GroupMemoryBarrierWithGroupSync();
 
 	int3 ClusterAABBMin = 9999999.0f;
 	int3 ClusterAABBMax =-9999999.0f;
 
-	const uint LoopCount = (IndexCount + GROUP_SIZE - 1) / GROUP_SIZE;
+	// TODO culled point have invalid position
+	const uint LoopCount = DivideAndRoundUp(CurvePointCount, GROUP_SIZE);
 	for (uint LoopIt = 0; LoopIt < LoopCount; ++LoopIt)
 	{
-		uint GroupReadOffset = LoopIt * GROUP_SIZE;
-		uint ThreadReadOffset = GroupReadOffset + GroupThreadId.x;
-
-		if (ThreadReadOffset < IndexCount)
+		const uint LocalPointIndex = ThreadIndex + LoopIt * GROUP_SIZE;
+		if (LocalPointIndex < CurvePointCount)
 		{
-			const float3 OutHairPositionOffset = RenderDeformedOffsetBuffer[0].xyz;
-			const uint PointIndex = ClusterVertexIdBuffer[IndexOffset + ThreadReadOffset];
-			if (PointIndex < PointCount)
-			{
-				const FHairControlPoint CP = UnpackHairControlPoint(RenderDeformedPositionBuffer[PointIndex], OutHairPositionOffset, 1, 1, 1);
+			const uint GlobalPointIndex = CurvePointIndex + LocalPointIndex;
+			const FHairControlPoint CP = UnpackHairControlPoint(RenderDeformedPositionBuffer[GlobalPointIndex], OutHairPositionOffset, 1, 1, 1);
 
-				if (all(IsFinite(CP.Position)))
-				{
-					const int3 WorldSpaceCentimeters = int3(mul(float4(CP.Position, 1.0f), LocalToTranslatedWorldMatrix).xyz);
-					const int3 WorldSpaceCentimetersMin = WorldSpaceCentimeters - 1;
-					const int3 WorldSpaceCentimetersMax = WorldSpaceCentimeters + 1;
+			if (all(IsFinite(CP.Position)))
+			{
+				const int3 WorldSpaceCentimeters = int3(mul(float4(CP.Position, 1.0f), LocalToTranslatedWorldMatrix).xyz);
+				const int3 WorldSpaceCentimetersMin = WorldSpaceCentimeters - 1;
+				const int3 WorldSpaceCentimetersMax = WorldSpaceCentimeters + 1;
 
-					ClusterAABBMin = min(ClusterAABBMin, WorldSpaceCentimetersMin);
-					ClusterAABBMax = max(ClusterAABBMax, WorldSpaceCentimetersMax);
-				}
+				ClusterAABBMin = min(ClusterAABBMin, WorldSpaceCentimetersMin);
+				ClusterAABBMax = max(ClusterAABBMax, WorldSpaceCentimetersMax);
 			}
 		}
 	}
 
 	// Write each thread result to shared memory
-	const uint LocalIndex = GroupThreadId.x;
-	SharedClusterAABBMin[LocalIndex] = ClusterAABBMin;
-	SharedClusterAABBMax[LocalIndex] = ClusterAABBMax;
+	SharedClusterAABBMin[ThreadIndex] = ClusterAABBMin;
+	SharedClusterAABBMax[ThreadIndex] = ClusterAABBMax;
 
 	// Do a local reduction in shared memory. Assumes ClusterLOD0.VertexCount>64 to have all min&max values valid.
-#if GROUP_SIZE==64
 	GroupMemoryBarrierWithGroupSync();
-	if (LocalIndex < 32)
+	if (ThreadIndex < 32)
 	{
-		SharedClusterAABBMin[LocalIndex] = min(SharedClusterAABBMin[LocalIndex], SharedClusterAABBMin[LocalIndex + 32]);
-		SharedClusterAABBMax[LocalIndex] = max(SharedClusterAABBMax[LocalIndex], SharedClusterAABBMax[LocalIndex + 32]);
+		SharedClusterAABBMin[ThreadIndex] = min(SharedClusterAABBMin[ThreadIndex], SharedClusterAABBMin[ThreadIndex + 32]);
+		SharedClusterAABBMax[ThreadIndex] = max(SharedClusterAABBMax[ThreadIndex], SharedClusterAABBMax[ThreadIndex + 32]);
 	}
-#endif
 	GroupMemoryBarrierWithGroupSync();
-	if (LocalIndex < 16)
+	if (ThreadIndex < 16)
 	{
-		SharedClusterAABBMin[LocalIndex] = min(SharedClusterAABBMin[LocalIndex], SharedClusterAABBMin[LocalIndex + 16]);
-		SharedClusterAABBMax[LocalIndex] = max(SharedClusterAABBMax[LocalIndex], SharedClusterAABBMax[LocalIndex + 16]);
+		SharedClusterAABBMin[ThreadIndex] = min(SharedClusterAABBMin[ThreadIndex], SharedClusterAABBMin[ThreadIndex + 16]);
+		SharedClusterAABBMax[ThreadIndex] = max(SharedClusterAABBMax[ThreadIndex], SharedClusterAABBMax[ThreadIndex + 16]);
 	}
 	GroupMemoryBarrierWithGroupSync();
-	if (LocalIndex < 8)
+	if (ThreadIndex < 8)
 	{
-		SharedClusterAABBMin[LocalIndex] = min(SharedClusterAABBMin[LocalIndex], SharedClusterAABBMin[LocalIndex + 8]);
-		SharedClusterAABBMax[LocalIndex] = max(SharedClusterAABBMax[LocalIndex], SharedClusterAABBMax[LocalIndex + 8]);
+		SharedClusterAABBMin[ThreadIndex] = min(SharedClusterAABBMin[ThreadIndex], SharedClusterAABBMin[ThreadIndex + 8]);
+		SharedClusterAABBMax[ThreadIndex] = max(SharedClusterAABBMax[ThreadIndex], SharedClusterAABBMax[ThreadIndex + 8]);
 	}
 	// No hardware has SIMD Vector unit operating in sync with less than 16 threads per group. So can skip sync now.
 	//GroupMemoryBarrierWithGroupSync();
-	if (LocalIndex < 4)
+	if (ThreadIndex < 4)
 	{
-		SharedClusterAABBMin[LocalIndex] = min(SharedClusterAABBMin[LocalIndex], SharedClusterAABBMin[LocalIndex + 4]);
-		SharedClusterAABBMax[LocalIndex] = max(SharedClusterAABBMax[LocalIndex], SharedClusterAABBMax[LocalIndex + 4]);
+		SharedClusterAABBMin[ThreadIndex] = min(SharedClusterAABBMin[ThreadIndex], SharedClusterAABBMin[ThreadIndex + 4]);
+		SharedClusterAABBMax[ThreadIndex] = max(SharedClusterAABBMax[ThreadIndex], SharedClusterAABBMax[ThreadIndex + 4]);
 	}
 	//GroupMemoryBarrierWithGroupSync();
-	if (LocalIndex < 2)
+	if (ThreadIndex < 2)
 	{
-		SharedClusterAABBMin[LocalIndex] = min(SharedClusterAABBMin[LocalIndex], SharedClusterAABBMin[LocalIndex + 2]);
-		SharedClusterAABBMax[LocalIndex] = max(SharedClusterAABBMax[LocalIndex], SharedClusterAABBMax[LocalIndex + 2]);
+		SharedClusterAABBMin[ThreadIndex] = min(SharedClusterAABBMin[ThreadIndex], SharedClusterAABBMin[ThreadIndex + 2]);
+		SharedClusterAABBMax[ThreadIndex] = max(SharedClusterAABBMax[ThreadIndex], SharedClusterAABBMax[ThreadIndex + 2]);
 	}
 	//GroupMemoryBarrierWithGroupSync();
-	if (LocalIndex < 1)
+	if (ThreadIndex < 1)
 	{
-		SharedClusterAABBMin[LocalIndex] = min(SharedClusterAABBMin[LocalIndex], SharedClusterAABBMin[LocalIndex + 1]);
-		SharedClusterAABBMax[LocalIndex] = max(SharedClusterAABBMax[LocalIndex], SharedClusterAABBMax[LocalIndex + 1]);
+		SharedClusterAABBMin[ThreadIndex] = min(SharedClusterAABBMin[ThreadIndex], SharedClusterAABBMin[ThreadIndex + 1]);
+		SharedClusterAABBMax[ThreadIndex] = max(SharedClusterAABBMax[ThreadIndex], SharedClusterAABBMax[ThreadIndex + 1]);
 	}
 
 	// Write out hair cluster AABB
-	if (LocalIndex == 0)
+	if (ThreadIndex == 0)
 	{
-		const uint ClusterIdx6 = ClusterId * 6;
+		const uint ClusterIdx6 = ClusterIndex * 6;
 		OutClusterAABBBuffer[ClusterIdx6 + 0] = SharedClusterAABBMin[0].x;
 		OutClusterAABBBuffer[ClusterIdx6 + 1] = SharedClusterAABBMin[0].y;
 		OutClusterAABBBuffer[ClusterIdx6 + 2] = SharedClusterAABBMin[0].z;
@@ -746,7 +735,7 @@ void ClusterAABBEvaluationCS(uint3 DispatchThreadId : SV_DispatchThreadID, uint3
 	#if DEBUG_ENABLE
 	if (ThreadIndex == 0)
 	{
-		const float4 DebugColor = float4(ColorMapViridis(ClusterId / float(TotalClusterCount)),1);
+		const float4 DebugColor = float4(ColorMapViridis(ClusterIndex / float(TotalClusterCount)),1);
 		AddAABBTWS(SharedClusterAABBMin[0], SharedClusterAABBMax[0], DebugColor);
 		AddReferentialWS(LocalToWorldMatrix, 10);
 	}
diff --git a/Engine/Shaders/Shared/HairStrandsDefinitions.h b/Engine/Shaders/Shared/HairStrandsDefinitions.h
--- a/Engine/Shaders/Shared/HairStrandsDefinitions.h
+++ b/Engine/Shaders/Shared/HairStrandsDefinitions.h
@@ -87,4 +87,11 @@
 	}
 
 #define PACK_CURVE_HAIR_ATTRIBUTE_OFFSETS(Out, In) PACK_HAIR_ATTRIBUTE_OFFSETS(Out, In, HAIR_CURVE_ATTRIBUTE_OFFSET_COUNT, HAIR_CURVE_ATTRIBUTE_COUNT)
-#define PACK_POINT_HAIR_ATTRIBUTE_OFFSETS(Out, In) PACK_HAIR_ATTRIBUTE_OFFSETS(Out, In, HAIR_POINT_ATTRIBUTE_OFFSET_COUNT, HAIR_POINT_ATTRIBUTE_COUNT)
\ No newline at end of file
+#define PACK_POINT_HAIR_ATTRIBUTE_OFFSETS(Out, In) PACK_HAIR_ATTRIBUTE_OFFSETS(Out, In, HAIR_POINT_ATTRIBUTE_OFFSET_COUNT, HAIR_POINT_ATTRIBUTE_COUNT)
+
+// Data for LODing points of a curve
+// The number of bits per vertex describe the maxium LOD decomposition
+// Each curve's point encode its minium LOD at which it becomes active
+#define HAIR_POINT_LOD_BIT_COUNT 4
+#define HAIR_POINT_LOD_COUNT_PER_UINT (32u/4u)
+#define HAIR_POINT_LOD_COUNT_PER_UINT_DIV_AS_SHIFT 3
\ No newline at end of file
diff --git a/Engine/Source/Runtime/Renderer/Private/HairStrands/HairStrandsInterface.cpp b/Engine/Source/Runtime/Renderer/Private/HairStrands/HairStrandsInterface.cpp
--- a/Engine/Source/Runtime/Renderer/Private/HairStrands/HairStrandsInterface.cpp
+++ b/Engine/Source/Runtime/Renderer/Private/HairStrands/HairStrandsInterface.cpp
@@ -340,6 +340,7 @@ void FHairGroupPublicData::Allocate(FRDGBuilder& GraphBuilder)
 	InternalCreateVertexBufferRDG(GraphBuilder, sizeof(int32), ClusterCount * 6, EPixelFormat::PF_R32_SINT, ClusterAABBBuffer, TEXT("Hair.Cluster_ClusterAABBBuffer"), GetOwnerName());
 	InternalCreateVertexBufferRDG(GraphBuilder, sizeof(int32), 6, EPixelFormat::PF_R32_SINT, GroupAABBBuffer, TEXT("Hair.Cluster_GroupAABBBuffer"), GetOwnerName());
 
+	InternalCreateVertexBufferRDG(GraphBuilder, sizeof(int32), RestCurveCount, EPixelFormat::PF_R32_UINT, CulledCurveBuffer, TEXT("Hair.Cluster_CulledCurveBuffer"), GetOwnerName());
 	InternalCreateVertexBufferRDG(GraphBuilder, sizeof(int32), RestPointCount, EPixelFormat::PF_R32_UINT, CulledVertexIdBuffer, TEXT("Hair.Cluster_CulledVertexIdBuffer"), GetOwnerName());
 	InternalCreateVertexBufferRDG(GraphBuilder, sizeof(float), RestPointCount, EPixelFormat::PF_R32_FLOAT, CulledVertexRadiusScaleBuffer, TEXT("Hair.Cluster_CulledVertexRadiusScaleBuffer"), GetOwnerName(), true);
 
@@ -359,6 +360,7 @@ void FHairGroupPublicData::Release()
 	DrawIndirectRasterComputeBuffer.Release();
 	ClusterAABBBuffer.Release();
 	GroupAABBBuffer.Release();
+	CulledCurveBuffer.Release();
 	CulledVertexIdBuffer.Release();
 	CulledVertexRadiusScaleBuffer.Release();
 	bIsInitialized = false;
@@ -376,6 +378,7 @@ uint32 FHairGroupPublicData::GetResourcesSize() const
 	Total += ExtractSize(DrawIndirectRasterComputeBuffer.Buffer);
 	Total += ExtractSize(ClusterAABBBuffer.Buffer);
 	Total += ExtractSize(GroupAABBBuffer.Buffer);
+	Total += ExtractSize(CulledCurveBuffer.Buffer);
 	Total += ExtractSize(CulledVertexIdBuffer.Buffer);
 	Total += ExtractSize(CulledVertexRadiusScaleBuffer.Buffer);
 	return Total;
diff --git a/Engine/Source/Runtime/Renderer/Private/HairStrands/HairStrandsVisibility.cpp b/Engine/Source/Runtime/Renderer/Private/HairStrands/HairStrandsVisibility.cpp
--- a/Engine/Source/Runtime/Renderer/Private/HairStrands/HairStrandsVisibility.cpp
+++ b/Engine/Source/Runtime/Renderer/Private/HairStrands/HairStrandsVisibility.cpp
@@ -2859,9 +2859,11 @@ void AddMeshDrawTransitionPass(
 			ExternalAccessQueue.Add(VFInput.Strands.PositionOffsetBuffer.Buffer);
 			ExternalAccessQueue.Add(VFInput.Strands.PrevPositionOffsetBuffer.Buffer);
 
+			FRDGBufferRef CulledCurveBuffer = Register(GraphBuilder, HairGroupPublicData->CulledCurveBuffer, ERDGImportedBufferFlags::None).Buffer;
 			FRDGBufferRef CulledVertexIdBuffer = Register(GraphBuilder, HairGroupPublicData->CulledVertexIdBuffer, ERDGImportedBufferFlags::None).Buffer;
 			FRDGBufferRef CulledVertexRadiusScaleBuffer = Register(GraphBuilder, HairGroupPublicData->CulledVertexRadiusScaleBuffer, ERDGImportedBufferFlags::None).Buffer;
 			FRDGBufferRef DrawIndirectBuffer = Register(GraphBuilder, HairGroupPublicData->DrawIndirectBuffer, ERDGImportedBufferFlags::None).Buffer;
+			ExternalAccessQueue.Add(CulledCurveBuffer);
 			ExternalAccessQueue.Add(CulledVertexIdBuffer);
 			ExternalAccessQueue.Add(CulledVertexRadiusScaleBuffer);
 			ExternalAccessQueue.Add(DrawIndirectBuffer, ERHIAccess::IndirectArgs);
diff --git a/Engine/Source/Runtime/Renderer/Public/HairStrandsInterface.h b/Engine/Source/Runtime/Renderer/Public/HairStrandsInterface.h
--- a/Engine/Source/Runtime/Renderer/Public/HairStrandsInterface.h
+++ b/Engine/Source/Runtime/Renderer/Public/HairStrandsInterface.h
@@ -235,9 +235,11 @@ public:
 	FRDGExternalBuffer& GetGroupAABBBuffer() { return GroupAABBBuffer; }
 	const FRDGExternalBuffer& GetGroupAABBBuffer() const { return GroupAABBBuffer; }
 
+	const FRDGExternalBuffer& GetCulledCurveBuffer() const { return CulledCurveBuffer; }
 	const FRDGExternalBuffer& GetCulledVertexIdBuffer() const { return CulledVertexIdBuffer; }
 	const FRDGExternalBuffer& GetCulledVertexRadiusScaleBuffer() const { return CulledVertexRadiusScaleBuffer; }
 
+	FRDGExternalBuffer& GetCulledCurveBuffer() { return CulledCurveBuffer; }
 	FRDGExternalBuffer& GetCulledVertexIdBuffer() { return CulledVertexIdBuffer; }
 	FRDGExternalBuffer& GetCulledVertexRadiusScaleBuffer() { return CulledVertexRadiusScaleBuffer; }
 
@@ -382,6 +384,7 @@ public:
 	bool bClusterAABBValid = false;
 
 	/* Culling & LODing results for a hair group */ // Better to be transient?
+	FRDGExternalBuffer CulledCurveBuffer;
 	FRDGExternalBuffer CulledVertexIdBuffer;
 	FRDGExternalBuffer CulledVertexRadiusScaleBuffer;
 	bool bCullingResultAvailable = false;
@@ -429,17 +432,21 @@ struct FHairStrandClusterData
 	{
 		uint32 ClusterCount = 0;
 		uint32 VertexCount = 0;
+		uint32 MaxPointPerCurve = 0;
 
 		float LODIndex = -1;
 		float LODBias = 0.0f;
 		bool bVisible = false;
 
+		FRDGExternalBuffer* CurveBuffer = nullptr;
+		FRDGExternalBuffer* PointLODBuffer = nullptr;
+
 		// See FHairStrandsClusterCullingResource fro details about those buffers.
 		FRDGExternalBuffer* GroupAABBBuffer = nullptr;
 		FRDGExternalBuffer* ClusterAABBBuffer = nullptr;
-		FRDGExternalBuffer* ClusterInfoBuffer = nullptr;
-		FRDGExternalBuffer* ClusterLODInfoBuffer = nullptr;
-		FRDGExternalBuffer* CurveToClusterIdBuffer = nullptr;
+		FRDGExternalBuffer* ClusterInfoBuffer = nullptr; // SRV
+		FRDGExternalBuffer* ClusterLODInfoBuffer = nullptr; // SRV
+		FRDGExternalBuffer* CurveToClusterIdBuffer = nullptr; // SRV
 		FRDGExternalBuffer* ClusterVertexIdBuffer = nullptr;
 
 		TRefCountPtr<FRDGPooledBuffer> ClusterIdBuffer;
@@ -447,6 +454,7 @@ struct FHairStrandClusterData
 		TRefCountPtr<FRDGPooledBuffer> ClusterIndexCountBuffer;
 
 		// Culling & LOD output
+		FRDGExternalBuffer* GetCulledCurveBuffer() const				{ return HairGroupPublicPtr ? &HairGroupPublicPtr->GetCulledCurveBuffer() : nullptr; }
 		FRDGExternalBuffer* GetCulledVertexIdBuffer() const				{ return HairGroupPublicPtr ? &HairGroupPublicPtr->GetCulledVertexIdBuffer() : nullptr; }
 		FRDGExternalBuffer* GetCulledVertexRadiusScaleBuffer() const	{ return HairGroupPublicPtr ? &HairGroupPublicPtr->GetCulledVertexRadiusScaleBuffer() : nullptr; }
 		bool GetCullingResultAvailable() const							{ return HairGroupPublicPtr ? HairGroupPublicPtr->GetCullingResultAvailable() : false; }
