Remove cluster index buffer.

diff --git a/Engine/Plugins/Runtime/HairStrands/Source/HairStrandsCore/Private/GroomAsset.cpp b/Engine/Plugins/Runtime/HairStrands/Source/HairStrandsCore/Private/GroomAsset.cpp
--- a/Engine/Plugins/Runtime/HairStrands/Source/HairStrandsCore/Private/GroomAsset.cpp
+++ b/Engine/Plugins/Runtime/HairStrands/Source/HairStrandsCore/Private/GroomAsset.cpp
@@ -129,7 +129,6 @@ uint32 FHairGroupPlatformData::FStrands::GetDataSize() const
 	CTotal += ClusterCullingBulkData.Header.CPULODScreenSize.GetAllocatedSize();
 	CTotal += ClusterCullingBulkData.Data.ClusterLODInfos.IsBulkDataLoaded()	? ClusterCullingBulkData.Data.ClusterLODInfos.GetBulkDataSize() : 0;
 	BTotal += ClusterCullingBulkData.Data.CurveToClusterIds.IsBulkDataLoaded()	? ClusterCullingBulkData.Data.CurveToClusterIds.GetBulkDataSize() : 0;
-	BTotal += ClusterCullingBulkData.Data.ClusterVertexIds.IsBulkDataLoaded()	? ClusterCullingBulkData.Data.ClusterVertexIds.GetBulkDataSize() : 0;
 	CTotal += ClusterCullingBulkData.Data.PackedClusterInfos.IsBulkDataLoaded()	? ClusterCullingBulkData.Data.PackedClusterInfos.GetBulkDataSize() : 0;
 
 	return BTotal + Total + CTotal;
diff --git a/Engine/Plugins/Runtime/HairStrands/Source/HairStrandsCore/Private/GroomBuilder.cpp b/Engine/Plugins/Runtime/HairStrands/Source/HairStrandsCore/Private/GroomBuilder.cpp
--- a/Engine/Plugins/Runtime/HairStrands/Source/HairStrandsCore/Private/GroomBuilder.cpp
+++ b/Engine/Plugins/Runtime/HairStrands/Source/HairStrandsCore/Private/GroomBuilder.cpp
@@ -40,7 +40,7 @@ static FAutoConsoleVariableRef CVarHairGroupIndexBuilder_MaxVoxelResolution(TEXT
 
 FString FGroomBuilder::GetVersion()
 {
-	return TEXT("v8r32");
+	return TEXT("v8r33");
 }
 
 namespace FHairStrandsDecimation
@@ -2992,10 +2992,6 @@ static void BuildClusterData(
 	Out.LODInfos.SetNum(LODCount);
 	Out.PointLODs.Init(0, Out.PointCount);
 
-	// Conservative allocation for inserting vertex indices for the various curves LOD
-	uint32* RawClusterVertexIds = new uint32[LODCount * InRenStrandsData.GetNumPoints()];
-	TAtomic<uint32> RawClusterVertexCount(0);
-
 	// 4. Write out cluster information
 	Out.ClusterLODInfos.SetNum(LODCount * Out.ClusterCount);
 	TArray<uint8> VertexLODMasks;
@@ -3020,8 +3016,6 @@ static void BuildClusterData(
 			&LocalClusterLODInfos,
 			&LocalCurveToClusterIds,
 			&VertexLODMasks,
-			&RawClusterVertexIds,
-			&RawClusterVertexCount
 		]
 	(uint32 ClusterIt)
 #else
@@ -3036,9 +3030,9 @@ static void BuildClusterData(
 		// Sort curve to have largest area first, so that lower area curves with less influence are removed first.
 		// This also helps the radius scaling to not explode.
 		Cluster.ClusterCurves.Sort([](const FClusterGrid::FCurve& A, const FClusterGrid::FCurve& B) -> bool
-			{
-				return A.Area > B.Area;
-			});
+		{
+			return A.Area > B.Area;
+		});
 
 		// 4.2 Compute cluster's area & fill in the vertex to cluster ID mapping
 		float ClusterArea = 0;
@@ -3138,26 +3132,7 @@ static void BuildClusterData(
 			}
 		}
 
-		// 4.5 Record/Insert vertex indices for each LOD of the current cluster
-		// Vertex offset is stored into the cluster LOD info
-		// Stores the accumulated vertex count per LOD
-		//
-		// ClusterVertexIds contains the vertex index of curve belonging to a cluster.
-		// Since for a given LOD, both the number of curve and vertices varies, we stores 
-		// this information per LOD.
-		//
-		//  Global Vertex index
-		//            v
-		// ||0 1 2 3 4 5 6 7 8 9 ||0 1 3 5 7 9 ||0 5 9 | |0 1 2 3 4 5 6 7 || 0 1 5 7 ||0 9 ||||11 12 ...
-		// ||____________________||____________||______| |________________||_________||____||||_____ _ _ 
-		// ||        LOD 0			 LOD 1		 LOD2  | |    LOD 0			 LOD 1	  LOD2 ||||  LOD 0
-		// ||__________________________________________| | ________________________________||||_____ _ _ 
-		// |                   Curve 0								Curve 1				    ||   Curve 0
-		// |________________________________________________________________________________||_____ _ _ 
-		//										Cluster 0										Cluster 1
-
-		TArray<uint32> LocalClusterVertexIds;
-		LocalClusterVertexIds.Reserve(LODCount * Cluster.ClusterCurves.Num() * 32); // Guestimate pre-allocation (32 points per curve in average)
+		// 4.5 Record/Insert vertex min LOD for each curve
 
 		FHairClusterInfo& ClusterInfo = LocalClusterInfos[ClusterIt];
 		ClusterInfo.LODCount = LODCount;
@@ -3165,9 +3140,6 @@ static void BuildClusterData(
 		for (uint8 LODIt = 0; LODIt < LODCount; ++LODIt)
 		{
 			FHairClusterLODInfo& ClusterLODInfo = LocalClusterLODInfos[ClusterInfo.LODInfoOffset + LODIt];
-			ClusterLODInfo.VertexOffset = LocalClusterVertexIds.Num(); // At the end, it will be the offset at which the data are inserted into ClusterVertexIds
-			ClusterLODInfo.VertexCount0 = 0;
-			ClusterLODInfo.VertexCount1 = 0;
 			ClusterLODInfo.RadiusScale0 = 0;
 			ClusterLODInfo.RadiusScale1 = 0;
 
@@ -3183,17 +3155,6 @@ static void BuildClusterData(
 					const uint8 LODMask = VertexLODMasks[GlobalPointIndex];
 					if (LODMask & (1 << LODIt))
 					{
-						// Count the number of vertices for all curves in the cluster as well as the vertex 
-						// of the remaining curves once the cluster has been decimated with the current LOD 
-						// settings
-						++ClusterLODInfo.VertexCount0;
-						if (CurveIt < NextCurveCount)
-						{
-							++ClusterLODInfo.VertexCount1;
-						}
-
-						LocalClusterVertexIds.Add(GlobalPointIndex);
-
 						// Store the minimum/coarser LOD at which the current point is visible
 						LocalPointLODs[GlobalPointIndex] = LODIt;
 					}
@@ -3201,35 +3162,19 @@ static void BuildClusterData(
 			}
 		}
 
-		// 4.5.1 Insert vertex indices for each LOD into the final array
-		// Since this runs in parallel, we prefill LocalClusterVertexIds with 
-		// all indices, then we insert the indices into the final array with a single allocation + memcopy
-		// We also patch the vertex offset so that it is correct
-		const uint32 AllocOffset = RawClusterVertexCount.AddExchange(LocalClusterVertexIds.Num());
-		FMemory::Memcpy(RawClusterVertexIds + AllocOffset, LocalClusterVertexIds.GetData(), LocalClusterVertexIds.Num() * sizeof(uint32));
-		for (uint8 LODIt = 0; LODIt < LODCount; ++LODIt)
-		{
-			FHairClusterLODInfo& ClusterLODInfo = LocalClusterLODInfos[ClusterInfo.LODInfoOffset + LODIt];
-			ClusterLODInfo.VertexOffset += AllocOffset;
-		}
-
 		// 4.6 Compute the radius scaling to preserve the cluster apperance as we decimate 
 		// the number of strands
 		for (uint8 LODIt = 0; LODIt < LODCount; ++LODIt)
 		{
 			// Compute the visible area for various orientation? 
 			// Reference: Stochastic Simplification of Aggregate Detail
-			float  LODArea = 0;
 			float  LODAvgRadiusRef = 0;
 			float  LODMaxRadiusRef = 0;
-			uint32 LODVertexCount = 0;
 
 			const uint32 ClusterCurveCount = LODCurveCount[LODIt];
 			for (uint32 CurveIt = 0; CurveIt < ClusterCurveCount; ++CurveIt)
 			{
 				const FClusterGrid::FCurve& ClusterCurve = Cluster.ClusterCurves[CurveIt];
-				LODVertexCount += ClusterCurve.Count;
-				LODArea += ClusterCurve.Area;
 				LODAvgRadiusRef += ClusterCurve.AvgRadius;
 				LODMaxRadiusRef = FMath::Max(LODMaxRadiusRef, ClusterCurve.MaxRadius);
 			}
@@ -3281,13 +3226,6 @@ static void BuildClusterData(
 		Out.CPULODScreenSize.Add(InSettings.LODs[LODIt].ScreenSize);
 		Out.LODVisibility.Add(InSettings.LODs[LODIt].bVisible);
 	}
-
-	// Copy the final value to the array which will be used to copy data to the GPU.
-	// This operations is not needer per se. We could just keep & use RawClusterVertexIds
-	Out.ClusterVertexIds.SetNum(RawClusterVertexCount);
-	FMemory::Memcpy(Out.ClusterVertexIds.GetData(), RawClusterVertexIds, RawClusterVertexCount * sizeof(uint32));
-
-	delete[] RawClusterVertexIds;
 }
 
 static void BuildClusterBulkData(
@@ -3299,7 +3237,6 @@ static void BuildClusterBulkData(
 	Out.Header.ClusterCount		= In.ClusterCount;
 	Out.Header.PointCount		= In.PointCount;
 	Out.Header.CurveCount		= In.CurveCount;
-	Out.Header.VertexLODCount	= In.ClusterVertexIds.Num();
 	Out.Header.ClusterLODCount 	= In.ClusterLODInfos.Num();
 
 	Out.Header.CPULODScreenSize	= In.CPULODScreenSize;
@@ -3310,11 +3247,9 @@ static void BuildClusterBulkData(
 	check(Out.Header.ClusterCount		== uint32(In.ClusterInfos.Num()));
 	check(Out.Header.ClusterLODCount	== uint32(In.ClusterLODInfos.Num()));
 	check(Out.Header.CurveCount			== uint32(In.CurveToClusterIds.Num()));
-	check(Out.Header.VertexLODCount		== uint32(In.ClusterVertexIds.Num()));
 	
 	HairStrandsBuilder::CopyToBulkData<FHairClusterLODInfoFormat>(Out.Data.ClusterLODInfos, In.ClusterLODInfos);
 	HairStrandsBuilder::CopyToBulkData<FHairClusterIndexFormat>(Out.Data.CurveToClusterIds, In.CurveToClusterIds);
-	HairStrandsBuilder::CopyToBulkData<FHairClusterIndexFormat>(Out.Data.ClusterVertexIds, In.ClusterVertexIds);
 
 	// Pack point LOD data
 	// |                          32bits                               |
diff --git a/Engine/Plugins/Runtime/HairStrands/Source/HairStrandsCore/Private/GroomResources.cpp b/Engine/Plugins/Runtime/HairStrands/Source/HairStrandsCore/Private/GroomResources.cpp
--- a/Engine/Plugins/Runtime/HairStrands/Source/HairStrandsCore/Private/GroomResources.cpp
+++ b/Engine/Plugins/Runtime/HairStrands/Source/HairStrandsCore/Private/GroomResources.cpp
@@ -1199,7 +1199,6 @@ void FHairStrandsClusterCullingResource::InternalAllocate(FRDGBuilder& GraphBuil
 	InternalCreateStructuredBufferRDG_FromHairBulkData<FHairClusterLODInfoFormat>(GraphBuilder, BulkData.Data.ClusterLODInfos, BulkData.Header.ClusterLODCount, ClusterLODInfoBuffer, ToHairResourceDebugName(TEXT("Hair.StrandsClusterCulling_ClusterLODInfoBuffer"), ResourceName), OwnerName, EHairResourceUsageType::Static);
 
 	InternalCreateVertexBufferRDG_FromHairBulkData<FHairClusterIndexFormat>(GraphBuilder, BulkData.Data.CurveToClusterIds, BulkData.Header.CurveCount, CurveToClusterIdBuffer, ToHairResourceDebugName(TEXT("Hair.StrandsClusterCulling_CurveToClusterIds"), ResourceName), OwnerName, EHairResourceUsageType::Static);
-	InternalCreateVertexBufferRDG_FromHairBulkData<FHairClusterIndexFormat>(GraphBuilder, BulkData.Data.ClusterVertexIds, BulkData.Header.VertexLODCount, ClusterVertexIdBuffer, ToHairResourceDebugName(TEXT("Hair.StrandsClusterCulling_ClusterVertexIds"), ResourceName), OwnerName, EHairResourceUsageType::Static);
 	
 	const uint32 EntryCount = FMath::DivideAndRoundUp(BulkData.Header.PointCount, HAIR_POINT_LOD_COUNT_PER_UINT);
 	InternalCreateVertexBufferRDG_FromHairBulkData<FHairClusterIndexFormat>(GraphBuilder, BulkData.Data.PointLODs, EntryCount, PointLODBuffer, ToHairResourceDebugName(TEXT("Hair.StrandsClusterCulling_PointLOD"), ResourceName), OwnerName, EHairResourceUsageType::Static);
@@ -1209,7 +1208,6 @@ void FHairStrandsClusterCullingResource::InternalRelease()
 {
 	ClusterInfoBuffer.Release();
 	ClusterLODInfoBuffer.Release();
-	ClusterVertexIdBuffer.Release();
 	CurveToClusterIdBuffer.Release();
 	PointLODBuffer.Release();
 	MaxAvailableCurveCount = 0;
diff --git a/Engine/Plugins/Runtime/HairStrands/Source/HairStrandsCore/Private/HairStrandsClusterCulling.cpp b/Engine/Plugins/Runtime/HairStrands/Source/HairStrandsCore/Private/HairStrandsClusterCulling.cpp
--- a/Engine/Plugins/Runtime/HairStrands/Source/HairStrandsCore/Private/HairStrandsClusterCulling.cpp
+++ b/Engine/Plugins/Runtime/HairStrands/Source/HairStrandsCore/Private/HairStrandsClusterCulling.cpp
@@ -279,7 +279,6 @@ void AddInstanceToClusterData(FHairGroupInstance* In, FHairStrandClusterData& Ou
 	HairGroupCluster.ClusterInfoBuffer = &In->Strands.ClusterCullingResource->ClusterInfoBuffer;
 	HairGroupCluster.ClusterLODInfoBuffer = &In->Strands.ClusterCullingResource->ClusterLODInfoBuffer;
 	HairGroupCluster.CurveToClusterIdBuffer = &In->Strands.ClusterCullingResource->CurveToClusterIdBuffer;
-	HairGroupCluster.ClusterVertexIdBuffer = &In->Strands.ClusterCullingResource->ClusterVertexIdBuffer;
 
 	HairGroupCluster.HairGroupPublicPtr = In->HairGroupPublicData;
 	HairGroupCluster.LODBias  = In->HairGroupPublicData->GetLODBias();
diff --git a/Engine/Plugins/Runtime/HairStrands/Source/HairStrandsCore/Private/HairStrandsClusterCulling.h b/Engine/Plugins/Runtime/HairStrands/Source/HairStrandsCore/Private/HairStrandsClusterCulling.h
--- a/Engine/Plugins/Runtime/HairStrands/Source/HairStrandsCore/Private/HairStrandsClusterCulling.h
+++ b/Engine/Plugins/Runtime/HairStrands/Source/HairStrandsCore/Private/HairStrandsClusterCulling.h
@@ -32,7 +32,6 @@ struct FHairStrandClusterData
 		FRDGExternalBuffer* ClusterInfoBuffer = nullptr; // SRV
 		FRDGExternalBuffer* ClusterLODInfoBuffer = nullptr; // SRV
 		FRDGExternalBuffer* CurveToClusterIdBuffer = nullptr; // SRV
-		FRDGExternalBuffer* ClusterVertexIdBuffer = nullptr;
 
 		TRefCountPtr<FRDGPooledBuffer> ClusterIdBuffer;
 		TRefCountPtr<FRDGPooledBuffer> ClusterIndexOffsetBuffer;
diff --git a/Engine/Plugins/Runtime/HairStrands/Source/HairStrandsCore/Private/HairStrandsDatas.cpp b/Engine/Plugins/Runtime/HairStrands/Source/HairStrandsCore/Private/HairStrandsDatas.cpp
--- a/Engine/Plugins/Runtime/HairStrands/Source/HairStrandsCore/Private/HairStrandsDatas.cpp
+++ b/Engine/Plugins/Runtime/HairStrands/Source/HairStrandsCore/Private/HairStrandsDatas.cpp
@@ -657,7 +657,6 @@ void FHairStrandsClusterCullingBulkData::Reset()
 	Header.ClusterLODCount = 0;
 	Header.PointCount = 0;
 	Header.CurveCount = 0;
-	Header.VertexLODCount = 0;
 
 	Header.LODVisibility.Empty();
 	Header.CPULODScreenSize.Empty();
@@ -665,14 +664,12 @@ void FHairStrandsClusterCullingBulkData::Reset()
 
 	Data.ClusterLODInfos.RemoveBulkData();
 	Data.CurveToClusterIds.RemoveBulkData();
-	Data.ClusterVertexIds.RemoveBulkData();
 	Data.PackedClusterInfos.RemoveBulkData();
 	Data.PointLODs.RemoveBulkData();
 
 	// Reset the bulk byte buffer to ensure the (serialize) data size is reset to 0
 	Data.ClusterLODInfos 	= FHairBulkContainer();
 	Data.CurveToClusterIds  = FHairBulkContainer();
-	Data.ClusterVertexIds 	= FHairBulkContainer();
 	Data.PackedClusterInfos = FHairBulkContainer();
 	Data.PointLODs			= FHairBulkContainer();
 }
@@ -681,8 +678,8 @@ void FHairStrandsClusterCullingBulkData::ResetLoadedSize()
 {
 	Data.ClusterLODInfos.LoadedSize 	= 0;
 	Data.CurveToClusterIds.LoadedSize 	= 0;
-	Data.ClusterVertexIds.LoadedSize 	= 0;
 	Data.PackedClusterInfos.LoadedSize	= 0;
+	Data.PointLODs.LoadedSize			= 0;
 }
 
 void FHairStrandsClusterCullingBulkData::SerializeHeader(FArchive& Ar, UObject* Owner)
@@ -691,7 +688,6 @@ void FHairStrandsClusterCullingBulkData::SerializeHeader(FArchive& Ar, UObject*
 	Ar << Header.ClusterLODCount;
 	Ar << Header.PointCount;
 	Ar << Header.CurveCount;
-	Ar << Header.VertexLODCount;
 	Ar << Header.LODVisibility;
 	Ar << Header.CPULODScreenSize;
 	uint32 LODInfosCount = Header.LODInfos.Num();
@@ -726,11 +722,6 @@ void FHairStrandsClusterCullingBulkData::GetResources(FHairStrandsBulkCommon::FQ
 		Out.Add(Data.CurveToClusterIds, TEXT("_CurveToClusterIds"), 0, 0); // Load all data
 	}
 
-	if (Header.VertexLODCount)
-	{
-		Out.Add(Data.ClusterVertexIds, TEXT("_ClusterVertexIds"), 0, 0); // Load all data
-	}
-
 	if (Header.ClusterCount)
 	{
 		Out.Add(Data.PackedClusterInfos, TEXT("_PackedClusterInfos"), 0, 0); // Load all data
diff --git a/Engine/Plugins/Runtime/HairStrands/Source/HairStrandsCore/Public/GroomResources.h b/Engine/Plugins/Runtime/HairStrands/Source/HairStrandsCore/Public/GroomResources.h
--- a/Engine/Plugins/Runtime/HairStrands/Source/HairStrandsCore/Public/GroomResources.h
+++ b/Engine/Plugins/Runtime/HairStrands/Source/HairStrandsCore/Public/GroomResources.h
@@ -331,7 +331,6 @@ struct FHairStrandsClusterCullingResource : public FHairCommonResource
 		Total += GetBufferTotalNumBytes(ClusterInfoBuffer);
 		Total += GetBufferTotalNumBytes(ClusterLODInfoBuffer);
 		Total += GetBufferTotalNumBytes(CurveToClusterIdBuffer);
-		Total += GetBufferTotalNumBytes(ClusterVertexIdBuffer);
 		Total += GetBufferTotalNumBytes(PointLODBuffer);
 		return Total;
 	}
@@ -341,10 +340,7 @@ struct FHairStrandsClusterCullingResource : public FHairCommonResource
 	FRDGExternalBuffer ClusterLODInfoBuffer;
 
 	/* CurveId => ClusterId to know which AABB to contribute to*/
-	FRDGExternalBuffer CurveToClusterIdBuffer;
-
-	/* Concatenated data for each cluster: list of VertexId pointed to by ClusterInfoBuffer */
-	FRDGExternalBuffer ClusterVertexIdBuffer;
+	FRDGExternalBuffer CurveToClusterIdBuffer; // remove?
 
 	/* Min. LOD a which a Point becomes visible */
 	FRDGExternalBuffer PointLODBuffer;
diff --git a/Engine/Plugins/Runtime/HairStrands/Source/HairStrandsCore/Public/HairStrandsDatas.h b/Engine/Plugins/Runtime/HairStrands/Source/HairStrandsCore/Public/HairStrandsDatas.h
--- a/Engine/Plugins/Runtime/HairStrands/Source/HairStrandsCore/Public/HairStrandsDatas.h
+++ b/Engine/Plugins/Runtime/HairStrands/Source/HairStrandsCore/Public/HairStrandsDatas.h
@@ -296,9 +296,6 @@ struct FHairClusterInfo
 	identical the GPU data layout (HairStrandsClusterCommon.ush). This uses by the GPU LOD selection. */
 struct FHairClusterLODInfo
 {
-	uint32 VertexOffset = 0;
-	uint32 VertexCount0 = 0;
-	uint32 VertexCount1 = 0;
 	FFloat16 RadiusScale0 = 0;
 	FFloat16 RadiusScale1 = 0;
 };
@@ -734,7 +731,6 @@ struct HAIRSTRANDSCORE_API FHairStrandsClusterCullingData
 	TArray<FHairClusterInfo>	ClusterInfos;
 	TArray<FHairClusterLODInfo> ClusterLODInfos;
 	TArray<uint32>				CurveToClusterIds;
-	TArray<uint32>				ClusterVertexIds;
 	TArray<FHairLODInfo>		LODInfos;
 	TArray<uint32>				PointLODs;
 
@@ -770,7 +766,6 @@ struct HAIRSTRANDSCORE_API FHairStrandsClusterCullingBulkData : FHairStrandsBulk
 		uint32 ClusterLODCount = 0;
 		uint32 PointCount = 0;
 		uint32 CurveCount = 0;
-		uint32 VertexLODCount = 0;
 	} Header;
 
 	struct FData
@@ -779,7 +774,6 @@ struct HAIRSTRANDSCORE_API FHairStrandsClusterCullingBulkData : FHairStrandsBulk
 		FHairBulkContainer	PackedClusterInfos;		// Size - ClusterCount
 		FHairBulkContainer	ClusterLODInfos;		// Size - ClusterLODCount
 		FHairBulkContainer	CurveToClusterIds;		// Size - CurveCount
-		FHairBulkContainer	ClusterVertexIds;		// Size - VertexLODCount // TODO REMOVE
 		FHairBulkContainer	PointLODs;				// Size - PointCount / HAIR_POINT_LOD_COUNT_PER_UINT
 	} Data;
 };
diff --git a/Engine/Shaders/Private/HairStrands/HairStrandsClusterCommon.ush b/Engine/Shaders/Private/HairStrands/HairStrandsClusterCommon.ush
--- a/Engine/Shaders/Private/HairStrands/HairStrandsClusterCommon.ush
+++ b/Engine/Shaders/Private/HairStrands/HairStrandsClusterCommon.ush
@@ -81,17 +81,11 @@ FHairClusterInfo UnpackHairClusterInfo(FPackedHairClusterInfo In)
 
 struct FPackedHairClusterLODInfo
 {
-	uint VertexOffset;
-	uint VertexCount0;
-	uint VertexCount1;
 	uint RadiusScale01;
 };
 
 struct FHairClusterLODInfo
 {
-	uint VertexOffset;
-	uint VertexCount0;
-	uint VertexCount1;
 	float RadiusScale0;
 	float RadiusScale1;
 };
@@ -99,9 +93,6 @@ struct FHairClusterLODInfo
 FHairClusterLODInfo UnpackHairClusterLODInfo(FPackedHairClusterLODInfo In)
 {
 	FHairClusterLODInfo Out;
-	Out.VertexOffset = In.VertexOffset;
-	Out.VertexCount0 = In.VertexCount0;
-	Out.VertexCount1 = In.VertexCount1;
 	Out.RadiusScale0 = f16tof32(0xFFFF & (In.RadiusScale01));
 	Out.RadiusScale1 = f16tof32(0xFFFF & (In.RadiusScale01 >> 16));
 	return Out;
@@ -111,8 +102,6 @@ FHairClusterLODInfo UnpackHairClusterLODInfo(FPackedHairClusterLODInfo In)
 
 struct FHairClusterLOD
 {
-	uint  VertexOffset;
-	uint  VertexCount;
 	float RadiusScale;
 	float LOD;
 };
@@ -192,8 +181,6 @@ FHairClusterLOD GetHairClusterLOD(
 	const float S = LOD - iLOD;
 
 	FHairClusterLOD Out;
-	Out.VertexOffset = ClusterLODInfo.VertexOffset;
-	Out.VertexCount = lerp(ClusterLODInfo.VertexCount0, ClusterLODInfo.VertexCount1, S);
 	Out.RadiusScale = lerp(ClusterLODInfo.RadiusScale0, ClusterLODInfo.RadiusScale1, S);
 	Out.LOD = LOD;
 
