Remove Cluster LOD data.

diff --git a/Engine/Plugins/Runtime/HairStrands/Source/HairStrandsCore/Private/GroomAsset.cpp b/Engine/Plugins/Runtime/HairStrands/Source/HairStrandsCore/Private/GroomAsset.cpp
--- a/Engine/Plugins/Runtime/HairStrands/Source/HairStrandsCore/Private/GroomAsset.cpp
+++ b/Engine/Plugins/Runtime/HairStrands/Source/HairStrandsCore/Private/GroomAsset.cpp
@@ -127,7 +127,6 @@ uint32 FHairGroupPlatformData::FStrands::GetDataSize() const
 	uint32 BTotal = 0;
 	CTotal += ClusterCullingBulkData.Header.LODVisibility.GetAllocatedSize();
 	CTotal += ClusterCullingBulkData.Header.CPULODScreenSize.GetAllocatedSize();
-	CTotal += ClusterCullingBulkData.Data.ClusterLODInfos.IsBulkDataLoaded()	? ClusterCullingBulkData.Data.ClusterLODInfos.GetBulkDataSize() : 0;
 	BTotal += ClusterCullingBulkData.Data.CurveToClusterIds.IsBulkDataLoaded()	? ClusterCullingBulkData.Data.CurveToClusterIds.GetBulkDataSize() : 0;
 	CTotal += ClusterCullingBulkData.Data.PackedClusterInfos.IsBulkDataLoaded()	? ClusterCullingBulkData.Data.PackedClusterInfos.GetBulkDataSize() : 0;
 
diff --git a/Engine/Plugins/Runtime/HairStrands/Source/HairStrandsCore/Private/GroomBuilder.cpp b/Engine/Plugins/Runtime/HairStrands/Source/HairStrandsCore/Private/GroomBuilder.cpp
--- a/Engine/Plugins/Runtime/HairStrands/Source/HairStrandsCore/Private/GroomBuilder.cpp
+++ b/Engine/Plugins/Runtime/HairStrands/Source/HairStrandsCore/Private/GroomBuilder.cpp
@@ -40,7 +40,7 @@ static FAutoConsoleVariableRef CVarHairGroupIndexBuilder_MaxVoxelResolution(TEXT
 
 FString FGroomBuilder::GetVersion()
 {
-	return TEXT("v8r33");
+	return TEXT("v8r37");
 }
 
 namespace FHairStrandsDecimation
@@ -2993,13 +2993,11 @@ static void BuildClusterData(
 	Out.PointLODs.Init(0, Out.PointCount);
 
 	// 4. Write out cluster information
-	Out.ClusterLODInfos.SetNum(LODCount * Out.ClusterCount);
 	TArray<uint8> VertexLODMasks;
 	VertexLODMasks.SetNum(InRenStrandsData.GetNumPoints());
 
 	// Local variable for being capture by the lambda
 	TArray<FHairClusterInfo>& LocalClusterInfos = Out.ClusterInfos;
-	TArray<FHairClusterLODInfo>& LocalClusterLODInfos = Out.ClusterLODInfos;
 	TArray<uint32>& LocalCurveToClusterIds = Out.CurveToClusterIds;
 	TArray<uint32>& LocalPointLODs = Out.PointLODs;
 #define USE_PARALLE_FOR 0
@@ -3013,7 +3011,6 @@ static void BuildClusterData(
 			&InRenStrandsData,
 			&ClusterGrid,
 			&LocalClusterInfos,
-			&LocalClusterLODInfos,
 			&LocalCurveToClusterIds,
 			&VertexLODMasks,
 		]
@@ -3136,12 +3133,11 @@ static void BuildClusterData(
 
 		FHairClusterInfo& ClusterInfo = LocalClusterInfos[ClusterIt];
 		ClusterInfo.LODCount = LODCount;
-		ClusterInfo.LODInfoOffset = LODCount * ClusterIt;
 		for (uint8 LODIt = 0; LODIt < LODCount; ++LODIt)
 		{
-			FHairClusterLODInfo& ClusterLODInfo = LocalClusterLODInfos[ClusterInfo.LODInfoOffset + LODIt];
-			ClusterLODInfo.RadiusScale0 = 0;
-			ClusterLODInfo.RadiusScale1 = 0;
+			ClusterInfo.ScreenSize[LODIt] = 1.f;
+			ClusterInfo.RadiusScale[LODIt]= 1.f;
+			ClusterInfo.bIsVisible[LODIt] = false;
 
 			const uint32 CurveCount = LODCurveCount[LODIt];
 			const uint32 NextCurveCount = LODIt < LODCount - 1 ? LODCurveCount[LODIt + 1] : CurveCount;
@@ -3203,17 +3199,7 @@ static void BuildClusterData(
 
 			ClusterInfo.ScreenSize[LODIt] = InSettings.LODs[LODIt].ScreenSize * ScreenSizeScale;
 			ClusterInfo.bIsVisible[LODIt] = InSettings.LODs[LODIt].bVisible;
-			FHairClusterLODInfo& ClusterLODInfo = LocalClusterLODInfos[ClusterInfo.LODInfoOffset + LODIt];
-			ClusterLODInfo.RadiusScale0 = LODScale;
-			ClusterLODInfo.RadiusScale1 = LODScale;
-		}
-
-		// Fill in transition radius between LOD to insure that the interpolation is continuous
-		for (uint8 LODIt = 0; LODIt < LODCount - 1; ++LODIt)
-		{
-			FHairClusterLODInfo& Curr = LocalClusterLODInfos[ClusterInfo.LODInfoOffset + LODIt];
-			FHairClusterLODInfo& Next = LocalClusterLODInfos[ClusterInfo.LODInfoOffset + LODIt + 1];
-			Curr.RadiusScale1 = Next.RadiusScale0;
+			ClusterInfo.RadiusScale[LODIt]= LODScale;
 		}
 	}
 #if USE_PARALLE_FOR
@@ -3237,18 +3223,15 @@ static void BuildClusterBulkData(
 	Out.Header.ClusterCount		= In.ClusterCount;
 	Out.Header.PointCount		= In.PointCount;
 	Out.Header.CurveCount		= In.CurveCount;
-	Out.Header.ClusterLODCount 	= In.ClusterLODInfos.Num();
 
 	Out.Header.CPULODScreenSize	= In.CPULODScreenSize;
 	Out.Header.LODVisibility	= In.LODVisibility;
 	Out.Header.LODInfos			= In.LODInfos;
 
 	// Sanity check
-	check(Out.Header.ClusterCount		== uint32(In.ClusterInfos.Num()));
-	check(Out.Header.ClusterLODCount	== uint32(In.ClusterLODInfos.Num()));
-	check(Out.Header.CurveCount			== uint32(In.CurveToClusterIds.Num()));
+	check(Out.Header.ClusterCount	== uint32(In.ClusterInfos.Num()));
+	check(Out.Header.CurveCount		== uint32(In.CurveToClusterIds.Num()));
 	
-	HairStrandsBuilder::CopyToBulkData<FHairClusterLODInfoFormat>(Out.Data.ClusterLODInfos, In.ClusterLODInfos);
 	HairStrandsBuilder::CopyToBulkData<FHairClusterIndexFormat>(Out.Data.CurveToClusterIds, In.CurveToClusterIds);
 
 	// Pack point LOD data
@@ -3275,36 +3258,76 @@ static void BuildClusterBulkData(
 		check(uint32(In.ClusterInfos.Num()) == Out.Header.ClusterCount);
 		check(uint32(In.CurveToClusterIds.Num()) == Out.Header.CurveCount);
 	
+		float MinScreenSize = FLT_MAX;
+		float MaxScreenSize = 0.f; 
+
+		float MinRadiusScale = FLT_MAX; 
+		float MaxRadiusScale = 0.f; 
+		for (const FHairClusterInfo& Info : In.ClusterInfos)
+		{
+			for (uint32 LODIt = 0; LODIt < Info.LODCount; ++LODIt)
+			{
+				MinScreenSize = FMath::Min(MinScreenSize, Info.ScreenSize[LODIt]);
+				MaxScreenSize = FMath::Max(MaxScreenSize, Info.ScreenSize[LODIt]);
+
+				MinRadiusScale = FMath::Min(MinRadiusScale, Info.RadiusScale[LODIt]);
+				MaxRadiusScale = FMath::Max(MaxRadiusScale, Info.RadiusScale[LODIt]);
+			}
+		}
+		check(MinScreenSize <= MaxScreenSize);
+		check(MinRadiusScale <= MaxRadiusScale);
+
+		const float ScreenOffset = MinScreenSize;
+		const float ScreenScale = FMath::Abs(MaxScreenSize - MinScreenSize) > 1e-4f ? (MaxScreenSize - MinScreenSize) : 1.f;
+		const float InvScreenScale = 1.f / ScreenScale;
+
+		const float RadiusOffset = MinRadiusScale;
+		const float RadiusScale = FMath::Abs(MaxRadiusScale - MinRadiusScale) > 1e-4f ? (MaxRadiusScale - MinRadiusScale) : 1.f;
+		const float InvRadiusScale = 1.f / RadiusScale;
+
+		Out.Header.ClusterInfoParameters = FVector4f(ScreenScale, ScreenOffset, RadiusScale, RadiusOffset);
+
+		auto Quantize8Bits = [](float In, float InScale, float InOffset)
+		{
+			const float N = (In - InOffset) * InScale;
+			return FMath::Clamp(uint32(N * 0xFFu), 0u, 0xFFu);
+		};
+
 		TArray<FHairClusterInfo::Packed> PackedClusterInfos;
 		PackedClusterInfos.Reserve(In.ClusterInfos.Num());
 		for (const FHairClusterInfo& Info : In.ClusterInfos)
 		{
-			FHairClusterInfo::Packed& PackedInfo = PackedClusterInfos.AddDefaulted_GetRef();
-			// Sanity check
-			check(Info.LODCount <= 8);
+			FHairClusterInfo::Packed& Packed = PackedClusterInfos.AddDefaulted_GetRef();
 
-			PackedInfo.LODCount = FMath::Clamp(Info.LODCount, 0u, 0xFFu);
-			PackedInfo.LODInfoOffset = FMath::Clamp(Info.LODInfoOffset, 0u, (1u << 24u) - 1u);
-			PackedInfo.LOD_ScreenSize_0 = to10Bits(Info.ScreenSize[0]);
-			PackedInfo.LOD_ScreenSize_1 = to10Bits(Info.ScreenSize[1]);
-			PackedInfo.LOD_ScreenSize_2 = to10Bits(Info.ScreenSize[2]);
-			PackedInfo.LOD_ScreenSize_3 = to10Bits(Info.ScreenSize[3]);
-			PackedInfo.LOD_ScreenSize_4 = to10Bits(Info.ScreenSize[4]);
-			PackedInfo.LOD_ScreenSize_5 = to10Bits(Info.ScreenSize[5]);
-			PackedInfo.LOD_ScreenSize_6 = to10Bits(Info.ScreenSize[6]);
-			PackedInfo.LOD_ScreenSize_7 = to10Bits(Info.ScreenSize[7]);
-			PackedInfo.LOD_bIsVisible = 0;
+			uint32 bIsVisible = 0;
 			for (uint32 LODIt = 0; LODIt < FHairClusterInfo::MaxLOD; ++LODIt)
 			{
 				if (Info.bIsVisible[LODIt])
 				{
-					PackedInfo.LOD_bIsVisible = PackedInfo.LOD_bIsVisible | (1 << LODIt);
+					bIsVisible = bIsVisible | (1 << LODIt);
 				}
 			}
 
-			PackedInfo.Pad0 = 0;
-			PackedInfo.Pad1 = 0;
-			PackedInfo.Pad2 = 0;
+			// Sanity check
+			check(Info.LODCount <= 8);
+
+			Packed.Screen0 = Info.LODCount;
+			Packed.Screen0 = Packed.Screen0 | (Quantize8Bits(Info.ScreenSize[1], InvScreenScale, ScreenOffset) << 8);
+			Packed.Screen0 = Packed.Screen0 | (Quantize8Bits(Info.ScreenSize[2], InvScreenScale, ScreenOffset) << 16);
+			Packed.Screen0 = Packed.Screen0 | (Quantize8Bits(Info.ScreenSize[3], InvScreenScale, ScreenOffset) << 24);
+			Packed.Screen1 = Packed.Screen1 | (Quantize8Bits(Info.ScreenSize[4], InvScreenScale, ScreenOffset) << 0);
+			Packed.Screen1 = Packed.Screen1 | (Quantize8Bits(Info.ScreenSize[5], InvScreenScale, ScreenOffset) << 8);
+			Packed.Screen1 = Packed.Screen1 | (Quantize8Bits(Info.ScreenSize[6], InvScreenScale, ScreenOffset) << 16);
+			Packed.Screen1 = Packed.Screen1 | (Quantize8Bits(Info.ScreenSize[7], InvScreenScale, ScreenOffset) << 24);
+
+			Packed.Radius0 = bIsVisible & 0xFF;
+			Packed.Radius0 = Packed.Radius0 | (Quantize8Bits(Info.RadiusScale[1], InvRadiusScale, RadiusOffset) << 8);
+			Packed.Radius0 = Packed.Radius0 | (Quantize8Bits(Info.RadiusScale[2], InvRadiusScale, RadiusOffset) << 16);
+			Packed.Radius0 = Packed.Radius0 | (Quantize8Bits(Info.RadiusScale[3], InvRadiusScale, RadiusOffset) << 24);
+			Packed.Radius1 = Packed.Radius1 | (Quantize8Bits(Info.RadiusScale[4], InvRadiusScale, RadiusOffset) << 0);
+			Packed.Radius1 = Packed.Radius1 | (Quantize8Bits(Info.RadiusScale[5], InvRadiusScale, RadiusOffset) << 8);
+			Packed.Radius1 = Packed.Radius1 | (Quantize8Bits(Info.RadiusScale[6], InvRadiusScale, RadiusOffset) << 16);
+			Packed.Radius1 = Packed.Radius1 | (Quantize8Bits(Info.RadiusScale[7], InvRadiusScale, RadiusOffset) << 24);
 
 			static_assert(sizeof(FHairClusterInfo::Packed) == sizeof(FHairClusterInfo::BulkType));
 		}
diff --git a/Engine/Plugins/Runtime/HairStrands/Source/HairStrandsCore/Private/GroomResources.cpp b/Engine/Plugins/Runtime/HairStrands/Source/HairStrandsCore/Private/GroomResources.cpp
--- a/Engine/Plugins/Runtime/HairStrands/Source/HairStrandsCore/Private/GroomResources.cpp
+++ b/Engine/Plugins/Runtime/HairStrands/Source/HairStrandsCore/Private/GroomResources.cpp
@@ -1190,24 +1190,16 @@ void FHairStrandsClusterCullingResource::InternalAllocate(FRDGBuilder& GraphBuil
 {
 	check(StreamingRequest.IsCompleted());
 
-	if (ValidateHairBulkData())
-	{
-		BulkData.Validate(false);
-	}
+	const uint32 EntryCount = FMath::DivideAndRoundUp(BulkData.Header.PointCount, HAIR_POINT_LOD_COUNT_PER_UINT);
 
 	InternalCreateStructuredBufferRDG_FromHairBulkData<FHairClusterInfoFormat>(GraphBuilder, BulkData.Data.PackedClusterInfos, BulkData.Header.ClusterCount, ClusterInfoBuffer, ToHairResourceDebugName(TEXT("Hair.StrandsClusterCulling_ClusterInfoBuffer"), ResourceName), OwnerName, EHairResourceUsageType::Static);
-	InternalCreateStructuredBufferRDG_FromHairBulkData<FHairClusterLODInfoFormat>(GraphBuilder, BulkData.Data.ClusterLODInfos, BulkData.Header.ClusterLODCount, ClusterLODInfoBuffer, ToHairResourceDebugName(TEXT("Hair.StrandsClusterCulling_ClusterLODInfoBuffer"), ResourceName), OwnerName, EHairResourceUsageType::Static);
-
-	InternalCreateVertexBufferRDG_FromHairBulkData<FHairClusterIndexFormat>(GraphBuilder, BulkData.Data.CurveToClusterIds, BulkData.Header.CurveCount, CurveToClusterIdBuffer, ToHairResourceDebugName(TEXT("Hair.StrandsClusterCulling_CurveToClusterIds"), ResourceName), OwnerName, EHairResourceUsageType::Static);
-	
-	const uint32 EntryCount = FMath::DivideAndRoundUp(BulkData.Header.PointCount, HAIR_POINT_LOD_COUNT_PER_UINT);
+	InternalCreateVertexBufferRDG_FromHairBulkData<FHairClusterIndexFormat>(GraphBuilder, BulkData.Data.CurveToClusterIds, BulkData.Header.CurveCount, CurveToClusterIdBuffer, ToHairResourceDebugName(TEXT("Hair.StrandsClusterCulling_CurveToClusterIds"), ResourceName), OwnerName, EHairResourceUsageType::Static);	
 	InternalCreateVertexBufferRDG_FromHairBulkData<FHairClusterIndexFormat>(GraphBuilder, BulkData.Data.PointLODs, EntryCount, PointLODBuffer, ToHairResourceDebugName(TEXT("Hair.StrandsClusterCulling_PointLOD"), ResourceName), OwnerName, EHairResourceUsageType::Static);
 }
 
 void FHairStrandsClusterCullingResource::InternalRelease()
 {
 	ClusterInfoBuffer.Release();
-	ClusterLODInfoBuffer.Release();
 	CurveToClusterIdBuffer.Release();
 	PointLODBuffer.Release();
 	MaxAvailableCurveCount = 0;
diff --git a/Engine/Plugins/Runtime/HairStrands/Source/HairStrandsCore/Private/HairStrandsClusterCulling.cpp b/Engine/Plugins/Runtime/HairStrands/Source/HairStrandsCore/Private/HairStrandsClusterCulling.cpp
--- a/Engine/Plugins/Runtime/HairStrands/Source/HairStrandsCore/Private/HairStrandsClusterCulling.cpp
+++ b/Engine/Plugins/Runtime/HairStrands/Source/HairStrandsCore/Private/HairStrandsClusterCulling.cpp
@@ -53,12 +53,12 @@ class FHairClusterCullCS: public FGlobalShader
 		SHADER_PARAMETER(float, LODBias)
 		SHADER_PARAMETER(uint32, CurveCount)
 		SHADER_PARAMETER(uint32, DispatchCountX)
+		SHADER_PARAMETER(FVector4f, ClusterInfoParameters)
 
 		SHADER_PARAMETER_RDG_BUFFER_SRV(Buffer<uint>, CurveBuffer)
 		SHADER_PARAMETER_RDG_BUFFER_SRV(Buffer<uint>, CurveToClusterIdBuffer)
 		SHADER_PARAMETER_RDG_BUFFER_SRV(Buffer<uint>, PointLODBuffer)
 		SHADER_PARAMETER_RDG_BUFFER_SRV(StructuredBuffer<FPackedHairClusterInfo>, ClusterInfoBuffer)
-		SHADER_PARAMETER_RDG_BUFFER_SRV(StructuredBuffer<FPackedHairClusterLODInfo>, ClusterLODInfoBuffer)
 	
 		SHADER_PARAMETER_RDG_BUFFER_UAV(RWBuffer<uint>, OutPointCounter)
 		SHADER_PARAMETER_RDG_BUFFER_UAV(RWBuffer<uint>, OutIndexBuffer)
@@ -133,7 +133,7 @@ static void AddClusterCullingPass(
 		Parameters->CurveToClusterIdBuffer=RegisterAsSRV(GraphBuilder, *ClusterData.CurveToClusterIdBuffer);
 		Parameters->PointLODBuffer 		 = RegisterAsSRV(GraphBuilder, *ClusterData.PointLODBuffer);
 		Parameters->ClusterInfoBuffer 	 = RegisterAsSRV(GraphBuilder, *ClusterData.ClusterInfoBuffer);
-		Parameters->ClusterLODInfoBuffer = RegisterAsSRV(GraphBuilder, *ClusterData.ClusterLODInfoBuffer);
+		Parameters->ClusterInfoParameters= ClusterData.ClusterInfoParameters;
 
 		Parameters->OutPointCounter 	 = PointCounterUAV;
 		Parameters->OutIndexBuffer 		 = RegisterAsUAV(GraphBuilder, *ClusterData.GetCulledVertexIdBuffer());
@@ -276,8 +276,8 @@ void AddInstanceToClusterData(FHairGroupInstance* In, FHairStrandClusterData& Ou
 	HairGroupCluster.CurveBuffer = &In->Strands.RestResource->CurveBuffer;
 	HairGroupCluster.PointLODBuffer = &In->Strands.ClusterCullingResource->PointLODBuffer;
 
+	HairGroupCluster.ClusterInfoParameters = In->Strands.ClusterCullingResource->BulkData.Header.ClusterInfoParameters;
 	HairGroupCluster.ClusterInfoBuffer = &In->Strands.ClusterCullingResource->ClusterInfoBuffer;
-	HairGroupCluster.ClusterLODInfoBuffer = &In->Strands.ClusterCullingResource->ClusterLODInfoBuffer;
 	HairGroupCluster.CurveToClusterIdBuffer = &In->Strands.ClusterCullingResource->CurveToClusterIdBuffer;
 
 	HairGroupCluster.HairGroupPublicPtr = In->HairGroupPublicData;
diff --git a/Engine/Plugins/Runtime/HairStrands/Source/HairStrandsCore/Private/HairStrandsClusterCulling.h b/Engine/Plugins/Runtime/HairStrands/Source/HairStrandsCore/Private/HairStrandsClusterCulling.h
--- a/Engine/Plugins/Runtime/HairStrands/Source/HairStrandsCore/Private/HairStrandsClusterCulling.h
+++ b/Engine/Plugins/Runtime/HairStrands/Source/HairStrandsCore/Private/HairStrandsClusterCulling.h
@@ -18,7 +18,7 @@ struct FHairStrandClusterData
 		uint32 ClusterCount = 0;
 		uint32 VertexCount = 0;
 		uint32 MaxPointPerCurve = 0;
-
+		FVector4f ClusterInfoParameters = FVector4f::Zero();
 		float LODIndex = -1;
 		float LODBias = 0.0f;
 		bool bVisible = false;
@@ -30,7 +30,6 @@ struct FHairStrandClusterData
 		FRDGExternalBuffer* GroupAABBBuffer = nullptr;
 		FRDGExternalBuffer* ClusterAABBBuffer = nullptr;
 		FRDGExternalBuffer* ClusterInfoBuffer = nullptr; // SRV
-		FRDGExternalBuffer* ClusterLODInfoBuffer = nullptr; // SRV
 		FRDGExternalBuffer* CurveToClusterIdBuffer = nullptr; // SRV
 
 		TRefCountPtr<FRDGPooledBuffer> ClusterIdBuffer;
diff --git a/Engine/Plugins/Runtime/HairStrands/Source/HairStrandsCore/Private/HairStrandsDatas.cpp b/Engine/Plugins/Runtime/HairStrands/Source/HairStrandsCore/Private/HairStrandsDatas.cpp
--- a/Engine/Plugins/Runtime/HairStrands/Source/HairStrandsCore/Private/HairStrandsDatas.cpp
+++ b/Engine/Plugins/Runtime/HairStrands/Source/HairStrandsCore/Private/HairStrandsDatas.cpp
@@ -654,21 +654,18 @@ void FHairStrandsInterpolationBulkData::ResetLoadedSize()
 void FHairStrandsClusterCullingBulkData::Reset()
 {
 	Header.ClusterCount = 0;
-	Header.ClusterLODCount = 0;
 	Header.PointCount = 0;
 	Header.CurveCount = 0;
-
+	Header.ClusterInfoParameters = FVector4f::Zero();
 	Header.LODVisibility.Empty();
 	Header.CPULODScreenSize.Empty();
 	Header.LODInfos.Empty();
 
-	Data.ClusterLODInfos.RemoveBulkData();
 	Data.CurveToClusterIds.RemoveBulkData();
 	Data.PackedClusterInfos.RemoveBulkData();
 	Data.PointLODs.RemoveBulkData();
 
 	// Reset the bulk byte buffer to ensure the (serialize) data size is reset to 0
-	Data.ClusterLODInfos 	= FHairBulkContainer();
 	Data.CurveToClusterIds  = FHairBulkContainer();
 	Data.PackedClusterInfos = FHairBulkContainer();
 	Data.PointLODs			= FHairBulkContainer();
@@ -676,7 +673,6 @@ void FHairStrandsClusterCullingBulkData::Reset()
 
 void FHairStrandsClusterCullingBulkData::ResetLoadedSize()
 {
-	Data.ClusterLODInfos.LoadedSize 	= 0;
 	Data.CurveToClusterIds.LoadedSize 	= 0;
 	Data.PackedClusterInfos.LoadedSize	= 0;
 	Data.PointLODs.LoadedSize			= 0;
@@ -685,11 +681,11 @@ void FHairStrandsClusterCullingBulkData::ResetLoadedSize()
 void FHairStrandsClusterCullingBulkData::SerializeHeader(FArchive& Ar, UObject* Owner)
 {
 	Ar << Header.ClusterCount;
-	Ar << Header.ClusterLODCount;
 	Ar << Header.PointCount;
 	Ar << Header.CurveCount;
 	Ar << Header.LODVisibility;
 	Ar << Header.CPULODScreenSize;
+	Ar << Header.ClusterInfoParameters;
 	uint32 LODInfosCount = Header.LODInfos.Num();
 	Ar << LODInfosCount;
 	if (Ar.IsLoading())
@@ -712,11 +708,6 @@ bool ValidateHairBulkData();
 
 void FHairStrandsClusterCullingBulkData::GetResources(FHairStrandsBulkCommon::FQuery & Out)
 {
-	if (Header.ClusterLODCount)
-	{
-		Out.Add(Data.ClusterLODInfos, TEXT("_ClusterLODInfos"), 0, 0); // Load all data
-	}
-
 	if (Header.CurveCount)
 	{
 		Out.Add(Data.CurveToClusterIds, TEXT("_CurveToClusterIds"), 0, 0); // Load all data
@@ -730,42 +721,8 @@ void FHairStrandsClusterCullingBulkData::GetResources(FHairStrandsBulkCommon::FQ
 	{
 		Out.Add(Data.PointLODs, TEXT("_PointLODs"), 0, 0); // Load all data. TODO load data based on active curve/point count
 	}
-
-	if (ValidateHairBulkData() && (Out.Type == FHairStrandsBulkCommon::FQuery::WriteDDC || Out.Type == FHairStrandsBulkCommon::FQuery::ReadWriteIO))
-	{
-		Validate(true);
-	}
-}
-
-void FHairStrandsClusterCullingBulkData::Validate(bool bIsSaving)
-{
-	return;
-
-	if (Header.ClusterCount == 0)
-	{
-		return;
-	}
-
-	const FHairClusterInfo::Packed* Datas = (const FHairClusterInfo::Packed*)Data.PackedClusterInfos.Data.Lock(LOCK_READ_ONLY);
-	
-	// Simple heuristic to check if the data are valid
-	const uint32 MaxCount = FMath::Min(Header.ClusterCount, 128u);
-	bool bIsValid = true;
-	for (uint32 It = 0; It < MaxCount; ++It)
-	{
-		bIsValid = bIsValid && Datas[It].LODCount <= 8;
-		if (!bIsValid) break;
-	}
-	if (!bIsValid)
-	{
-		FString DebugName = Data.ClusterLODInfos.GetDebugName();
-		UE_LOG(LogHairStrands, Error, TEXT("[Groom/DDC] Strands - Invalid ClusterCullingBulkData when %s bulk data - %s"), bIsSaving ? TEXT("Saving") : TEXT("Loading"), *DebugName);
-	}
-
-	Data.PackedClusterInfos.Data.Unlock();
 }
 
-
 /////////////////////////////////////////////////////////////////////////////////////////
 // Root data
 
diff --git a/Engine/Plugins/Runtime/HairStrands/Source/HairStrandsCore/Public/GroomAsset.h b/Engine/Plugins/Runtime/HairStrands/Source/HairStrandsCore/Public/GroomAsset.h
--- a/Engine/Plugins/Runtime/HairStrands/Source/HairStrandsCore/Public/GroomAsset.h
+++ b/Engine/Plugins/Runtime/HairStrands/Source/HairStrandsCore/Public/GroomAsset.h
@@ -179,20 +179,7 @@ struct HAIRSTRANDSCORE_API FHairGroupPlatformData
 		}
 
 		uint32 GetDataSize() const;
-#if 0
-		{
-			uint32 Total = 0;
-			Total += FBaseWithInterpolation::GetDataSize();
-
-			Total += ClusterCullingBulkData.Header.LODVisibility.GetAllocatedSize();
-			Total += ClusterCullingBulkData.Header.CPULODScreenSize.GetAllocatedSize();
-			Total += ClusterCullingBulkData.Data.ClusterInfos.IsBulkDataLoaded()		? ClusterCullingBulkData.Data.ClusterInfos.GetBulkDataSize() : 0;
-			Total += ClusterCullingBulkData.Data.ClusterLODInfos.IsBulkDataLoaded()		? ClusterCullingBulkData.Data.ClusterLODInfos.GetBulkDataSize() : 0;
-			Total += ClusterCullingBulkData.Data.CurveToClusterIds.IsBulkDataLoaded()	? ClusterCullingBulkData.Data.CurveToClusterIds.GetBulkDataSize() : 0;
-			Total += ClusterCullingBulkData.Data.ClusterVertexIds.IsBulkDataLoaded()	? ClusterCullingBulkData.Data.ClusterVertexIds.GetBulkDataSize() : 0;
-			return Total;
-		}
-#endif
+
 		FHairStrandsClusterCullingBulkData	ClusterCullingBulkData;
 		FHairStrandsClusterCullingResource* ClusterCullingResource = nullptr;
 
diff --git a/Engine/Plugins/Runtime/HairStrands/Source/HairStrandsCore/Public/GroomResources.h b/Engine/Plugins/Runtime/HairStrands/Source/HairStrandsCore/Public/GroomResources.h
--- a/Engine/Plugins/Runtime/HairStrands/Source/HairStrandsCore/Public/GroomResources.h
+++ b/Engine/Plugins/Runtime/HairStrands/Source/HairStrandsCore/Public/GroomResources.h
@@ -329,7 +329,6 @@ struct FHairStrandsClusterCullingResource : public FHairCommonResource
 	{
 		uint32 Total = 0;
 		Total += GetBufferTotalNumBytes(ClusterInfoBuffer);
-		Total += GetBufferTotalNumBytes(ClusterLODInfoBuffer);
 		Total += GetBufferTotalNumBytes(CurveToClusterIdBuffer);
 		Total += GetBufferTotalNumBytes(PointLODBuffer);
 		return Total;
@@ -337,10 +336,9 @@ struct FHairStrandsClusterCullingResource : public FHairCommonResource
 
 	/* Cluster info buffer */
 	FRDGExternalBuffer ClusterInfoBuffer;
-	FRDGExternalBuffer ClusterLODInfoBuffer;
 
 	/* CurveId => ClusterId to know which AABB to contribute to*/
-	FRDGExternalBuffer CurveToClusterIdBuffer; // remove?
+	FRDGExternalBuffer CurveToClusterIdBuffer;
 
 	/* Min. LOD a which a Point becomes visible */
 	FRDGExternalBuffer PointLODBuffer;
diff --git a/Engine/Plugins/Runtime/HairStrands/Source/HairStrandsCore/Public/HairStrandsDatas.h b/Engine/Plugins/Runtime/HairStrands/Source/HairStrandsCore/Public/HairStrandsDatas.h
--- a/Engine/Plugins/Runtime/HairStrands/Source/HairStrandsCore/Public/HairStrandsDatas.h
+++ b/Engine/Plugins/Runtime/HairStrands/Source/HairStrandsCore/Public/HairStrandsDatas.h
@@ -257,23 +257,10 @@ struct FHairClusterInfo
 
 	struct Packed
 	{
-		uint32 LODInfoOffset : 24;
-		uint32 LODCount : 8;
-
-		uint32 LOD_ScreenSize_0 : 10;
-		uint32 LOD_ScreenSize_1 : 10;
-		uint32 LOD_ScreenSize_2 : 10;
-		uint32 Pad0 : 2;
-
-		uint32 LOD_ScreenSize_3 : 10;
-		uint32 LOD_ScreenSize_4 : 10;
-		uint32 LOD_ScreenSize_5 : 10;
-		uint32 Pad1 : 2;
-
-		uint32 LOD_ScreenSize_6 : 10;
-		uint32 LOD_ScreenSize_7 : 10;
-		uint32 LOD_bIsVisible : 8;
-		uint32 Pad2 : 4;
+		uint32 Screen0;
+		uint32 Screen1;
+		uint32 Radius0;
+		uint32 Radius1;
 	};
 	typedef FUintVector4 BulkType;
 
@@ -282,24 +269,17 @@ struct FHairClusterInfo
 		for (uint32 LODIt = 0; LODIt < MaxLOD; ++LODIt)
 		{
 			ScreenSize[LODIt] = 0;
+			RadiusScale[LODIt] = 1.f;
 			bIsVisible[LODIt] = true;
 		}
 	}
 
 	uint32 LODCount = 0;
-	uint32 LODInfoOffset = 0;
 	TStaticArray<float,MaxLOD> ScreenSize;
+	TStaticArray<float,MaxLOD> RadiusScale;
 	TStaticArray<bool, MaxLOD> bIsVisible;
 };
 
-/*  Structure describing the LOD settings common to all clusters. The layout of this structure is
-	identical the GPU data layout (HairStrandsClusterCommon.ush). This uses by the GPU LOD selection. */
-struct FHairClusterLODInfo
-{
-	FFloat16 RadiusScale0 = 0;
-	FFloat16 RadiusScale1 = 0;
-};
-
 struct FHairLODInfo
 {
 	uint32 CurveCount = 0;
@@ -313,13 +293,6 @@ struct FHairClusterInfoFormat
 	static const uint32 SizeInByte = sizeof(Type);
 };
 
-struct FHairClusterLODInfoFormat
-{
-	typedef FHairClusterLODInfo Type;
-	typedef FHairClusterLODInfo BulkType;
-	static const uint32 SizeInByte = sizeof(Type);
-};
-
 struct FHairClusterIndexFormat
 {
 	typedef uint32 Type;
@@ -729,7 +702,6 @@ struct HAIRSTRANDSCORE_API FHairStrandsClusterCullingData
 
 	/* LOD info for the various clusters for LOD management on GPU */
 	TArray<FHairClusterInfo>	ClusterInfos;
-	TArray<FHairClusterLODInfo> ClusterLODInfos;
 	TArray<uint32>				CurveToClusterIds;
 	TArray<FHairLODInfo>		LODInfos;
 	TArray<uint32>				PointLODs;
@@ -749,7 +721,6 @@ struct HAIRSTRANDSCORE_API FHairStrandsClusterCullingBulkData : FHairStrandsBulk
 	virtual void GetResources(FQuery& Out) override;
 
 	bool IsValid() const { return Header.ClusterCount > 0 && Header.PointCount > 0; }
-	void Validate(bool bIsSaving);
 
 	struct FHeader
 	{
@@ -763,16 +734,15 @@ struct HAIRSTRANDSCORE_API FHairStrandsClusterCullingBulkData : FHairStrandsBulk
 		TArray<FHairLODInfo> LODInfos;
 	
 		uint32 ClusterCount = 0;
-		uint32 ClusterLODCount = 0;
 		uint32 PointCount = 0;
 		uint32 CurveCount = 0;
+		FVector4f ClusterInfoParameters = FVector4f::Zero(); // xy:Scale/Offset for ScreenSize, zw:Scale/Offset for Radius
 	} Header;
 
 	struct FData
 	{
 		/* LOD info for the various clusters for LOD management on GPU */
 		FHairBulkContainer	PackedClusterInfos;		// Size - ClusterCount
-		FHairBulkContainer	ClusterLODInfos;		// Size - ClusterLODCount
 		FHairBulkContainer	CurveToClusterIds;		// Size - CurveCount
 		FHairBulkContainer	PointLODs;				// Size - PointCount / HAIR_POINT_LOD_COUNT_PER_UINT
 	} Data;
diff --git a/Engine/Shaders/Private/HairStrands/HairStrandsClusterCommon.ush b/Engine/Shaders/Private/HairStrands/HairStrandsClusterCommon.ush
--- a/Engine/Shaders/Private/HairStrands/HairStrandsClusterCommon.ush
+++ b/Engine/Shaders/Private/HairStrands/HairStrandsClusterCommon.ush
@@ -34,7 +34,9 @@ uint GetHairStrandsVertexFetchIndex_Culled(uint2 InDispatchThreadId, uint InGrou
 struct FHairClusterInfo
 {
 	uint  LODCount;
-	uint  LODInfoOffset;
+	uint  LOD_bIsVisible;
+
+	// ScreenSize
 	float LOD0_ScreenSize;
 	float LOD1_ScreenSize;
 	float LOD2_ScreenSize;
@@ -43,58 +45,57 @@ struct FHairClusterInfo
 	float LOD5_ScreenSize;
 	float LOD6_ScreenSize;
 	float LOD7_ScreenSize;
-	uint  LOD_bIsVisible;
+
+	// Scale
+	float LOD0_RadiusScale;
+	float LOD1_RadiusScale;
+	float LOD2_RadiusScale;
+	float LOD3_RadiusScale;
+	float LOD4_RadiusScale;
+	float LOD5_RadiusScale;
+	float LOD6_RadiusScale;
+	float LOD7_RadiusScale;
 };
 
 struct FPackedHairClusterInfo
 {
-	uint LODCount_LODInfoOffset;
-	uint LOD_ScreenSize_0_1_2;
-	uint LOD_ScreenSize_3_4_5;
-	uint LOD_ScreenSize_6_7_IsVisible_Pad;
+	uint ScreenSizeX123;
+	uint ScreenSize4567;
+	uint RadiusScaleX123;
+	uint RadiusScale4567;
 };
 
-FHairClusterInfo UnpackHairClusterInfo(FPackedHairClusterInfo In)
+FHairClusterInfo UnpackHairClusterInfo(FPackedHairClusterInfo In, float4 InParameters)
 {
-	FHairClusterInfo Out;
-	Out.LODCount		= (In.LODCount_LODInfoOffset >> 24) & 0xFF;
-	Out.LODInfoOffset   = (In.LODCount_LODInfoOffset & 0x00FFFFFF);
-
-	Out.LOD0_ScreenSize = (0x3FF & (In.LOD_ScreenSize_0_1_2)) / 1023.f;
-	Out.LOD1_ScreenSize = (0x3FF & (In.LOD_ScreenSize_0_1_2 >> 10)) / 1023.f;
-	Out.LOD2_ScreenSize = (0x3FF & (In.LOD_ScreenSize_0_1_2 >> 20)) / 1023.f;
+	const float4 Screen0 = UnpackRGBA8(In.ScreenSizeX123);
+	const float4 Screen1 = UnpackRGBA8(In.ScreenSize4567);
 
-	Out.LOD3_ScreenSize = (0x3FF & (In.LOD_ScreenSize_3_4_5)) / 1023.f;
-	Out.LOD4_ScreenSize = (0x3FF & (In.LOD_ScreenSize_3_4_5 >> 10)) / 1023.f;
-	Out.LOD5_ScreenSize = (0x3FF & (In.LOD_ScreenSize_3_4_5 >> 20)) / 1023.f;
+	const float4 Radius0 = UnpackRGBA8(In.RadiusScaleX123);
+	const float4 Radius1 = UnpackRGBA8(In.RadiusScale4567);
 
-	Out.LOD6_ScreenSize = (0x3FF & (In.LOD_ScreenSize_6_7_IsVisible_Pad)) / 1023.f;
-	Out.LOD7_ScreenSize = (0x3FF & (In.LOD_ScreenSize_6_7_IsVisible_Pad >> 10)) / 1023.f;
 
-	Out.LOD_bIsVisible  = (0xFF  & (In.LOD_ScreenSize_6_7_IsVisible_Pad >> 20));
-	
-	return Out;
-}
-
-////////////////////////////////////////////////////////////////////////////////
-// Store info about a specific cluster LOD
-
-struct FPackedHairClusterLODInfo
-{
-	uint RadiusScale01;
-};
-
-struct FHairClusterLODInfo
-{
-	float RadiusScale0;
-	float RadiusScale1;
-};
+	FHairClusterInfo Out;
+	Out.LOD0_ScreenSize = 1.f;
+	Out.LOD1_ScreenSize = Screen0.y * InParameters.x + InParameters.y;
+	Out.LOD2_ScreenSize = Screen0.z * InParameters.x + InParameters.y;
+	Out.LOD3_ScreenSize = Screen0.w * InParameters.x + InParameters.y;
+	Out.LOD4_ScreenSize = Screen1.x * InParameters.x + InParameters.y;
+	Out.LOD5_ScreenSize = Screen1.y * InParameters.x + InParameters.y;
+	Out.LOD6_ScreenSize = Screen1.z * InParameters.x + InParameters.y;
+	Out.LOD7_ScreenSize = Screen1.w * InParameters.x + InParameters.y;
+
+	Out.LOD0_RadiusScale = 1.f;
+	Out.LOD1_RadiusScale = Radius0.y * InParameters.z + InParameters.w;
+	Out.LOD2_RadiusScale = Radius0.z * InParameters.z + InParameters.w;
+	Out.LOD3_RadiusScale = Radius0.w * InParameters.z + InParameters.w;
+	Out.LOD4_RadiusScale = Radius1.x * InParameters.z + InParameters.w;
+	Out.LOD5_RadiusScale = Radius1.y * InParameters.z + InParameters.w;
+	Out.LOD6_RadiusScale = Radius1.z * InParameters.z + InParameters.w;
+	Out.LOD7_RadiusScale = Radius1.w * InParameters.z + InParameters.w;
+
+	Out.LODCount		= In.ScreenSizeX123 & 0xFFu;
+	Out.LOD_bIsVisible  = In.RadiusScaleX123 & 0xFFu;
 
-FHairClusterLODInfo UnpackHairClusterLODInfo(FPackedHairClusterLODInfo In)
-{
-	FHairClusterLODInfo Out;
-	Out.RadiusScale0 = f16tof32(0xFFFF & (In.RadiusScale01));
-	Out.RadiusScale1 = f16tof32(0xFFFF & (In.RadiusScale01 >> 16));
 	return Out;
 }
 
@@ -104,6 +105,7 @@ struct FHairClusterLOD
 {
 	float RadiusScale;
 	float LOD;
+	uint  LODCount;
 };
 
 ////////////////////////////////////////////////////////////////////////////////
@@ -171,33 +173,34 @@ float GetLOD(FHairClusterInfo InInfo, float InScreenSize, float InLODBias)
 
 // Compute the hair cluster LOD for a given LOD level
 FHairClusterLOD GetHairClusterLOD(
-	FHairClusterInfo InClusterInfo, 
-	StructuredBuffer<FPackedHairClusterLODInfo> InClusterLODInfoBuffer,
+	FPackedHairClusterInfo InPacked, 
+	float4 InClusterInfoParameters,
 	float LOD)
 {
-	const uint iLOD = clamp(floor(LOD), 0, InClusterInfo.LODCount-1);
-	FHairClusterLODInfo ClusterLODInfo = UnpackHairClusterLODInfo(InClusterLODInfoBuffer[InClusterInfo.LODInfoOffset + iLOD]);
+	const FHairClusterInfo InInfo = UnpackHairClusterInfo(InPacked, InClusterInfoParameters);
 
+	float RadiusScale[MAX_HAIR_LOD] = 
+	{
+		InInfo.LOD0_RadiusScale,
+		InInfo.LOD1_RadiusScale,
+		InInfo.LOD2_RadiusScale,
+		InInfo.LOD3_RadiusScale,
+		InInfo.LOD4_RadiusScale,
+		InInfo.LOD5_RadiusScale,
+		InInfo.LOD6_RadiusScale,
+		InInfo.LOD7_RadiusScale
+	};
+
+	const uint iLOD = clamp(floor(LOD), 0, InInfo.LODCount-1);
 	const float S = LOD - iLOD;
 
 	FHairClusterLOD Out;
-	Out.RadiusScale = lerp(ClusterLODInfo.RadiusScale0, ClusterLODInfo.RadiusScale1, S);
+	Out.RadiusScale = lerp(RadiusScale[iLOD], RadiusScale[uint(min(iLOD + 1, InInfo.LODCount - 1))], S);
 	Out.LOD = LOD;
-
+	Out.LODCount = InInfo.LODCount;
 	return Out;
 }
 
-// Compute the hair cluster LOD based on screen size and lod bias
-FHairClusterLOD GetHairClusterLOD(
-	FHairClusterInfo InClusterInfo,
-	StructuredBuffer<FPackedHairClusterLODInfo> InClusterLODInfoBuffer,
-	float InScreenSize, 
-	float InLODBias)
-{
-	const float LOD = GetLOD(InClusterInfo, InScreenSize, InLODBias);
-	return GetHairClusterLOD(InClusterInfo, InClusterLODInfoBuffer, LOD);
-}
-
 // Compute the screen size of a bounding sphere
 // This is the equivalent of ComputeBoundsScreenSize in SceneManagement.h
 float ComputeBoundsScreenSize(float3 InSphereOrigin, float InSphereRadius, float3 InViewOrigin, float4x4 InProjMatrix)
diff --git a/Engine/Shaders/Private/HairStrands/HairStrandsClusterCulling.usf b/Engine/Shaders/Private/HairStrands/HairStrandsClusterCulling.usf
--- a/Engine/Shaders/Private/HairStrands/HairStrandsClusterCulling.usf
+++ b/Engine/Shaders/Private/HairStrands/HairStrandsClusterCulling.usf
@@ -300,12 +300,12 @@ float LODBias;
 uint CurveCount;
 uint CurvePerGroup;
 uint DispatchCountX;
+float4 ClusterInfoParameters;
 
 Buffer<uint> CurveBuffer;
 Buffer<uint> CurveToClusterIdBuffer;
 Buffer<uint> PointLODBuffer;
 StructuredBuffer<FPackedHairClusterInfo> ClusterInfoBuffer;
-StructuredBuffer<FPackedHairClusterLODInfo> ClusterLODInfoBuffer; 
 
 RWBuffer<uint>  OutPointCounter;
 RWBuffer<uint>  OutIndexBuffer;
@@ -372,14 +372,11 @@ void Main(uint3 DispatchThreadId : SV_DispatchThreadID, uint ThreadIndex : SV_Gr
 		// 1.1 Cluster ID
 		const uint ClusterId = CurveToClusterIdBuffer[GlobalCurveIndex];
 	
-		// 1.2 Get cluster info (i.e. map all info per LOD for a given cluster ID)
-		const FHairClusterInfo ClusterInfo = UnpackHairClusterInfo(ClusterInfoBuffer[ClusterId]);
-	
-		// 1.3 Derive effective LOD data (radius scale)
-		const FHairClusterLOD ClusterLODInfo = GetHairClusterLOD(ClusterInfo, ClusterLODInfoBuffer, LODIndex);
+		// 1.3 Get cluster info (i.e. map all info per LOD for a given cluster ID) and derive effective LOD data (radius scale)
+		const FHairClusterLOD ClusterLODInfo = GetHairClusterLOD(ClusterInfoBuffer[ClusterId], ClusterInfoParameters, LODIndex);
 		ClusterLOD[LocalCurveIndex] 		 = ClusterLODInfo.LOD;
 		ClusterRadiusScale[LocalCurveIndex]  = ClusterLODInfo.RadiusScale;
-		LODCount[LocalCurveIndex]			 = ClusterInfo.LODCount;
+		LODCount[LocalCurveIndex]			 = ClusterLODInfo.LODCount;
 
 		// 1.4 Curve's point count/offset
 		const FHairCurve Curve = UnpackHairCurve(CurveBuffer[GlobalCurveIndex]);
@@ -500,7 +497,7 @@ groupshared float CulledPointFadeScales[HAIR_MAX_NUM_POINT_PER_CURVE];
 // Version for curve with >32 points
 [numthreads(GROUP_SIZE, 1, 1)]
 void Main(uint3 DispatchThreadId : SV_DispatchThreadID, uint ThreadIndex : SV_GroupIndex, uint2 GroupIndex : SV_GroupID)
-{
+{	
 	// 1. Read the Curve data
 	const uint CurveIndex = GetGlobalCurveIndex(ThreadIndex, GroupIndex);
 	
@@ -509,14 +506,11 @@ void Main(uint3 DispatchThreadId : SV_DispatchThreadID, uint ThreadIndex : SV_Gr
 		// 1.1 Cluster ID
 		const uint ClusterId = CurveToClusterIdBuffer[CurveIndex];
 	
-		// 1.2 Get cluster info (i.e. map all info per LOD for a given cluster ID)
-		const FHairClusterInfo ClusterInfo = UnpackHairClusterInfo(ClusterInfoBuffer[ClusterId]);
-	
-		// 1.3 Derive effective LOD data (radius scale)
-		const FHairClusterLOD ClusterLODInfo = GetHairClusterLOD(ClusterInfo, ClusterLODInfoBuffer, LODIndex);
+		// 1.2 Get cluster info (i.e. map all info per LOD for a given cluster ID) and derive effective LOD data (radius scale)
+		const FHairClusterLOD ClusterLODInfo = GetHairClusterLOD(ClusterInfoBuffer[ClusterId], ClusterInfoParameters, LODIndex);
 		ClusterLOD = ClusterLODInfo.LOD;
 		ClusterRadiusScale = ClusterLODInfo.RadiusScale;
-		LODCount = ClusterInfo.LODCount;
+		LODCount = ClusterLODInfo.LODCount;
 
 		// 1.4 Curve's point count/offset
 		const FHairCurve Curve = UnpackHairCurve(CurveBuffer[CurveIndex]);
