Fix CIS warning.

diff --git a/Engine/Plugins/Runtime/HairStrands/Source/HairStrandsCore/Private/HairStrandsInterpolation.cpp b/Engine/Plugins/Runtime/HairStrands/Source/HairStrandsCore/Private/HairStrandsInterpolation.cpp
--- a/Engine/Plugins/Runtime/HairStrands/Source/HairStrandsCore/Private/HairStrandsInterpolation.cpp
+++ b/Engine/Plugins/Runtime/HairStrands/Source/HairStrandsCore/Private/HairStrandsInterpolation.cpp
@@ -902,7 +902,7 @@ static void AddHairClusterAABBPass(
 
 	if (UpdateType == EHairAABBUpdateType::UpdateClusterAABB)
 	{
-		FIntVector DispatchCount(ClusterData->ClusterCount, 1, 1);
+		FIntVector DispatchCount(Parameters->ClusterCount, 1, 1);
 		FComputeShaderUtils::AddPass(
 			GraphBuilder,
 			RDG_EVENT_NAME("HairStrands::ClusterAABB(Cluster,Group)"),
