Fix crash in hair culling when shader print is not available.

diff --git a/Engine/Plugins/Runtime/HairStrands/Source/HairStrandsCore/Private/HairStrandsClusterCulling.cpp b/Engine/Plugins/Runtime/HairStrands/Source/HairStrandsCore/Private/HairStrandsClusterCulling.cpp
--- a/Engine/Plugins/Runtime/HairStrands/Source/HairStrandsCore/Private/HairStrandsClusterCulling.cpp
+++ b/Engine/Plugins/Runtime/HairStrands/Source/HairStrandsCore/Private/HairStrandsClusterCulling.cpp
@@ -140,7 +140,10 @@ static void AddClusterCullingPass(
 		Parameters->OutRadiusScaleBuffer = RegisterAsUAV(GraphBuilder, *ClusterData.GetCulledVertexRadiusScaleBuffer());
 		Parameters->OutCulledCurveBuffer = RegisterAsUAV(GraphBuilder, *ClusterData.GetCulledCurveBuffer()); // TODO: this could be changed to be transient buffer
 
-		ShaderPrint::SetParameters(GraphBuilder, *ShaderPrintData, Parameters->ShaderPrintUniformBuffer);
+		if (ShaderPrintData)
+		{
+			ShaderPrint::SetParameters(GraphBuilder, *ShaderPrintData, Parameters->ShaderPrintUniformBuffer);
+		}
 
 		// Pick permutation based on groom max number of point per curve
 		//const uint32 PointPerCurve = FMath::Clamp(uint32(FMath::Pow( FMath::RoundFromZero(FMath::Log2(float(ClusterData.MaxPointPerCurve))),2u)), 4u, 32u);
