Fix cook issue with groom asset probably caused by incorrect DDC caching.

diff --git a/Engine/Plugins/Runtime/HairStrands/Source/HairStrandsCore/Private/GroomAsset.cpp b/Engine/Plugins/Runtime/HairStrands/Source/HairStrandsCore/Private/GroomAsset.cpp
--- a/Engine/Plugins/Runtime/HairStrands/Source/HairStrandsCore/Private/GroomAsset.cpp
+++ b/Engine/Plugins/Runtime/HairStrands/Source/HairStrandsCore/Private/GroomAsset.cpp
@@ -129,6 +129,7 @@ uint32 FHairGroupPlatformData::FStrands::GetDataSize() const
 	CTotal += ClusterCullingBulkData.Header.CPULODScreenSize.GetAllocatedSize();
 	BTotal += ClusterCullingBulkData.Data.CurveToClusterIds.IsBulkDataLoaded()	? ClusterCullingBulkData.Data.CurveToClusterIds.GetBulkDataSize() : 0;
 	CTotal += ClusterCullingBulkData.Data.PackedClusterInfos.IsBulkDataLoaded()	? ClusterCullingBulkData.Data.PackedClusterInfos.GetBulkDataSize() : 0;
+	CTotal += ClusterCullingBulkData.Data.PointLODs.IsBulkDataLoaded()	? ClusterCullingBulkData.Data.PointLODs.GetBulkDataSize() : 0;
 
 	return BTotal + Total + CTotal;
 }
diff --git a/Engine/Plugins/Runtime/HairStrands/Source/HairStrandsCore/Private/GroomBuilder.cpp b/Engine/Plugins/Runtime/HairStrands/Source/HairStrandsCore/Private/GroomBuilder.cpp
--- a/Engine/Plugins/Runtime/HairStrands/Source/HairStrandsCore/Private/GroomBuilder.cpp
+++ b/Engine/Plugins/Runtime/HairStrands/Source/HairStrandsCore/Private/GroomBuilder.cpp
@@ -40,7 +40,7 @@ static FAutoConsoleVariableRef CVarHairGroupIndexBuilder_MaxVoxelResolution(TEXT
 
 FString FGroomBuilder::GetVersion()
 {
-	return TEXT("v8r37");
+	return TEXT("v8r37_a");
 }
 
 namespace FHairStrandsDecimation
diff --git a/Engine/Plugins/Runtime/HairStrands/Source/HairStrandsCore/Private/HairStrandsDatas.cpp b/Engine/Plugins/Runtime/HairStrands/Source/HairStrandsCore/Private/HairStrandsDatas.cpp
--- a/Engine/Plugins/Runtime/HairStrands/Source/HairStrandsCore/Private/HairStrandsDatas.cpp
+++ b/Engine/Plugins/Runtime/HairStrands/Source/HairStrandsCore/Private/HairStrandsDatas.cpp
@@ -708,18 +708,12 @@ bool ValidateHairBulkData();
 
 void FHairStrandsClusterCullingBulkData::GetResources(FHairStrandsBulkCommon::FQuery & Out)
 {
-	if (Header.CurveCount)
+	const bool bHasClusterData = Header.ClusterCount > 0;
+	if (bHasClusterData)
 	{
-		Out.Add(Data.CurveToClusterIds, TEXT("_CurveToClusterIds"), 0, 0); // Load all data
-	}
-
-	if (Header.ClusterCount)
-	{
-		Out.Add(Data.PackedClusterInfos, TEXT("_PackedClusterInfos"), 0, 0); // Load all data
-	}
-
-	{
-		Out.Add(Data.PointLODs, TEXT("_PointLODs"), 0, 0); // Load all data. TODO load data based on active curve/point count
+		Out.Add(Data.CurveToClusterIds, TEXT("_CurveToClusterIds"), 0, 0); 	// Load all data. TODO load data based on active curve/point count
+		Out.Add(Data.PackedClusterInfos, TEXT("_PackedClusterInfos"), 0, 0);// Load all data
+		Out.Add(Data.PointLODs, TEXT("_PointLODs"), 0, 0); 					// Load all data. TODO load data based on active curve/point count
 	}
 }
 
