Fix crash in MLDeformerSample.

diff --git a/Engine/Plugins/Runtime/HairStrands/Source/HairStrandsCore/Private/GroomResources.cpp b/Engine/Plugins/Runtime/HairStrands/Source/HairStrandsCore/Private/GroomResources.cpp
--- a/Engine/Plugins/Runtime/HairStrands/Source/HairStrandsCore/Private/GroomResources.cpp
+++ b/Engine/Plugins/Runtime/HairStrands/Source/HairStrandsCore/Private/GroomResources.cpp
@@ -1255,7 +1255,7 @@ bool FHairStrandsRestRootResource::InternalGetOrRequestData(uint32 InRequestedCu
 	{
 		StreamingRequest.Request(InRequestedCurveCount, InRequestedPointCount, InLODIndex, BulkData, false, false, false, OwnerName);
 	}
-	return StreamingRequest.IsCompleted();
+	return InLODIndex >= 0 && StreamingRequest.IsCompleted();
 }
 
 void FHairStrandsRestRootResource::InternalAllocate(FRDGBuilder& GraphBuilder, uint32 InCurveCount, int32 InLODIndex)
