Fix hair physics settings working currectly with custom reset.

diff --git a/Engine/Plugins/Runtime/HairStrands/Source/HairStrandsEditor/Private/GroomAssetDetails.cpp b/Engine/Plugins/Runtime/HairStrands/Source/HairStrandsEditor/Private/GroomAssetDetails.cpp
--- a/Engine/Plugins/Runtime/HairStrands/Source/HairStrandsEditor/Private/GroomAssetDetails.cpp
+++ b/Engine/Plugins/Runtime/HairStrands/Source/HairStrandsEditor/Private/GroomAssetDetails.cpp
@@ -1152,6 +1152,42 @@ void FGroomRenderingDetails::ExpandStruct(TSharedRef<IPropertyHandle>& PropertyH
 	}
 }
 
+void FGroomRenderingDetails::AddPropertySeparator(FName PropertyName, IDetailChildrenBuilder& ChildrenBuilder)
+{
+	static const FSlateBrush* GenericBrush = FCoreStyle::Get().GetBrush("GenericWhiteBox");
+	float OtherMargin = 0.0f;
+	float RightMargin = 10.0f;
+	const FLinearColor SeparatorColor(0.05f, 0.05f,0.05f);
+	ChildrenBuilder.AddCustomRow(LOCTEXT("Hair_Separator", "Separator"))
+	.WholeRowContent()
+	.VAlign(VAlign_Fill)
+	.HAlign(HAlign_Fill)
+	[
+		SNew(SOverlay)
+		+ SOverlay::Slot()
+		[
+			SNew(SImage)
+			.Image(GenericBrush)
+			.ColorAndOpacity(SeparatorColor)
+		]
+		+ SOverlay::Slot()
+		.HAlign(HAlign_Left)
+		.VAlign(VAlign_Center)
+		[
+			SNew(SHorizontalBox)
+			+ SHorizontalBox::Slot()
+			.VAlign(VAlign_Center)
+			.Padding(RightMargin, OtherMargin, RightMargin, OtherMargin)
+			[
+				SNew(STextBlock)
+				.Font(IDetailLayoutBuilder::GetDetailFont())
+				.ColorAndOpacity(FLinearColor(0.9f, 0.9f,0.9f))
+				.Text(FText::FromName(PropertyName))
+			]						
+		]
+	];
+}
+
 FReply FGroomRenderingDetails::OnRemoveLODClicked(int32 GroupIndex, int32 LODIndex, FProperty* Property)
 {
 	check(GroomAsset);
@@ -1747,7 +1783,29 @@ void FGroomRenderingDetails::OnGenerateElementForHairGroup(TSharedRef<IPropertyH
 		break;
 		case EMaterialPanelType::Physics:
 		{
-			AddPropertyWithCustomReset(ChildHandle, ChildrenBuilder, GroupIndex, -1);
+			if (PropertyName == GET_MEMBER_NAME_CHECKED(FHairGroupsPhysics, SolverSettings) ||
+				PropertyName == GET_MEMBER_NAME_CHECKED(FHairGroupsPhysics, ExternalForces) ||
+				PropertyName == GET_MEMBER_NAME_CHECKED(FHairGroupsPhysics, StrandsParameters))
+			{
+				AddPropertySeparator(PropertyName, ChildrenBuilder);
+				ExpandStruct(ChildHandle, ChildrenBuilder, GroupIndex, -1, true);
+			}
+			else if (PropertyName == GET_MEMBER_NAME_CHECKED(FHairGroupsPhysics, MaterialConstraints))
+			{
+				// Expand the constraint, so that each constraints type has its own block (so that custom value reset works correctly)
+				uint32 SubChildrenCount = 0;
+				ChildHandle->GetNumChildren(SubChildrenCount);
+				for (uint32 SubChildIt = 0; SubChildIt < SubChildrenCount; ++SubChildIt)
+				{
+					TSharedPtr<IPropertyHandle> SubChildHandle = ChildHandle->GetChildHandle(SubChildIt);					
+					AddPropertySeparator(SubChildHandle->GetProperty()->GetFName(), ChildrenBuilder);
+					ExpandStruct(SubChildHandle, ChildrenBuilder, GroupIndex, -1, true);
+				}
+			}
+			else
+			{
+				AddPropertyWithCustomReset(ChildHandle, ChildrenBuilder, GroupIndex, -1);
+			}
 		}
 		break;
 		case EMaterialPanelType::Bindings:
diff --git a/Engine/Plugins/Runtime/HairStrands/Source/HairStrandsEditor/Private/GroomAssetDetails.h b/Engine/Plugins/Runtime/HairStrands/Source/HairStrandsEditor/Private/GroomAssetDetails.h
--- a/Engine/Plugins/Runtime/HairStrands/Source/HairStrandsEditor/Private/GroomAssetDetails.h
+++ b/Engine/Plugins/Runtime/HairStrands/Source/HairStrandsEditor/Private/GroomAssetDetails.h
@@ -88,6 +88,7 @@ private:
 	FReply OnGenerateCardDataUsingPlugin(int32 GroupIndex);
 	FReply OnSelectBinding(int32 BindingIndex, FProperty* Property);
 
+	void AddPropertySeparator(FName PropertyName, IDetailChildrenBuilder& ChildrenBuilder);
 	void ExpandStructForLOD(TSharedRef<IPropertyHandle>& PropertyHandle, IDetailChildrenBuilder& ChildrenBuilder, int32 GroupIndex, int32 LODIndex, bool bOverrideReset);
 	void ExpandStruct(TSharedPtr<IPropertyHandle>& PropertyHandle, IDetailChildrenBuilder& ChildrenBuilder, int32 GroupIndex, int32 LODIndex, bool bOverrideReset);
 	void ExpandStruct(TSharedRef<IPropertyHandle>& PropertyHandle, IDetailChildrenBuilder& ChildrenBuilder, int32 GroupIndex, int32 LODIndex, bool bOverrideReset);
