GroomCache: Fixed streaming issue where the chunks ref count could get out of balance.

diff --git a/Engine/Plugins/Runtime/HairStrands/Source/HairStrandsCore/Private/GroomCacheStreamingData.cpp b/Engine/Plugins/Runtime/HairStrands/Source/HairStrandsCore/Private/GroomCacheStreamingData.cpp
--- a/Engine/Plugins/Runtime/HairStrands/Source/HairStrandsCore/Private/GroomCacheStreamingData.cpp
+++ b/Engine/Plugins/Runtime/HairStrands/Source/HairStrandsCore/Private/GroomCacheStreamingData.cpp
@@ -124,7 +124,8 @@ void FGroomCacheStreamingData::PrefetchData(UGroomComponent *Component)
 		// Already requested an async load but not complete yet 
 		if (ChunksRequested.Contains(ChunkId))
 		{
-			ChunksRequested.Remove(ChunkId);
+			// Chunk will be removed from ChunksRequested once the read request is completed in ProcessCompletedChunks
+			// This prevents potentially requesting another read for this Chunk
 
 			// Let the in-flight read request complete with the initial AnimDataPtr to avoid threading issues
 			// The initial buffer will be orphaned but will be deleted once the completed read request is processed
@@ -145,6 +146,14 @@ void FGroomCacheStreamingData::UpdateStreamingStatus(bool bAsyncDeletionAllowed)
 {
 	TRACE_CPUPROFILER_EVENT_SCOPE(FGroomCacheStreamingData::UpdateStreamingStatus);
 
+	// Allow updating the StreamingData only one per frame since PrefetchData can happen
+	// many times, like when the time jumps around with thumbnail generation in Sequencer.
+	// This will prevent adding and removing chunks unnecessarily.
+	if (LastUpdate == GFrameCounter)
+	{
+		return;
+	}
+
 	// Find any chunks that aren't available yet
 	for (int32 NeededIndex : ChunksNeeded)
 	{
@@ -296,6 +305,8 @@ void FGroomCacheStreamingData::UpdateStreamingStatus(bool bAsyncDeletionAllowed)
 			}
 		}
 	}
+
+	LastUpdate = GFrameCounter;
 }
 
 bool FGroomCacheStreamingData::BlockTillAllRequestsFinished(float TimeLimit)
@@ -353,16 +364,28 @@ void FGroomCacheStreamingData::ProcessCompletedChunks()
 		// in which case, we don't care about this chunk anymore 
 		if (Chunk->IORequest != nullptr)
 		{
-			if (CompletedChunk.bReadSuccessful)
+			// There could be more than one read requests in-flight for the same Chunk so make sure that the completed
+			// read request is still the one we're waiting for
+			if (Chunk->IORequest == CompletedChunk.ReadRequest)
 			{
-				ChunksAvailable.Add(CompletedChunk.LoadedChunkIndex);
-			}
-			else
-			{
-				UE_LOG(LogGroomCacheStreaming, Log, TEXT("Async loading request failed!"));
+				if (CompletedChunk.bReadSuccessful)
+				{
+					ChunksAvailable.Add(CompletedChunk.LoadedChunkIndex);
+				}
+				else
+				{
+					UE_LOG(LogGroomCacheStreaming, Log, TEXT("Async loading request failed!"));
+				}
+
+				Chunk->IORequest = nullptr;
+				ChunksRequested.Remove(CompletedChunk.LoadedChunkIndex);
 			}
+			// If the read request doesn't match, it means we're still waiting for IORequest to complete so it must remain in ChunksRequested
+		}
+		else
+		{
+			ChunksRequested.Remove(CompletedChunk.LoadedChunkIndex);
 		}
-		ChunksRequested.Remove(CompletedChunk.LoadedChunkIndex);
 
 		// The AnimData buffer can change while the read request was in-flight because the associated chunk was prefetched synchronously instead
 		// or it's been removed from resident chunks, so just delete the orphaned buffer
@@ -371,8 +394,6 @@ void FGroomCacheStreamingData::ProcessCompletedChunks()
 			delete CompletedChunk.AnimDataPtr;
 		}
 
-		Chunk->IORequest = nullptr;
-
 		// Process the read request deletion later as in can be quite time-consuming
 		DelayedDeleteReadRequests.Add(CompletedChunk.ReadRequest);
 	}
diff --git a/Engine/Plugins/Runtime/HairStrands/Source/HairStrandsCore/Private/GroomCacheStreamingData.h b/Engine/Plugins/Runtime/HairStrands/Source/HairStrandsCore/Private/GroomCacheStreamingData.h
--- a/Engine/Plugins/Runtime/HairStrands/Source/HairStrandsCore/Private/GroomCacheStreamingData.h
+++ b/Engine/Plugins/Runtime/HairStrands/Source/HairStrandsCore/Private/GroomCacheStreamingData.h
@@ -65,6 +65,8 @@ public:
 	/** Adds ChunkIndex to the list of chunks needed to be loaded */
 	void AddNeededChunk(uint32 ChunkIndex);
 
+	UGroomCache* GetGroomCache() const { return GroomCache; }
+
 private:
 
 	FResidentGroomCacheChunk& AddResidentChunk(int32 ChunkId, const struct FGroomCacheChunk& ChunkInfo);
@@ -95,6 +97,9 @@ private:
 	/** List of read requests to be deleted after completion */
 	TSet<IBulkDataIORequest*> DelayedDeleteReadRequests;
 
+	/** Timestamp of the last update */
+	uint64 LastUpdate = 0;
+
 	FCriticalSection CriticalSection;
 };
 
diff --git a/Engine/Plugins/Runtime/HairStrands/Source/HairStrandsCore/Private/GroomCacheStreamingManager.cpp b/Engine/Plugins/Runtime/HairStrands/Source/HairStrandsCore/Private/GroomCacheStreamingManager.cpp
--- a/Engine/Plugins/Runtime/HairStrands/Source/HairStrandsCore/Private/GroomCacheStreamingManager.cpp
+++ b/Engine/Plugins/Runtime/HairStrands/Source/HairStrandsCore/Private/GroomCacheStreamingManager.cpp
@@ -195,7 +195,28 @@ void FGroomCacheStreamingManager::RegisterGroomCache(UGroomCache* GroomCache, UG
 
 	if (!StreamingGroomCaches.Contains(GroomCache))
 	{
-		StreamingGroomCaches.Add(GroomCache, new FGroomCacheStreamingData(GroomCache));
+		// Try to recover a StreamingData previously marked for deletion.
+		// This helps performance by re-using data that is already loaded
+		// and helps maintain accurate ref counts of the Chunks since a
+		// delayed unmap could happen after an Unregister+RegisterGroomCache.
+		FGroomCacheStreamingData* StreamingData = nullptr;
+		for (FGroomCacheStreamingData* StreamingDataForDeletion : StreamingGroomCachesToDelete)
+		{
+			if (StreamingDataForDeletion->GetGroomCache() == GroomCache)
+			{
+				StreamingData = StreamingDataForDeletion;
+				StreamingGroomCachesToDelete.Remove(StreamingDataForDeletion);
+				break;
+			}
+		}
+
+		if (!StreamingData)
+		{
+			StreamingData = new FGroomCacheStreamingData(GroomCache);
+		}
+
+		StreamingGroomCaches.Add(GroomCache, StreamingData);
+
 		GroomCacheUsers.Add(GroomCache, {});
 	}
 
diff --git a/Engine/Plugins/Runtime/HairStrands/Source/HairStrandsCore/Private/GroomComponent.cpp b/Engine/Plugins/Runtime/HairStrands/Source/HairStrandsCore/Private/GroomComponent.cpp
--- a/Engine/Plugins/Runtime/HairStrands/Source/HairStrandsCore/Private/GroomComponent.cpp
+++ b/Engine/Plugins/Runtime/HairStrands/Source/HairStrandsCore/Private/GroomComponent.cpp
@@ -1113,6 +1113,8 @@ public:
 private:
 	void ResetInternal()
 	{
+		FScopeLock Lock(GetCriticalSection());
+
 		// Unmap the frames that are currently mapped
 		if (CurrentFrameIndex != -1)
 		{
