Fix slow opening of groom binding tab in groom editor by avoiding loading asset.

Assert search use name matching as heuristic.

diff --git a/Engine/Plugins/Runtime/HairStrands/Source/HairStrandsEditor/Private/GroomCustomAssetEditorToolkit.cpp b/Engine/Plugins/Runtime/HairStrands/Source/HairStrandsEditor/Private/GroomCustomAssetEditorToolkit.cpp
--- a/Engine/Plugins/Runtime/HairStrands/Source/HairStrandsEditor/Private/GroomCustomAssetEditorToolkit.cpp
+++ b/Engine/Plugins/Runtime/HairStrands/Source/HairStrandsEditor/Private/GroomCustomAssetEditorToolkit.cpp
@@ -326,29 +326,62 @@ void FGroomCustomAssetEditorToolkit::OnClose()
 
 static void ListAllBindingAssets(const UGroomAsset* InGroomAsset, TWeakObjectPtr<UGroomBindingAssetList>& Out)
 {
-	FARFilter Filter;
-	Filter.ClassPaths.Add(UGroomBindingAsset::StaticClass()->GetClassPathName());
-
-	// Search for binding with matching groom asset
-	//{
-	//	const FName GroomAssetProperty = GET_MEMBER_NAME_CHECKED(UGroomBindingAsset, Groom);
-	//	FString GroomAssetName = FAssetData(InGroomAsset).GetExportTextName();
-	//	FString GroomAssetName = TSoftObjectPtr<UGroomAsset>(InGroomAsset).ToString();
-	//	Filter.TagsAndValues.Add(GroomAssetProperty, GroomAssetName);
-	//}
-
 	FAssetRegistryModule& AssetRegistryModule = FModuleManager::LoadModuleChecked<FAssetRegistryModule>("AssetRegistry");
 	TArray<FAssetData> BindingAssetData;
-	AssetRegistryModule.Get().GetAssets(Filter, BindingAssetData);
 
-	// Filter binding asset which match the groom asset (as the tag/value filter above does not work)
-	for (FAssetData& Asset :  BindingAssetData)
+	// 1. Use tagged properties for searching compatible binding assets (for assets saved after AssetRegistrySearchable has been added)
 	{
-		if (UGroomBindingAsset* Binding = (UGroomBindingAsset*)Asset.GetAsset())
+		const FName GroomAssetProperty = GET_MEMBER_NAME_CHECKED(UGroomBindingAsset, Groom);
+		FString GroomAssetName = FAssetData(InGroomAsset).GetExportTextName();
+		FARFilter Filter;
+		Filter.ClassPaths.Add(UGroomBindingAsset::StaticClass()->GetClassPathName());
+		Filter.TagsAndValues.Add(GroomAssetProperty, GroomAssetName);
+		AssetRegistryModule.Get().GetAssets(Filter, BindingAssetData);
+
+		for (FAssetData& Asset : BindingAssetData)
 		{
-			if (Binding->Groom == InGroomAsset)
+			if (UGroomBindingAsset* Binding = (UGroomBindingAsset*)Asset.GetAsset(/*bLoad*/))
+			{
+				if (Binding->Groom == InGroomAsset)
+				{
+					Out->Bindings.Add(Binding);
+				}
+			}
+		}
+	}
+
+	// 2. Use name matching for searching compatible binding assets
+	if (BindingAssetData.Num() == 0)
+	{
+		FARFilter Filter;
+		Filter.ClassPaths.Add(UGroomBindingAsset::StaticClass()->GetClassPathName());
+		AssetRegistryModule.Get().GetAssets(Filter, BindingAssetData);
+	
+		// Filter binding asset which match the groom asset (as the tag/value filter above does not work)
+		// Avoid to load binding asset to slow down panel opening
+		for (FAssetData& Asset :  BindingAssetData)
+		{
+			UGroomBindingAsset* Binding = (UGroomBindingAsset*)Asset.FastGetAsset(false /*bLoad*/);
+			if (Binding)
+			{
+				if (Binding->Groom == InGroomAsset)
+				{
+					Out->Bindings.Add(Binding);
+				}
+			}
+			else 
 			{
-				Out->Bindings.Add(Binding);
+				// Use heuristic that binding asset usually contain groom asset name (when preserving the auto-generated binding name)
+				const FString BindingAssetName = Asset.AssetName.GetPlainNameString();
+				const FString GroomAssetName = InGroomAsset->GetName();
+				if (TCString<TCHAR>::Strfind(*BindingAssetName, *GroomAssetName, true))
+				{
+					Binding = (UGroomBindingAsset*)Asset.GetAsset(/*bLoad*/);
+					if (Binding->Groom == InGroomAsset)
+					{
+						Out->Bindings.Add(Binding);
+					}
+				}
 			}
 		}
 	}
