Fix pop when switching betwee rigid <-> skinned binding

diff --git a/Engine/Plugins/Runtime/HairStrands/Source/HairStrandsCore/Private/GroomComponent.cpp b/Engine/Plugins/Runtime/HairStrands/Source/HairStrandsCore/Private/GroomComponent.cpp
--- a/Engine/Plugins/Runtime/HairStrands/Source/HairStrandsCore/Private/GroomComponent.cpp
+++ b/Engine/Plugins/Runtime/HairStrands/Source/HairStrandsCore/Private/GroomComponent.cpp
@@ -883,6 +883,12 @@ public:
 		FMatrix CurrentLocalToWorld = bUseProxy ? GetLocalToWorld() : Instance->GetCurrentLocalToWorld().ToMatrixWithScale();
 		PreviousLocalToWorld = bUseProxy ? PreviousLocalToWorld : Instance->GetPreviousLocalToWorld().ToMatrixWithScale();
 
+		// Band-aid to avoid invalid velociy vector when switching LOD skinned <-> rigid
+		if (Instance->HairGroupPublicData->VFInput.bHasLODSwitch && Instance->HairGroupPublicData->VFInput.bHasLODSwitchBindingType)
+		{
+			PreviousLocalToWorld = CurrentLocalToWorld;
+		}
+
 		FDynamicPrimitiveUniformBuffer& DynamicPrimitiveUniformBuffer = Collector.AllocateOneFrameResource<FDynamicPrimitiveUniformBuffer>();
 		DynamicPrimitiveUniformBuffer.Set(CurrentLocalToWorld, PreviousLocalToWorld, GetBounds(), GetLocalBounds(), true, false, bOutputVelocity);
 		BatchElement.PrimitiveUniformBufferResource = &DynamicPrimitiveUniformBuffer.UniformBuffer; // automatic copy to the gpu scene buffer
diff --git a/Engine/Plugins/Runtime/HairStrands/Source/HairStrandsCore/Private/GroomDebug.cpp b/Engine/Plugins/Runtime/HairStrands/Source/HairStrandsCore/Private/GroomDebug.cpp
--- a/Engine/Plugins/Runtime/HairStrands/Source/HairStrandsCore/Private/GroomDebug.cpp
+++ b/Engine/Plugins/Runtime/HairStrands/Source/HairStrandsCore/Private/GroomDebug.cpp
@@ -886,10 +886,11 @@ static void AddHairDebugPrintInstancePass(
 
 	struct FInstanceInfos
 	{
-		FUintVector4 Data0;
-		FUintVector4 Data1;
-		FUintVector4 Data2;
-		FUintVector4 Data3;
+		FUintVector4 Data0 = { 0,0,0,0 };
+		FUintVector4 Data1 = { 0,0,0,0 };
+		FUintVector4 Data2 = { 0,0,0,0 };
+		FUintVector4 Data3 = { 0,0,0,0 };
+		FUintVector4 Data4 = { 0,0,0,0 };
 	};
 	TArray<FInstanceInfos> Infos;
 	Infos.Reserve(InstanceCount);
@@ -920,11 +921,8 @@ static void AddHairDebugPrintInstancePass(
 		const uint32 LODCount = Instance->HairGroupPublicData->GetLODScreenSizes().Num();
 		const EGroomCacheType ActiveGroomCacheType = GetHairInstanceCacheType(Instance);
 
-		FUintVector4 Data0 = { 0,0,0,0 };
-		FUintVector4 Data1 = { 0,0,0,0 };
-		FUintVector4 Data2 = { 0,0,0,0 };
-		FUintVector4 Data3 = { 0,0,0,0 };
-		Data0.X =
+		FInstanceInfos& D = Infos.AddDefaulted_GetRef();
+		D.Data0.X =
 			((Instance->Debug.GroupIndex & 0xFF)) |
 			((Instance->Debug.GroupCount & 0xFF) << 8) |
 			((LODCount & 0xFF) << 16) |
@@ -933,80 +931,81 @@ static void AddHairDebugPrintInstancePass(
 			((Instance->Guides.bIsSimulationEnable ? 0x1 : 0x0) << 30) |
 			((Instance->Guides.bHasGlobalInterpolation ? 0x1 : 0x0) << 31);
 
-		Data0.Y =
+		D.Data0.Y =
 			(FFloat16(LODIndex).Encoded) |
 			(FFloat16(Instance->Strands.Modifier.HairLengthScale_Override ? Instance->Strands.Modifier.HairLengthScale : -1.f).Encoded << 16);
 		
-		switch (Instance->GeometryType)
+		D.Data4.X = Instance->HairGroupPublicData->VFInput.bHasLODSwitch ? 1u : 0u;
+
+		switch (Instance->GeometryType) 
 		{
 		case EHairGeometryType::Strands:
 			if (Instance->Strands.IsValid())
 			{
 				check(Instance->Strands.Data);
-				Data0.Z = Instance->Strands.Data->GetNumCurves(); // Change this later on for having dynamic value
-				Data0.W = Instance->Strands.Data->GetNumPoints(); // Change this later on for having dynamic value
+				D.Data0.Z = Instance->Strands.Data->GetNumCurves(); // Change this later on for having dynamic value
+				D.Data0.W = Instance->Strands.Data->GetNumPoints(); // Change this later on for having dynamic value
 				const int32 MeshLODIndex = Instance->HairGroupPublicData->MeshLODIndex;
 				if (MeshLODIndex>=0 && Instance->Strands.RestRootResource)
 				{
-					Data1.X = Instance->Strands.RestRootResource->BulkData.Header.LODs[MeshLODIndex].UniqueSectionIndices.Num();
-					Data1.Y = Instance->Strands.RestRootResource->BulkData.Header.LODs[MeshLODIndex].UniqueTriangleCount;
-					Data1.Z = Instance->Strands.RestRootResource->BulkData.Header.RootCount;
-					Data1.W = Instance->Strands.RestRootResource->BulkData.Header.PointCount;
+					D.Data1.X = Instance->Strands.RestRootResource->BulkData.Header.LODs[MeshLODIndex].UniqueSectionIndices.Num();
+					D.Data1.Y = Instance->Strands.RestRootResource->BulkData.Header.LODs[MeshLODIndex].UniqueTriangleCount;
+					D.Data1.Z = Instance->Strands.RestRootResource->BulkData.Header.RootCount;
+					D.Data1.W = Instance->Strands.RestRootResource->BulkData.Header.PointCount;
 				}
 
 				{
-					Data2 = FUintVector4(0);
-					Data2.X |= InstanceIndex < InstanceCountPerType[HairInstanceCount_StrandsPrimaryView]  ? 0x1u  : 0u;
-					Data2.X |= InstanceIndex < InstanceCountPerType[HairInstanceCount_StrandsShadowView]   ? 0x2u  : 0u;
-					Data2.X |= Instance->HairGroupPublicData->VFInput.Strands.Common.bScatterSceneLighting ? 0x4u  : 0u;
-					Data2.X |= Instance->HairGroupPublicData->VFInput.Strands.Common.bRaytracingGeometry   ? 0x8u  : 0u;
-					Data2.X |= Instance->HairGroupPublicData->VFInput.Strands.Common.bStableRasterization  ? 0x10u : 0u;
-					Data2.X |= Instance->HairGroupPublicData->bSupportVoxelization                         ? 0x20u : 0u;
-					Data2.X |= ActiveGroomCacheType == EGroomCacheType::Guides                             ? 0x40u : 0u;
-					Data2.X |= ActiveGroomCacheType == EGroomCacheType::Strands                            ? 0x80u : 0u;
-					Data2.X |= uint32(FFloat16(Instance->HairGroupPublicData->ContinuousLODScreenSize).Encoded) << 16u;
-
-					Data2.Y = Instance->HairGroupPublicData->GetActiveStrandsPointCount();
-					Data2.Z = Instance->HairGroupPublicData->GetActiveStrandsCurveCount();
-
-					Data2.W = Instance->Strands.Data->Header.ImportedAttributes;
+					D.Data2 = FUintVector4(0);
+					D.Data2.X |= InstanceIndex < InstanceCountPerType[HairInstanceCount_StrandsPrimaryView]  ? 0x1u  : 0u;
+					D.Data2.X |= InstanceIndex < InstanceCountPerType[HairInstanceCount_StrandsShadowView]   ? 0x2u  : 0u;
+					D.Data2.X |= Instance->HairGroupPublicData->VFInput.Strands.Common.bScatterSceneLighting ? 0x4u  : 0u;
+					D.Data2.X |= Instance->HairGroupPublicData->VFInput.Strands.Common.bRaytracingGeometry   ? 0x8u  : 0u;
+					D.Data2.X |= Instance->HairGroupPublicData->VFInput.Strands.Common.bStableRasterization  ? 0x10u : 0u;
+					D.Data2.X |= Instance->HairGroupPublicData->bSupportVoxelization                         ? 0x20u : 0u;
+					D.Data2.X |= ActiveGroomCacheType == EGroomCacheType::Guides                             ? 0x40u : 0u;
+					D.Data2.X |= ActiveGroomCacheType == EGroomCacheType::Strands                            ? 0x80u : 0u;
+					D.Data2.X |= uint32(FFloat16(Instance->HairGroupPublicData->ContinuousLODScreenSize).Encoded) << 16u;
+
+					D.Data2.Y = Instance->HairGroupPublicData->GetActiveStrandsPointCount();
+					D.Data2.Z = Instance->HairGroupPublicData->GetActiveStrandsCurveCount();
+
+					D.Data2.W = Instance->Strands.Data->Header.ImportedAttributes;
 				}
 				
-				Data3 = FUintVector4(0);
-				Data3.X |= FFloat16(Instance->HairGroupPublicData->VFInput.Strands.Common.Radius).Encoded;
-				Data3.X |= FFloat16(Instance->HairGroupPublicData->VFInput.Strands.Common.Density).Encoded << 16u;
+				D.Data3 = FUintVector4(0);
+				D.Data3.X |= FFloat16(Instance->HairGroupPublicData->VFInput.Strands.Common.Radius).Encoded;
+				D.Data3.X |= FFloat16(Instance->HairGroupPublicData->VFInput.Strands.Common.Density).Encoded << 16u;
 
-				Data3.Y |= FFloat16(Instance->HairGroupPublicData->VFInput.Strands.Common.RootScale).Encoded;
-				Data3.Y |= FFloat16(Instance->HairGroupPublicData->VFInput.Strands.Common.TipScale).Encoded << 16u;
+				D.Data3.Y |= FFloat16(Instance->HairGroupPublicData->VFInput.Strands.Common.RootScale).Encoded;
+				D.Data3.Y |= FFloat16(Instance->HairGroupPublicData->VFInput.Strands.Common.TipScale).Encoded << 16u;
 
-				Data3.Z |= FFloat16(Instance->HairGroupPublicData->VFInput.Strands.Common.Length).Encoded;
-				Data3.Z |= FFloat16(Instance->HairGroupPublicData->VFInput.Strands.Common.LengthScale).Encoded << 16u;
+				D.Data3.Z |= FFloat16(Instance->HairGroupPublicData->VFInput.Strands.Common.Length).Encoded;
+				D.Data3.Z |= FFloat16(Instance->HairGroupPublicData->VFInput.Strands.Common.LengthScale).Encoded << 16u;
 			}
 			break;
 		case EHairGeometryType::Cards:
 			if (Instance->Cards.IsValid(IntLODIndex))
 			{
-				Data0.Z = Instance->Cards.LODs[IntLODIndex].Guides.IsValid() ? Instance->Cards.LODs[IntLODIndex].Guides.Data->GetNumCurves() : 0;
-				Data0.W = Instance->Cards.LODs[IntLODIndex].Data->GetNumVertices();
+				D.Data0.Z = Instance->Cards.LODs[IntLODIndex].Guides.IsValid() ? Instance->Cards.LODs[IntLODIndex].Guides.Data->GetNumCurves() : 0;
+				D.Data0.W = Instance->Cards.LODs[IntLODIndex].Data->GetNumVertices();
 
-				Data2 = FUintVector4(0);
-				Data2.X |= InstanceIndex < InstanceCountPerType[HairInstanceCount_CardsOrMeshesPrimaryView] ? 0x1u : 0u;
-				Data2.X |= InstanceIndex < InstanceCountPerType[HairInstanceCount_CardsOrMeshesShadowView]  ? 0x2u : 0u;
+				D.Data2 = FUintVector4(0);
+				D.Data2.X |= InstanceIndex < InstanceCountPerType[HairInstanceCount_CardsOrMeshesPrimaryView] ? 0x1u : 0u;
+				D.Data2.X |= InstanceIndex < InstanceCountPerType[HairInstanceCount_CardsOrMeshesShadowView]  ? 0x2u : 0u;
 			}
 			break;
 		case EHairGeometryType::Meshes:
 			if (Instance->Meshes.IsValid(IntLODIndex))
 			{
-				Data0.Z = 0;
-				Data0.W = Instance->Meshes.LODs[IntLODIndex].Data->GetNumVertices();
+				D.Data0.Z = 0;
+				D.Data0.W = Instance->Meshes.LODs[IntLODIndex].Data->GetNumVertices();
 
-				Data2 = FUintVector4(0);
-				Data2.X |= InstanceIndex < InstanceCountPerType[HairInstanceCount_CardsOrMeshesPrimaryView] ? 0x1u : 0u;
-				Data2.X |= InstanceIndex < InstanceCountPerType[HairInstanceCount_CardsOrMeshesShadowView]  ? 0x2u : 0u;
+				D.Data2 = FUintVector4(0);
+				D.Data2.X |= InstanceIndex < InstanceCountPerType[HairInstanceCount_CardsOrMeshesPrimaryView] ? 0x1u : 0u;
+				D.Data2.X |= InstanceIndex < InstanceCountPerType[HairInstanceCount_CardsOrMeshesShadowView]  ? 0x2u : 0u;
 			}
 			break;
 		}
-		Infos.Add({Data0, Data1, Data2, Data3 });
 	}
 
 	if (InstanceNameInfos.IsEmpty())
@@ -1043,8 +1042,8 @@ static void AddHairDebugPrintInstancePass(
 	}
 
 	const uint32 InfoInBytes  = sizeof(FInstanceInfos);
-	const uint32 InfoIn4Bytes = sizeof(FInstanceInfos) / sizeof(uint8);
-	FRDGBufferRef InfoBuffer = CreateVertexBuffer(GraphBuilder, TEXT("Hair.Debug.InstanceInfos"), FRDGBufferDesc::CreateBufferDesc(4, InfoIn4Bytes * Infos.Num()), Infos.GetData(), InfoInBytes * Infos.Num());
+	const uint32 InfoInUints = sizeof(FInstanceInfos) / sizeof(uint32);
+	FRDGBufferRef InfoBuffer = CreateVertexBuffer(GraphBuilder, TEXT("Hair.Debug.InstanceInfos"), FRDGBufferDesc::CreateBufferDesc(4, InfoInUints * Infos.Num()), Infos.GetData(), InfoInBytes * Infos.Num());
 
 	// Resource for instance names
 	FRDGBufferRef InstanceNameBuffer = CreateVertexBuffer(GraphBuilder, TEXT("Hair.Debug.InstanceNames"), FRDGBufferDesc::CreateBufferDesc(1, InstanceNames.Num()), InstanceNames.GetData(), InstanceNames.Num());
diff --git a/Engine/Plugins/Runtime/HairStrands/Source/HairStrandsCore/Private/GroomManager.cpp b/Engine/Plugins/Runtime/HairStrands/Source/HairStrandsCore/Private/GroomManager.cpp
--- a/Engine/Plugins/Runtime/HairStrands/Source/HairStrandsCore/Private/GroomManager.cpp
+++ b/Engine/Plugins/Runtime/HairStrands/Source/HairStrandsCore/Private/GroomManager.cpp
@@ -1046,16 +1046,20 @@ static void RunHairLODSelection(
 			const bool bIsLODDataReady = !ResourceStatus.HasStatus(EHairResourceStatus::EStatus::Loading) || (IsHairStrandsContinousLODEnabled() && ResourceStatus.AvailableCurveCount > 0);
 			if (bIsLODDataReady)
 			{
-				EHairBindingType BindingType = Instance->HairGroupPublicData->GetBindingType(IntLODIndex);
+				const uint32 IntPrevLODIndex = FMath::FloorToInt(PrevLODIndex); 
+				const EHairBindingType CurrBindingType = Instance->HairGroupPublicData->GetBindingType(IntLODIndex);
+				const EHairBindingType PrevBindingType = Instance->HairGroupPublicData->GetBindingType(IntPrevLODIndex);
+
 				Instance->HairGroupPublicData->SetLODVisibility(bIsVisible);
 				Instance->HairGroupPublicData->SetLODIndex(LODIndex);
 				Instance->HairGroupPublicData->SetLODBias(0);
 				Instance->HairGroupPublicData->SetMeshLODIndex(MeshLODIndex);
 				Instance->HairGroupPublicData->VFInput.GeometryType = GeometryType;
-				Instance->HairGroupPublicData->VFInput.BindingType = BindingType;
-				Instance->HairGroupPublicData->VFInput.bHasLODSwitch = (FMath::FloorToInt(PrevLODIndex) != FMath::FloorToInt(LODIndex));
+				Instance->HairGroupPublicData->VFInput.BindingType = CurrBindingType;
+				Instance->HairGroupPublicData->VFInput.bHasLODSwitch = IntPrevLODIndex != IntLODIndex;
+				Instance->HairGroupPublicData->VFInput.bHasLODSwitchBindingType = CurrBindingType != PrevBindingType;
 				Instance->GeometryType = GeometryType;
-				Instance->BindingType = BindingType;
+				Instance->BindingType = CurrBindingType;
 				Instance->Guides.bIsSimulationEnable = Instance->HairGroupPublicData->IsSimulationEnable(IntLODIndex);
 				Instance->Guides.bHasGlobalInterpolation = Instance->HairGroupPublicData->IsGlobalInterpolationEnable(IntLODIndex);
 				Instance->Guides.bIsDeformationEnable = Instance->HairGroupPublicData->bIsDeformationEnable;
diff --git a/Engine/Plugins/Runtime/HairStrands/Source/HairStrandsCore/Private/HairStrandsInterpolation.cpp b/Engine/Plugins/Runtime/HairStrands/Source/HairStrandsCore/Private/HairStrandsInterpolation.cpp
--- a/Engine/Plugins/Runtime/HairStrands/Source/HairStrandsCore/Private/HairStrandsInterpolation.cpp
+++ b/Engine/Plugins/Runtime/HairStrands/Source/HairStrandsCore/Private/HairStrandsInterpolation.cpp
@@ -1061,7 +1061,7 @@ static void AddHairCardsDeformationPass(
 		DispatchCount);
 
 	// If LOD has switched, copy the current buffer, so that we don't get incorrect motion vector
-	if (bHasLODSwitch)
+	//if (bHasLODSwitch)
 	{
 		AddCopyBufferPass(GraphBuilder, CardsDeformedPositionBuffer_Prev.Buffer, CardsDeformedPositionBuffer_Curr.Buffer);
 	}
diff --git a/Engine/Shaders/Private/HairStrands/HairStrandsDebug.usf b/Engine/Shaders/Private/HairStrands/HairStrandsDebug.usf
--- a/Engine/Shaders/Private/HairStrands/HairStrandsDebug.usf
+++ b/Engine/Shaders/Private/HairStrands/HairStrandsDebug.usf
@@ -1199,6 +1199,7 @@ struct FInstanceInfo
 	bool StableRaster;
 	bool CacheRen;
 	bool CacheSim;
+	bool bHasLODSwitched;
 
 	// Binding
 	uint UniqueTriangle;
@@ -1212,46 +1213,52 @@ struct FInstanceInfo
 
 FInstanceInfo ReadInstanceInfo(uint Index)
 {
-	const uint Index16 = Index * 16;
+	const uint Index20 = Index * 20u;
 
 	uint4 Data0 = 0;
-	Data0.x = Infos[Index16 + 0];
-	Data0.y = Infos[Index16 + 1];
-	Data0.z = Infos[Index16 + 2];
-	Data0.w = Infos[Index16 + 3];
+	Data0.x = Infos[Index20 + 0];
+	Data0.y = Infos[Index20 + 1];
+	Data0.z = Infos[Index20 + 2];
+	Data0.w = Infos[Index20 + 3];
 
 	uint4 Data1 = 0;
-	Data1.x = Infos[Index16 + 4];
-	Data1.y = Infos[Index16 + 5];
-	Data1.z = Infos[Index16 + 6];
-	Data1.w = Infos[Index16 + 7];
+	Data1.x = Infos[Index20 + 4];
+	Data1.y = Infos[Index20 + 5];
+	Data1.z = Infos[Index20 + 6];
+	Data1.w = Infos[Index20 + 7];
 
 	uint4 Data2 = 0;
-	Data2.x = Infos[Index16 + 8];
-	Data2.y = Infos[Index16 + 9];
-	Data2.z = Infos[Index16 + 10];
-	Data2.w = Infos[Index16 + 11];
+	Data2.x = Infos[Index20 + 8];
+	Data2.y = Infos[Index20 + 9];
+	Data2.z = Infos[Index20 + 10];
+	Data2.w = Infos[Index20 + 11];
 
 	uint4 Data3 = 0;
-	Data3.x = Infos[Index16 + 12];
-	Data3.y = Infos[Index16 + 13];
-	Data3.z = Infos[Index16 + 14];
-	Data3.w = Infos[Index16 + 15];
+	Data3.x = Infos[Index20 + 12];
+	Data3.y = Infos[Index20 + 13];
+	Data3.z = Infos[Index20 + 14];
+	Data3.w = Infos[Index20 + 15];
+
+	uint4 Data4 = 0;
+	Data4.x = Infos[Index20 + 16];
+	Data4.y = Infos[Index20 + 17];
+	Data4.z = Infos[Index20 + 18];
+	Data4.w = Infos[Index20 + 19];
 
 	FInstanceInfo Out = (FInstanceInfo)0;
-	Out.GroupIndex = (Data0.x >> 0) & 0xFF;
-	Out.GroupCount = (Data0.x >> 8) & 0xFF;
-	Out.LODCount = (Data0.x >> 16) & 0xFF;
-	Out.GeometryType = (Data0.x >> 24) & 0x7;
-	Out.BindingType = (Data0.x >> 27) & 0x7;
-	Out.bHasSim = (Data0.x >> 30) & 0x1;
-	Out.bHasRBF = (Data0.x >> 31) & 0x1;
+	Out.GroupIndex 		= (Data0.x >> 0) & 0xFF;
+	Out.GroupCount 		= (Data0.x >> 8) & 0xFF;
+	Out.LODCount 		= (Data0.x >> 16) & 0xFF;
+	Out.GeometryType 	= (Data0.x >> 24) & 0x7;
+	Out.BindingType 	= (Data0.x >> 27) & 0x7;
+	Out.bHasSim 		= (Data0.x >> 30) & 0x1;
+	Out.bHasRBF 		= (Data0.x >> 31) & 0x1;
 
-	Out.LOD = f16tof32((Data0.y >> 0) & 0xFFFF);
-	Out.ClipValue = f16tof32((Data0.y >> 16) & 0xFFFF);
+	Out.LOD 			= f16tof32((Data0.y >> 0) & 0xFFFF);
+	Out.ClipValue 		= f16tof32((Data0.y >> 16) & 0xFFFF);
 
-	Out.StrandsCount = Data0.z;
-	Out.VertexCount = Data0.w;
+	Out.StrandsCount 	= Data0.z;
+	Out.VertexCount 	= Data0.w;
 
 	// Strands binding data
 	Out.UniqueSection		= Data1.x;
@@ -1279,6 +1286,8 @@ FInstanceInfo ReadInstanceInfo(uint Index)
 	Out.Length					= f16tof32(Data3.z & 0xFFFF);
 	Out.LengthScale				= f16tof32(Data3.z >> 16u);
 
+	Out.bHasLODSwitched			= Data4.x > 0;
+
 	return Out;
 }
 
@@ -1413,7 +1422,7 @@ void MainCS(uint GroupIndex : SV_GroupIndex, uint3 DispatchThreadId : SV_Dispatc
 	Newline(Context);
 
 	// Draw instance details
-	Print(Context, TEXT("   Group  LOD     Geometry Binding Cache Sim RBF Clip Screen Strands            Points             Bind.Tri/Sect Vis. Shad. Name"));
+	Print(Context, TEXT("   Group  LOD     Geometry Binding Cache Sim RBF Clip Screen Strands            Points             Bind.Tri/Sect Vis. Shad. SwLOD Name"));
 	Newline(Context);
 
 	for (uint InstanceIndex = 0; InstanceIndex < InstanceCount; ++InstanceIndex)
@@ -1498,6 +1507,10 @@ void MainCS(uint GroupIndex : SV_GroupIndex, uint3 DispatchThreadId : SV_Dispatc
 		PrintBool(Context, Info.bIsVisibleInShadowView);
 		Print(Context, TEXT("  "));
 
+		// Has LOD switched
+		PrintBool(Context, Info.bHasLODSwitched);
+		Print(Context, TEXT("  "));
+
 		// Name
 		PrintInstanceName(Context, InstanceIndex, FontWhite);
 
diff --git a/Engine/Source/Runtime/Renderer/Public/HairStrandsInterface.h b/Engine/Source/Runtime/Renderer/Public/HairStrandsInterface.h
--- a/Engine/Source/Runtime/Renderer/Public/HairStrandsInterface.h
+++ b/Engine/Source/Runtime/Renderer/Public/HairStrandsInterface.h
@@ -356,6 +356,7 @@ public:
 		} Meshes;
 
 		bool bHasLODSwitch = false;
+		bool bHasLODSwitchBindingType = false;
 		EHairGeometryType GeometryType = EHairGeometryType::NoneGeometry;
 		EHairBindingType BindingType = EHairBindingType::NoneBinding;
 		FTransform LocalToWorldTransform;
