* Changed groom rigging to be editable in groom editor
* Introduce guide type to ease understanding/manipulation of the different type of guides

diff --git a/Engine/Plugins/Runtime/HairStrands/Source/HairStrandsCore/Private/GroomAsset.cpp b/Engine/Plugins/Runtime/HairStrands/Source/HairStrandsCore/Private/GroomAsset.cpp
--- a/Engine/Plugins/Runtime/HairStrands/Source/HairStrandsCore/Private/GroomAsset.cpp
+++ b/Engine/Plugins/Runtime/HairStrands/Source/HairStrandsCore/Private/GroomAsset.cpp
@@ -1233,6 +1233,31 @@ void UGroomAsset::PostLoad()
 		}
 	}
 #endif
+	// Convert old rigging/guide override
+	for (FHairGroupsInterpolation& Group : HairGroupsInterpolation)
+	{
+		const bool bNeedLegacyConversion = Group.InterpolationSettings.bOverrideGuides_DEPRECATED || Group.RiggingSettings.bEnableRigging_DEPRECATED;
+		if (bNeedLegacyConversion)
+		{				
+			if (Group.RiggingSettings.bEnableRigging_DEPRECATED && Group.InterpolationSettings.bOverrideGuides_DEPRECATED)
+			{
+				Group.InterpolationSettings.GuideType = EGroomGuideType::Rigged;
+			}
+			else if (Group.InterpolationSettings.bOverrideGuides_DEPRECATED)
+			{
+				Group.InterpolationSettings.GuideType = EGroomGuideType::Generated;
+			}
+
+			Group.InterpolationSettings.RiggedGuideNumCurves = Group.RiggingSettings.NumCurves_DEPRECATED;
+			Group.InterpolationSettings.RiggedGuideNumPoints = Group.RiggingSettings.NumPoints_DEPRECATED;
+
+			// Reset deprecated values, so that if the asset is saved, the legacy conversion will no longer occur
+			Group.InterpolationSettings.bOverrideGuides_DEPRECATED = false;
+			Group.RiggingSettings.bEnableRigging_DEPRECATED = false;
+			Group.RiggingSettings.NumCurves_DEPRECATED = 0;
+			Group.RiggingSettings.NumPoints_DEPRECATED = 0;
+		}
+	}
 
 #if WITH_EDITORONLY_DATA
 	bool bSucceed = true;
@@ -1511,9 +1536,6 @@ static bool IsStrandsInterpolationAttributes(const FName PropertyName)
 		   PropertyName == GET_MEMBER_NAME_CHECKED(FHairDecimationSettings, CurveDecimation)
 		|| PropertyName == GET_MEMBER_NAME_CHECKED(FHairDecimationSettings, VertexDecimation)
 
-		|| PropertyName == GET_MEMBER_NAME_CHECKED(FHairDeformationSettings, NumCurves)
-		|| PropertyName == GET_MEMBER_NAME_CHECKED(FHairDeformationSettings, NumPoints)
-		|| PropertyName == GET_MEMBER_NAME_CHECKED(FHairDeformationSettings, bEnableRigging)
 
 		// Add dependency on simulation and per LOD-simulation/global-interpolation to strip-out interpolation data if there are not needed
 		|| PropertyName == GET_MEMBER_NAME_CHECKED(UGroomAsset, EnableGlobalInterpolation)
@@ -1530,8 +1552,10 @@ static bool IsStrandsInterpolationAttributes(const FName PropertyName)
 		|| PropertyName == GET_MEMBER_NAME_CHECKED(FHairLODSettings, Simulation)
 		|| PropertyName == GET_MEMBER_NAME_CHECKED(FHairLODSettings, GlobalInterpolation)
 
-		|| PropertyName == GET_MEMBER_NAME_CHECKED(FHairInterpolationSettings, bOverrideGuides)
+		|| PropertyName == GET_MEMBER_NAME_CHECKED(FHairInterpolationSettings, GuideType)
 		|| PropertyName == GET_MEMBER_NAME_CHECKED(FHairInterpolationSettings, HairToGuideDensity)
+		|| PropertyName == GET_MEMBER_NAME_CHECKED(FHairInterpolationSettings, RiggedGuideNumCurves)
+		|| PropertyName == GET_MEMBER_NAME_CHECKED(FHairInterpolationSettings, RiggedGuideNumPoints)
 		|| PropertyName == GET_MEMBER_NAME_CHECKED(FHairInterpolationSettings, InterpolationQuality)
 		|| PropertyName == GET_MEMBER_NAME_CHECKED(FHairInterpolationSettings, InterpolationDistance)
 		|| PropertyName == GET_MEMBER_NAME_CHECKED(FHairInterpolationSettings, bRandomizeGuide)
@@ -1562,6 +1586,14 @@ static bool IsStrandsLODAttributes(const FName PropertyName)
 		|| PropertyName == GET_MEMBER_NAME_CHECKED(FHairLODSettings, GlobalInterpolation);
 }
 
+static bool IsRiggingSettings(const FName PropertyName)
+{
+	return
+		PropertyName == GET_MEMBER_NAME_CHECKED(FHairInterpolationSettings, GuideType) ||
+		PropertyName == GET_MEMBER_NAME_CHECKED(FHairInterpolationSettings, RiggedGuideNumCurves) ||
+		PropertyName == GET_MEMBER_NAME_CHECKED(FHairInterpolationSettings, RiggedGuideNumPoints);
+}
+
 static bool IsCardsAttributes(const FName PropertyName)
 {
 	return PropertyName == GET_MEMBER_NAME_CHECKED(UGroomAsset, HairGroupsCards);
@@ -1641,6 +1673,21 @@ void UGroomAsset::PostEditChangeProperty(FPropertyChangedEvent& PropertyChangedE
 		InitCardsTextureResources(this);
 	}
 
+	// If rigging options changed, recreate skeletal mesh asset
+	const bool bRiggingChanged = IsRiggingSettings(PropertyName);
+	if (bRiggingChanged)
+	{
+		bool bUpdateRiggedMesh = false;
+		for (const FHairGroupsInterpolation& Interpolation : HairGroupsInterpolation)
+		{
+			bUpdateRiggedMesh |= Interpolation.InterpolationSettings.GuideType == EGroomGuideType::Rigged;
+		}
+		if (bUpdateRiggedMesh)
+		{
+			RiggedSkeletalMesh = FGroomDeformerBuilder::CreateSkeletalMesh(this);
+		}
+	}
+
 	const bool bCardMaterialChanged = PropertyName == GET_MEMBER_NAME_CHECKED(FHairGroupsCardsSourceDescription, Material);
 	const bool bMeshMaterialChanged = PropertyName == GET_MEMBER_NAME_CHECKED(FHairGroupsMeshesSourceDescription, Material);
 
@@ -3005,7 +3052,7 @@ bool UGroomAsset::BuildCardsData(uint32 GroupIndex)
 				// Create own interpolation settings for cards.
 				// Force closest guides as this is the most relevant matching metric for cards, due to their coarse geometry
 				FHairInterpolationSettings CardsInterpolationSettings = HairGroupsInterpolation[GroupIndex].InterpolationSettings;
-				CardsInterpolationSettings.bOverrideGuides = false;
+				CardsInterpolationSettings.GuideType = EGroomGuideType::Imported;
 				CardsInterpolationSettings.bUseUniqueGuide = true;
 				CardsInterpolationSettings.bRandomizeGuide = false;
 				CardsInterpolationSettings.InterpolationDistance = EHairInterpolationWeight::Parametric;
@@ -3570,7 +3617,9 @@ bool UGroomAsset::IsVisible(int32 GroupIndex, int32 LODIndex) const
 
 bool UGroomAsset::IsDeformationEnable(int32 GroupIndex) const
 {
-	return HairGroupsInterpolation.IsValidIndex(GroupIndex) && HairGroupsInterpolation[GroupIndex].RiggingSettings.bEnableRigging && HairGroupsInterpolation[GroupIndex].InterpolationSettings.bOverrideGuides;
+	return 
+		HairGroupsInterpolation.IsValidIndex(GroupIndex) && 
+		HairGroupsInterpolation[GroupIndex].InterpolationSettings.GuideType == EGroomGuideType::Rigged;
 }
 
 bool UGroomAsset::IsSimulationEnable(int32 GroupIndex, int32 LODIndex) const 
diff --git a/Engine/Plugins/Runtime/HairStrands/Source/HairStrandsCore/Private/GroomAssetInterpolation.cpp b/Engine/Plugins/Runtime/HairStrands/Source/HairStrandsCore/Private/GroomAssetInterpolation.cpp
--- a/Engine/Plugins/Runtime/HairStrands/Source/HairStrandsCore/Private/GroomAssetInterpolation.cpp
+++ b/Engine/Plugins/Runtime/HairStrands/Source/HairStrandsCore/Private/GroomAssetInterpolation.cpp
@@ -12,20 +12,22 @@ FHairDecimationSettings::FHairDecimationSettings()
 
 FHairInterpolationSettings::FHairInterpolationSettings()
 {
-	bOverrideGuides = false;
+	bOverrideGuides_DEPRECATED = false;
 	HairToGuideDensity = 0.1f;
 	InterpolationQuality = EHairInterpolationQuality::High;
 	InterpolationDistance = EHairInterpolationWeight::Parametric;
 	bRandomizeGuide = false;
 	bUseUniqueGuide = false;
+	GuideType = EGroomGuideType::Imported;
+	RiggedGuideNumCurves = 10;
+	RiggedGuideNumPoints = 4;
 }
 
 FHairDeformationSettings::FHairDeformationSettings()
 {
-	bCanEditRigging = false;
-	bEnableRigging = false;
-	NumCurves = 10;
-	NumPoints = 4;
+	bEnableRigging_DEPRECATED = false;
+	NumCurves_DEPRECATED = 10;
+	NumPoints_DEPRECATED = 4;
 }
 
 FHairGroupsInterpolation::FHairGroupsInterpolation()
@@ -82,12 +84,14 @@ bool FHairDecimationSettings::operator==(const FHairDecimationSettings& A) const
 bool FHairInterpolationSettings::operator==(const FHairInterpolationSettings& A) const
 {
 	return
-		bOverrideGuides == A.bOverrideGuides &&
+		GuideType == A.GuideType &&
 		HairToGuideDensity == A.HairToGuideDensity &&
 		InterpolationQuality == A.InterpolationQuality &&
 		InterpolationDistance == A.InterpolationDistance &&
 		bRandomizeGuide == A.bRandomizeGuide &&
-		bUseUniqueGuide == A.bUseUniqueGuide;
+		bUseUniqueGuide == A.bUseUniqueGuide &&
+		RiggedGuideNumCurves == A.RiggedGuideNumCurves &&
+		RiggedGuideNumPoints == A.RiggedGuideNumPoints;
 }
 
 bool FHairGroupsInterpolation::operator==(const FHairGroupsInterpolation& A) const
@@ -97,29 +101,18 @@ bool FHairGroupsInterpolation::operator==(const FHairGroupsInterpolation& A) con
 		InterpolationSettings == A.InterpolationSettings;
 }
 
-bool FHairDeformationSettings::operator==(const FHairDeformationSettings& A) const
-{
-	return
-		bEnableRigging == A.bEnableRigging &&
-		NumCurves == A.NumCurves &&
-		NumPoints == A.NumPoints &&
-		bCanEditRigging == A.bCanEditRigging;
-}
-
 void FHairGroupsInterpolation::BuildDDCKey(FArchive& Ar)
 {
 	Ar << DecimationSettings.CurveDecimation;
 	Ar << DecimationSettings.VertexDecimation;
-	Ar << InterpolationSettings.bOverrideGuides;
+	Ar << InterpolationSettings.GuideType;
 	Ar << InterpolationSettings.HairToGuideDensity;
 	Ar << InterpolationSettings.InterpolationQuality;
 	Ar << InterpolationSettings.InterpolationDistance;
 	Ar << InterpolationSettings.bRandomizeGuide;
 	Ar << InterpolationSettings.bUseUniqueGuide;
-	Ar << RiggingSettings.NumCurves;
-	Ar << RiggingSettings.NumPoints;
-	Ar << RiggingSettings.bEnableRigging;
-	Ar << RiggingSettings.bCanEditRigging;
+	Ar << InterpolationSettings.RiggedGuideNumCurves;
+	Ar << InterpolationSettings.RiggedGuideNumPoints;
 }
 
 void FHairGroupsLOD::BuildDDCKey(FArchive& Ar)
diff --git a/Engine/Plugins/Runtime/HairStrands/Source/HairStrandsCore/Private/GroomBuilder.cpp b/Engine/Plugins/Runtime/HairStrands/Source/HairStrandsCore/Private/GroomBuilder.cpp
--- a/Engine/Plugins/Runtime/HairStrands/Source/HairStrandsCore/Private/GroomBuilder.cpp
+++ b/Engine/Plugins/Runtime/HairStrands/Source/HairStrandsCore/Private/GroomBuilder.cpp
@@ -2064,32 +2064,37 @@ void FGroomBuilder::BuildData(
 		// Simulation data
 		{
 			OutSim = InHairDescriptionGroup.Guides;
-			if (InHairDescriptionGroup.Info.NumGuides > 0 && !InSettings.InterpolationSettings.bOverrideGuides)
+			const bool bHasImportedGuides = InHairDescriptionGroup.Info.NumGuides > 0;
+			if (bHasImportedGuides && InSettings.InterpolationSettings.GuideType == EGroomGuideType::Imported)
 			{
 				HairStrandsBuilder::BuildInternalData(OutSim);
 			}
 			else
 			{
 				OutSim.Reset();
-				if(InSettings.RiggingSettings.bEnableRigging)
+				if(InSettings.InterpolationSettings.GuideType == EGroomGuideType::Rigged)
 				{
 					if (InHairDescriptionGroup.Info.NumGuides > 0)
 					{
-						// We pick the new guides among the imported ones
+						// Pick the new guides among the imported ones
 						FHairStrandsDatas TempSim = InHairDescriptionGroup.Guides;
 						HairStrandsBuilder::BuildInternalData(TempSim);
-						FHairStrandsDecimation::Decimate(TempSim, InSettings.RiggingSettings.NumCurves, InSettings.RiggingSettings.NumPoints, OutSim);
+						FHairStrandsDecimation::Decimate(TempSim, InSettings.InterpolationSettings.RiggedGuideNumCurves, InSettings.InterpolationSettings.RiggedGuideNumPoints, OutSim);
 					}
 					else
 					{
 						// Otherwise let s pick the guides among the rendered strands
-						FHairStrandsDecimation::Decimate(OutRen, InSettings.RiggingSettings.NumCurves, InSettings.RiggingSettings.NumPoints, OutSim);
+						FHairStrandsDecimation::Decimate(OutRen, InSettings.InterpolationSettings.RiggedGuideNumCurves, InSettings.InterpolationSettings.RiggedGuideNumPoints, OutSim);
 					}
 				}
-				else
+				else if (!bHasImportedGuides || InSettings.InterpolationSettings.GuideType == EGroomGuideType::Generated)
 				{
 					FHairStrandsDecimation::Decimate(OutRen, HairToGuideDensity, 1, false, OutSim);
 				}
+				else
+				{
+					checkNoEntry();
+				}
 			}
 
 			OutGroupInfo.NumGuides			= OutSim.GetNumCurves();
@@ -2135,7 +2140,7 @@ void FGroomBuilder::BuildInterplationData(
 	{
 		// If there's usable closest guides and guide weights attributes, fill them into the asset
 		// This step requires the HairSimulationData (guides) to be filled prior to this
-		const bool bUsePrecomputedWeights = !InInterpolationSettings.bOverrideGuides && InRenData.StrandsCurves.CurvesClosestGuideIDs.Num() > 0 && InRenData.StrandsCurves.CurvesClosestGuideWeights.Num() > 0;
+		const bool bUsePrecomputedWeights = InInterpolationSettings.GuideType == EGroomGuideType::Imported && InRenData.StrandsCurves.CurvesClosestGuideIDs.Num() > 0 && InRenData.StrandsCurves.CurvesClosestGuideWeights.Num() > 0;
 		if (bUsePrecomputedWeights)
 		{
 			FHairStrandsInterpolationDatas OutInterpolation;
diff --git a/Engine/Plugins/Runtime/HairStrands/Source/HairStrandsCore/Private/GroomRBFDeformer.cpp b/Engine/Plugins/Runtime/HairStrands/Source/HairStrandsCore/Private/GroomRBFDeformer.cpp
--- a/Engine/Plugins/Runtime/HairStrands/Source/HairStrandsCore/Private/GroomRBFDeformer.cpp
+++ b/Engine/Plugins/Runtime/HairStrands/Source/HairStrandsCore/Private/GroomRBFDeformer.cpp
@@ -651,7 +651,7 @@ void FGroomRBFDeformer::GetRBFDeformedGroomAsset(const UGroomAsset* InGroomAsset
 
 			// Get deformed guides
 			// If the groom override the value, we output dummy value for the guides, since they won't be used
-			if (InGroomAsset->HairGroupsInterpolation[GroupIndex].InterpolationSettings.bOverrideGuides)
+			if (InGroomAsset->HairGroupsInterpolation[GroupIndex].InterpolationSettings.GuideType != EGroomGuideType::Imported)
 			{
 				const uint32 OriginalVertexCount = OriginalGuides.StrandsPoints.Num();
 				DeformedPositions[GroupIndex].GuideStrands.Init(FVector3f::ZeroVector, OriginalVertexCount);
diff --git a/Engine/Plugins/Runtime/HairStrands/Source/HairStrandsCore/Public/GroomAssetInterpolation.h b/Engine/Plugins/Runtime/HairStrands/Source/HairStrandsCore/Public/GroomAssetInterpolation.h
--- a/Engine/Plugins/Runtime/HairStrands/Source/HairStrandsCore/Public/GroomAssetInterpolation.h
+++ b/Engine/Plugins/Runtime/HairStrands/Source/HairStrandsCore/Public/GroomAssetInterpolation.h
@@ -50,6 +50,15 @@ enum class EGroomOverrideType : uint8
 	Disable UMETA(DisplayName = "Disable", ToolTip = "Override the asset value, and force disabled")
 };
 
+
+UENUM(BlueprintType)
+enum class EGroomGuideType : uint8
+{
+	Imported	UMETA(DisplayName = "Imported Guides", ToolTip = "Use imported asset guides."),
+	Generated	UMETA(DisplayName = "Generated Guides", ToolTip = "Generate guides from imported strands"),
+	Rigged 		UMETA(DisplayName = "Rigged  Guides", ToolTip = "Generated rigged guides from imported strands")
+};
+
 USTRUCT(BlueprintType)
 struct HAIRSTRANDSCORE_API FHairLODSettings
 {
@@ -121,13 +130,25 @@ struct HAIRSTRANDSCORE_API FHairInterpolationSettings
 	FHairInterpolationSettings();
 
 	/** Flag to override the imported guides with generated guides. */
-	UPROPERTY(EditAnywhere, BlueprintReadWrite, Category = "InterpolationSettings", meta = (ToolTip = "If checked, override imported guides with generated ones."))
-	bool bOverrideGuides;
+	UPROPERTY(EditAnywhere, BlueprintReadWrite, Category = "InterpolationSettings", meta = (ToolTip = "Type of guides:\n - Imported: use imported guides\n - Generated: generate guides from strands\n - Rigged: generated rigged guides from strands."))
+	EGroomGuideType GuideType;
+
+	/** Flag to override the imported guides with generated guides. DEPRECATED */
+	UPROPERTY()
+	bool bOverrideGuides_DEPRECATED;
 
 	/** Density factor for converting hair into guide curve if no guides are provided. The value should be between 0 and 1, and can be thought as a ratio/percentage of strands used as guides.*/
-	UPROPERTY(EditAnywhere, BlueprintReadWrite, Category = "InterpolationSettings", meta = (ClampMin = "0", ClampMax = "1.0", UIMin = "0", UIMax = "1.0"))
+	UPROPERTY(EditAnywhere, BlueprintReadWrite, Category = "InterpolationSettings", meta = (DisplayName = "Generated guide density", ClampMin = "0", ClampMax = "1.0", UIMin = "0", UIMax = "1.0", EditCondition="GuideType == EGroomGuideType::Generated"))
 	float HairToGuideDensity;
 
+	/** Number of curves to generate in the skel mesh */
+	UPROPERTY(EditAnywhere, BlueprintReadOnly, Category = "InterpolationSettings", meta = (DisplayName = "Rigged guide num. curves", ToolTip = "Number of guides that will be generated on the groom and the skeletal mesh", EditCondition="GuideType == EGroomGuideType::Rigged"))
+	int32 RiggedGuideNumCurves;
+
+	/** Number of points per curve */
+	UPROPERTY(EditAnywhere, BlueprintReadOnly, Category = "InterpolationSettings", meta = (DisplayName = "Rigged guide num. points", ToolTip = "Number of points/bones per generated guide", EditCondition="GuideType == EGroomGuideType::Rigged"))
+	int32 RiggedGuideNumPoints;
+
 	/** Interpolation data quality. */
 	UPROPERTY(EditAnywhere, BlueprintReadWrite, Category = "InterpolationSettings")
 	EHairInterpolationQuality InterpolationQuality;
@@ -153,21 +174,15 @@ struct HAIRSTRANDSCORE_API FHairDeformationSettings
 	GENERATED_USTRUCT_BODY()
 
 	FHairDeformationSettings();
-	
-	UPROPERTY(VisibleAnywhere, Category = "SkeletonGeneration")
-	bool bCanEditRigging;
 
-	/** Enable the generation of a skeletal mesh that will drive the guides deformation */
-	UPROPERTY(EditAnywhere, BlueprintReadWrite, Category = "SkeletonGeneration", meta = (ToolTip = "Enable the generation of a skeletal mesh that will drive the guides deformation. The total number of bones per group is for now 256 (NumCurves * NumPoints)", EditCondition = "bCanEditRigging"))
-	bool bEnableRigging;
+	UPROPERTY()
+	bool bEnableRigging_DEPRECATED;
 
-	/** Number of curves to generate in the skel mesh */
-	UPROPERTY(EditAnywhere, BlueprintReadWrite, Category = "SkeletonGeneration", meta = (ToolTip = "Number of guides that will be generated on the groom and the skeletal mesh", EditCondition = "bCanEditRigging"))
-	int32 NumCurves;
+	UPROPERTY()
+	int32 NumCurves_DEPRECATED;
 
-	/** Number of points per curve */
-	UPROPERTY(EditAnywhere, BlueprintReadWrite, Category = "SkeletonGeneration", meta = (ToolTip = "Number of points/bones per generated guide", EditCondition = "bCanEditRigging"))
-	int32 NumPoints;
+	UPROPERTY()
+	int32 NumPoints_DEPRECATED;
 
 	bool operator==(const FHairDeformationSettings& A) const;
 };
@@ -185,7 +200,8 @@ struct HAIRSTRANDSCORE_API FHairGroupsInterpolation
 	UPROPERTY(EditAnywhere, Category = "InterpolationSettings", meta = (ToolTip = "Interpolation settings"))
 	FHairInterpolationSettings InterpolationSettings;
 
-	UPROPERTY(EditAnywhere, Category = "SkeletonGeneration", meta = (DisplayName = "Skeleton Generation", ToolTip = "Skeleton Generation"))
+	// DEPRECATED
+	UPROPERTY()
 	FHairDeformationSettings RiggingSettings;
 
 	bool operator==(const FHairGroupsInterpolation& A) const;
diff --git a/Engine/Plugins/Runtime/HairStrands/Source/HairStrandsEditor/Private/GroomActions.cpp b/Engine/Plugins/Runtime/HairStrands/Source/HairStrandsEditor/Private/GroomActions.cpp
--- a/Engine/Plugins/Runtime/HairStrands/Source/HairStrandsEditor/Private/GroomActions.cpp
+++ b/Engine/Plugins/Runtime/HairStrands/Source/HairStrandsEditor/Private/GroomActions.cpp
@@ -121,7 +121,6 @@ void ExecuteRebuild(const FToolMenuContext& InContext)
 						OutGroup.Attributes 	= HairDescriptionGroups.HairGroups[GroupIndex].GetHairAttributes();
 						OutGroup.AttributeFlags = HairDescriptionGroups.HairGroups[GroupIndex].GetHairAttributeFlags();
 						OutGroup.InterpolationSettings = GroomAsset->HairGroupsInterpolation[GroupIndex];
-						OutGroup.InterpolationSettings.RiggingSettings.bCanEditRigging = true;
 					}
 				}
 				TSharedPtr<SGroomImportOptionsWindow> GroomOptionWindow = SGroomImportOptionsWindow::DisplayRebuildOptions(CurrentOptions, GroupsPreview, Filename);
@@ -136,9 +135,7 @@ void ExecuteRebuild(const FToolMenuContext& InContext)
 				for (uint32 GroupIndex = 0; GroupIndex < GroupCount; ++GroupIndex)
 				{
 					GroomAsset->HairGroupsInterpolation[GroupIndex] = GroupsPreview->Groups[GroupIndex].InterpolationSettings;
-					GroomAsset->HairGroupsInterpolation[GroupIndex].RiggingSettings.bCanEditRigging = false;
-					bEnableRigging |= GroomAsset->HairGroupsInterpolation[GroupIndex].RiggingSettings.bEnableRigging &&
-						GroomAsset->HairGroupsInterpolation[GroupIndex].InterpolationSettings.bOverrideGuides;
+					bEnableRigging |= GroomAsset->HairGroupsInterpolation[GroupIndex].InterpolationSettings.GuideType == EGroomGuideType::Rigged;
 				}
 
 				bool bSucceeded = GroomAsset->CacheDerivedDatas();
diff --git a/Engine/Plugins/Runtime/HairStrands/Source/HairStrandsEditor/Private/GroomAssetDetails.cpp b/Engine/Plugins/Runtime/HairStrands/Source/HairStrandsEditor/Private/GroomAssetDetails.cpp
--- a/Engine/Plugins/Runtime/HairStrands/Source/HairStrandsEditor/Private/GroomAssetDetails.cpp
+++ b/Engine/Plugins/Runtime/HairStrands/Source/HairStrandsEditor/Private/GroomAssetDetails.cpp
@@ -910,17 +910,12 @@ bool FGroomRenderingDetails::CommonResetToDefault(TSharedPtr<IPropertyHandle> Ch
 
 		{
 			FHairInterpolationSettings Default;
-			HAIR_RESET1(HairGroupsInterpolation, FHairInterpolationSettings, InterpolationSettings, bOverrideGuides);
+			HAIR_RESET1(HairGroupsInterpolation, FHairInterpolationSettings, InterpolationSettings, GuideType);
 			HAIR_RESET1(HairGroupsInterpolation, FHairInterpolationSettings, InterpolationSettings, HairToGuideDensity);
 			HAIR_RESET1(HairGroupsInterpolation, FHairInterpolationSettings, InterpolationSettings, bRandomizeGuide);
 			HAIR_RESET1(HairGroupsInterpolation, FHairInterpolationSettings, InterpolationSettings, bUseUniqueGuide);
-		}
-
-		{
-			FHairDeformationSettings Default;
-			HAIR_RESET1(HairGroupsInterpolation, FHairDeformationSettings, RiggingSettings, NumCurves);
-			HAIR_RESET1(HairGroupsInterpolation, FHairDeformationSettings, RiggingSettings, NumPoints);
-			HAIR_RESET1(HairGroupsInterpolation, FHairDeformationSettings, RiggingSettings, bEnableRigging);
+			HAIR_RESET1(HairGroupsInterpolation, FHairInterpolationSettings, InterpolationSettings, RiggedGuideNumCurves);
+			HAIR_RESET1(HairGroupsInterpolation, FHairInterpolationSettings, InterpolationSettings, RiggedGuideNumPoints);
 		}
 	}
 
