Deprecate direct access to groom binding members.

diff --git a/Engine/Plugins/Importers/USDImporter/Source/USDClasses/Private/USDClassesModule.cpp b/Engine/Plugins/Importers/USDImporter/Source/USDClasses/Private/USDClassesModule.cpp
--- a/Engine/Plugins/Importers/USDImporter/Source/USDClasses/Private/USDClassesModule.cpp
+++ b/Engine/Plugins/Importers/USDImporter/Source/USDClasses/Private/USDClassesModule.cpp
@@ -356,11 +356,11 @@ TSet<UObject*> IUsdClassesModule::GetAssetDependencies(UObject* Asset)
 	}
 	else if (UGroomBindingAsset* GroomBinding = Cast<UGroomBindingAsset>(Asset))
 	{
-		Result.Add(GroomBinding->Groom);
-		Result.Add(GroomBinding->TargetSkeletalMesh);
-		Result.Add(GroomBinding->SourceSkeletalMesh);
-		Result.Add(GroomBinding->SourceGeometryCache);
-		Result.Add(GroomBinding->TargetGeometryCache);
+		Result.Add(GroomBinding->GetGroom());
+		Result.Add(GroomBinding->GetTargetSkeletalMesh());
+		Result.Add(GroomBinding->GetSourceSkeletalMesh());
+		Result.Add(GroomBinding->GetSourceGeometryCache());
+		Result.Add(GroomBinding->GetTargetGeometryCache());
 	}
 	else if (UGroomAsset* GroomAsset = Cast<UGroomAsset>(Asset))
 	{
diff --git a/Engine/Plugins/Importers/USDImporter/Source/USDSchemas/Private/USDGroomTranslatorUtils.cpp b/Engine/Plugins/Importers/USDImporter/Source/USDSchemas/Private/USDGroomTranslatorUtils.cpp
--- a/Engine/Plugins/Importers/USDImporter/Source/USDSchemas/Private/USDGroomTranslatorUtils.cpp
+++ b/Engine/Plugins/Importers/USDImporter/Source/USDSchemas/Private/USDGroomTranslatorUtils.cpp
@@ -111,21 +111,21 @@ namespace UE::UsdGroomTranslatorUtils::Private
 		);
 		if (GroomBinding)
 		{
-			GroomBinding->GroomBindingType = Settings.GroomBindingType;
-			GroomBinding->Groom = Settings.Groom;
-			if (GroomBinding->GroomBindingType == EGroomBindingMeshType::SkeletalMesh)
+			GroomBinding->SetGroomBindingType(Settings.GroomBindingType);
+			GroomBinding->SetGroom(Settings.Groom);
+			if (GroomBinding->GetGroomBindingType() == EGroomBindingMeshType::SkeletalMesh)
 			{
-				GroomBinding->SourceSkeletalMesh = Cast<USkeletalMesh>(Settings.SourceMesh);
-				GroomBinding->TargetSkeletalMesh = Cast<USkeletalMesh>(Settings.TargetMesh);
+				GroomBinding->SetSourceSkeletalMesh(Cast<USkeletalMesh>(Settings.SourceMesh));
+				GroomBinding->SetTargetSkeletalMesh(Cast<USkeletalMesh>(Settings.TargetMesh));
 			}
 			else
 			{
-				GroomBinding->SourceGeometryCache = Cast<UGeometryCache>(Settings.SourceMesh);
-				GroomBinding->TargetGeometryCache = Cast<UGeometryCache>(Settings.TargetMesh);
+				GroomBinding->SetSourceGeometryCache(Cast<UGeometryCache>(Settings.SourceMesh));
+				GroomBinding->SetTargetGeometryCache(Cast<UGeometryCache>(Settings.TargetMesh));
 			}
-			GroomBinding->HairGroupsPlatformData.Reserve(Settings.Groom->HairGroupsPlatformData.Num());
-			GroomBinding->NumInterpolationPoints = Settings.NumInterpolationPoints;
-			GroomBinding->MatchingSection = Settings.MatchingSection;
+			GroomBinding->GetHairGroupsPlatformData().Reserve(Settings.Groom->HairGroupsPlatformData.Num());
+			GroomBinding->SetNumInterpolationPoints(Settings.NumInterpolationPoints);
+			GroomBinding->SetMatchingSection(Settings.MatchingSection);
 
 			GroomBinding->Build();
 		}
diff --git a/Engine/Plugins/Importers/USDImporter/Source/USDStageImporter/Private/USDStageImporter.cpp b/Engine/Plugins/Importers/USDImporter/Source/USDStageImporter/Private/USDStageImporter.cpp
--- a/Engine/Plugins/Importers/USDImporter/Source/USDStageImporter/Private/USDStageImporter.cpp
+++ b/Engine/Plugins/Importers/USDImporter/Source/USDStageImporter/Private/USDStageImporter.cpp
@@ -907,11 +907,11 @@ namespace UsdStageImporterImpl
 			}
 			else if (UGroomBindingAsset* GroomBinding = Cast<UGroomBindingAsset>(Asset))
 			{
-				if(GroomBinding->TargetSkeletalMesh)
+				if (GroomBinding->GetTargetSkeletalMesh())
 				{
 					AssetTypeFolderPtr = &SkeletalMeshesFolder;
 				}
-				else if (GroomBinding->TargetGeometryCache)
+				else if (GroomBinding->GetTargetGeometryCache())
 				{
 					AssetTypeFolderPtr = &GeometryCachesFolder;
 				}
diff --git a/Engine/Plugins/Runtime/HairStrands/Source/HairStrandsCore/Private/GroomAsset.cpp b/Engine/Plugins/Runtime/HairStrands/Source/HairStrandsCore/Private/GroomAsset.cpp
--- a/Engine/Plugins/Runtime/HairStrands/Source/HairStrandsCore/Private/GroomAsset.cpp
+++ b/Engine/Plugins/Runtime/HairStrands/Source/HairStrandsCore/Private/GroomAsset.cpp
@@ -281,30 +281,30 @@ void DumpLoadedGroomBindingAssets(IConsoleVariable* InCVarPakTesterEnabled)
 	{
 		if (AssetIt)
 		{			
-			const uint32 GroupCount = AssetIt->HairGroupsPlatformData.Num();
+			const uint32 GroupCount = AssetIt->GetHairGroupsPlatformData().Num();
 			for (uint32 GroupIt = 0; GroupIt < GroupCount; ++GroupIt)
 			{
-				const uint32 CPUMemorySize_Guides  = AssetIt->HairGroupsPlatformData[GroupIt].SimRootBulkData.GetDataSize();
-				const uint32 CPUMemorySize_Strands = AssetIt->HairGroupsPlatformData[GroupIt].RenRootBulkData.GetDataSize();
+				const uint32 CPUMemorySize_Guides  = AssetIt->GetHairGroupsPlatformData()[GroupIt].SimRootBulkData.GetDataSize();
+				const uint32 CPUMemorySize_Strands = AssetIt->GetHairGroupsPlatformData()[GroupIt].RenRootBulkData.GetDataSize();
 
 				uint32 GPUMemorySize_Guides = 0;
-				if (const FHairStrandsRestRootResource* RootResource = AssetIt->HairGroupResources[GroupIt].SimRootResources)
+				if (const FHairStrandsRestRootResource* RootResource = AssetIt->GetHairGroupResources()[GroupIt].SimRootResources)
 				{
 					GPUMemorySize_Guides += RootResource->GetResourcesSize();
 				}
 				uint32 GPUMemorySize_Strands = 0;
-				if (const FHairStrandsRestRootResource* RootResource = AssetIt->HairGroupResources[GroupIt].RenRootResources)
+				if (const FHairStrandsRestRootResource* RootResource = AssetIt->GetHairGroupResources()[GroupIt].RenRootResources)
 				{
 					GPUMemorySize_Strands += RootResource->GetResourcesSize();
 				}
 
 				uint32 CPUMemorySize_Cards = 0;
 				uint32 GPUMemorySize_Cards = 0;
-				const uint32 CardCount = AssetIt->HairGroupsPlatformData[GroupIt].CardsRootBulkData.Num();
+				const uint32 CardCount = AssetIt->GetHairGroupsPlatformData()[GroupIt].CardsRootBulkData.Num();
 				for (uint32 CardIt = 0; CardIt < CardCount; ++CardIt)
 				{
-					CPUMemorySize_Cards += AssetIt->HairGroupsPlatformData[GroupIt].CardsRootBulkData[CardIt].GetDataSize();
-					if (const FHairStrandsRestRootResource* RootResource = AssetIt->HairGroupResources[GroupIt].CardsRootResources[CardIt])
+					CPUMemorySize_Cards += AssetIt->GetHairGroupsPlatformData()[GroupIt].CardsRootBulkData[CardIt].GetDataSize();
+					if (const FHairStrandsRestRootResource* RootResource = AssetIt->GetHairGroupResources()[GroupIt].CardsRootResources[CardIt])
 					{
 						GPUMemorySize_Cards += RootResource->GetResourcesSize();
 					}
@@ -321,7 +321,7 @@ void DumpLoadedGroomBindingAssets(IConsoleVariable* InCVarPakTesterEnabled)
 				Total_GPUMemorySize_Strands+= GPUMemorySize_Strands;
 				Total_GPUMemorySize_Cards  += GPUMemorySize_Cards;
 
-				const uint32 SkelLODCount = AssetIt->HairGroupsPlatformData[GroupIt].RenRootBulkData.GetLODCount();
+				const uint32 SkelLODCount = AssetIt->GetHairGroupsPlatformData()[GroupIt].RenRootBulkData.GetLODCount();
 				if (bDetails)
 				{
 //					UE_LOG(LogHairStrands, Log, TEXT("--  No.  - LOD -    CPU Total (     Guides|    Strands|      Cards) -    GPU Total (     Guides|    Strands|      Cards) - Asset Name "));
diff --git a/Engine/Plugins/Runtime/HairStrands/Source/HairStrandsCore/Private/GroomBindingAsset.cpp b/Engine/Plugins/Runtime/HairStrands/Source/HairStrandsCore/Private/GroomBindingAsset.cpp
--- a/Engine/Plugins/Runtime/HairStrands/Source/HairStrandsCore/Private/GroomBindingAsset.cpp
+++ b/Engine/Plugins/Runtime/HairStrands/Source/HairStrandsCore/Private/GroomBindingAsset.cpp
@@ -135,9 +135,9 @@ void UGroomBindingAsset::Serialize(FArchive& Ar)
 {
 	uint8 Flags = 0;
 #if WITH_EDITOR
-	if (Groom)
+	if (GetGroom())
 	{
-		Flags = Groom->GenerateClassStripFlags(Ar);
+		Flags = GetGroom()->GenerateClassStripFlags(Ar);
 	}
 #endif
 
@@ -147,7 +147,7 @@ void UGroomBindingAsset::Serialize(FArchive& Ar)
 	if (Ar.IsCooking())
 #endif
 	{
-		InternalSerializePlatformDatas(Ar, this, HairGroupsPlatformData, Flags);
+		InternalSerializePlatformDatas(Ar, this, GetHairGroupsPlatformData(), Flags);
 		bIsValid = true;
 	}
 }
@@ -156,13 +156,15 @@ void UGroomBindingAsset::InitResource()
 {
 	LLM_SCOPE_BYTAG(Groom);
 
+	TRACE_CPUPROFILER_EVENT_SCOPE(UGroomBindingAsset::InitResource);
+
 	// Ensure we are releasing binding resources before reallocating them
 	ReleaseResource();
 
-	for (UGroomBindingAsset::FHairGroupPlatformData& BulkData : HairGroupsPlatformData)
+	for (UGroomBindingAsset::FHairGroupPlatformData& BulkData : GetHairGroupsPlatformData())
 	{
-		const int32 GroupIndex = HairGroupResources.Num();
-		FHairGroupResource& Resource = HairGroupResources.AddDefaulted_GetRef();
+		const int32 GroupIndex = GetHairGroupResources().Num();
+		FHairGroupResource& Resource = GetHairGroupResources().AddDefaulted_GetRef();
 
 		FHairResourceName ResourceName(GetFName(), GroupIndex);
 		const FName OwnerName = GetAssetPathName();
@@ -203,7 +205,7 @@ void UGroomBindingAsset::InitResource()
 
 void UGroomBindingAsset::UpdateResource()
 {
-	for (FHairGroupResource& Resource : HairGroupResources)
+	for (FHairGroupResource& Resource : GetHairGroupResources())
 	{
 		if (Resource.SimRootResources)
 		{
@@ -228,9 +230,9 @@ void UGroomBindingAsset::UpdateResource()
 void UGroomBindingAsset::ReleaseResource()
 {
 	// Delay destruction to insure that the rendering thread is done with all resources usage
-	if (HairGroupResources.Num() > 0)
+	if (GetHairGroupResources().Num() > 0)
 	{
-		for (FHairGroupResource& Resource : HairGroupResources)
+		for (FHairGroupResource& Resource : GetHairGroupResources())
 		{
 			FHairStrandsRestRootResource* InSimRootResources = Resource.SimRootResources;
 			FHairStrandsRestRootResource* InRenRootResources = Resource.RenRootResources;
@@ -265,12 +267,12 @@ void UGroomBindingAsset::ReleaseResource()
 				InCardsRootResources = nullptr;
 			}
 		}
-		HairGroupResources.Empty();
+		GetHairGroupResources().Empty();
 	}
 
 	// Process resources to be deleted (should happen only in editor)
 	FHairGroupResource ResourceToDelete;
-	while (HairGroupResourcesToDelete.Dequeue(ResourceToDelete))
+	while (RemoveHairGroupResourcesToDelete(ResourceToDelete))
 	{
 		FHairStrandsRestRootResource* InSimRootResources = ResourceToDelete.SimRootResources;
 		FHairStrandsRestRootResource* InRenRootResources = ResourceToDelete.RenRootResources;
@@ -295,7 +297,7 @@ void UGroomBindingAsset::ReleaseResource()
 void UGroomBindingAsset::Reset()
 {
 	ReleaseResource();
-	for (UGroomBindingAsset::FHairGroupPlatformData& Data : HairGroupsPlatformData)
+	for (UGroomBindingAsset::FHairGroupPlatformData& Data : GetHairGroupsPlatformData())
 	{
 		Data.SimRootBulkData.Reset();
 		Data.RenRootBulkData.Reset();
@@ -327,33 +329,32 @@ void UGroomBindingAsset::PostLoad()
 	LLM_SCOPE_BYTAG(Groom);
 
 	Super::PostLoad();
-
-	if (Groom)
+	if (UGroomAsset* LocalGroom = GetGroom())
 	{
 		// Make sure that the asset initialized its resources first since the component needs them to initialize its own resources
-		Groom->ConditionalPostLoad();
+		LocalGroom->ConditionalPostLoad();
 
 	#if WITH_EDITOR
 		CacheDerivedDatas();
 
 		// Sanity check. This function will report back warnings/issues into the log for user.
-		UGroomBindingAsset::IsCompatible(Groom, this, true);
+		UGroomBindingAsset::IsCompatible(LocalGroom, this, true);
 
-		if (TargetSkeletalMesh)
+		if (USkeletalMesh* LocalTargetSkeletalMesh = GetTargetSkeletalMesh())
 		{
-			TargetSkeletalMesh->OnPostMeshCached().AddUObject(this, &UGroomBindingAsset::InvalidateBinding);
+			LocalTargetSkeletalMesh->OnPostMeshCached().AddUObject(this, &UGroomBindingAsset::InvalidateBinding);
 			bRegisterTargetMeshCallback = true;
 		}
 
-		if (SourceSkeletalMesh)
+		if (USkeletalMesh* LocalSourceSkeletalMesh = GetSourceSkeletalMesh())
 		{
-			SourceSkeletalMesh->OnPostMeshCached().AddUObject(this, &UGroomBindingAsset::InvalidateBinding);
+			LocalSourceSkeletalMesh->OnPostMeshCached().AddUObject(this, &UGroomBindingAsset::InvalidateBinding);
 			bRegisterSourceMeshCallback = true;
 		}
 
-		if (Groom)
+		if (LocalGroom)
 		{
-			Groom->GetOnGroomAssetChanged().AddUObject(this, &UGroomBindingAsset::InvalidateBinding);
+			LocalGroom->GetOnGroomAssetChanged().AddUObject(this, &UGroomBindingAsset::InvalidateBinding);
 			bRegisterGroomAssetCallback = true;
 		}
 	#endif
@@ -407,21 +408,21 @@ void UGroomBindingAsset::BeginDestroy()
 	Super::BeginDestroy();
 
 #if WITH_EDITOR
-	if (TargetSkeletalMesh && bRegisterTargetMeshCallback)
+	if (GetTargetSkeletalMesh() && bRegisterTargetMeshCallback)
 	{
-		TargetSkeletalMesh->OnPostMeshCached().RemoveAll(this);
+		GetTargetSkeletalMesh()->OnPostMeshCached().RemoveAll(this);
 		bRegisterTargetMeshCallback = false;
 	}
 
-	if (SourceSkeletalMesh && bRegisterSourceMeshCallback)
+	if (GetSourceSkeletalMesh() && bRegisterSourceMeshCallback)
 	{
-		SourceSkeletalMesh->OnPostMeshCached().RemoveAll(this);
+		GetSourceSkeletalMesh()->OnPostMeshCached().RemoveAll(this);
 		bRegisterSourceMeshCallback = false;
 	}
 
-	if (Groom && bRegisterGroomAssetCallback)
+	if (GetGroom() && bRegisterGroomAssetCallback)
 	{
-		Groom->GetOnGroomAssetChanged().RemoveAll(this);
+		GetGroom()->GetOnGroomAssetChanged().RemoveAll(this);
 		bRegisterGroomAssetCallback = false;
 	}
 #endif
@@ -431,7 +432,7 @@ bool UGroomBindingAsset::IsCompatible(const USkeletalMesh* InSkeletalMesh, const
 {
 	if (InBinding && InSkeletalMesh && IsHairStrandsBindingEnable())
 	{
-		if (InBinding->GroomBindingType != EGroomBindingMeshType::SkeletalMesh)
+		if (InBinding->GetGroomBindingType() != EGroomBindingMeshType::SkeletalMesh)
 		{
 			if (bIssueWarning)
 			{
@@ -440,7 +441,7 @@ bool UGroomBindingAsset::IsCompatible(const USkeletalMesh* InSkeletalMesh, const
 			return false;
 		}
 
-		if (!InBinding->TargetSkeletalMesh)
+		if (!InBinding->GetTargetSkeletalMesh())
 		{
 			if (bIssueWarning)
 			{
@@ -469,7 +470,7 @@ bool UGroomBindingAsset::IsCompatible(const USkeletalMesh* InSkeletalMesh, const
 		}
 		#endif
 
-		for (const FHairGroupResource& Resource : InBinding->HairGroupResources)
+		for (const FHairGroupResource& Resource : InBinding->GetHairGroupResources())
 		{
 			if (Resource.SimRootResources && InSkeletalMesh->GetLODNum() != Resource.SimRootResources->GetLODCount())
 			{
@@ -506,7 +507,7 @@ bool UGroomBindingAsset::IsCompatible(const UGeometryCache* InGeometryCache, con
 {
 	if (InBinding && InGeometryCache && IsHairStrandsBindingEnable())
 	{
-		if (InBinding->GroomBindingType != EGroomBindingMeshType::GeometryCache)
+		if (InBinding->GetGroomBindingType() != EGroomBindingMeshType::GeometryCache)
 		{
 			if (bIssueWarning)
 			{
@@ -515,7 +516,7 @@ bool UGroomBindingAsset::IsCompatible(const UGeometryCache* InGeometryCache, con
 			return false;
 		}
 
-		if (!InBinding->TargetGeometryCache)
+		if (!InBinding->GetTargetGeometryCache())
 		{
 			if (bIssueWarning)
 			{
@@ -566,34 +567,34 @@ bool UGroomBindingAsset::IsCompatible(const UGroomAsset* InGroom, const UGroomBi
 {
 	if (InBinding && InGroom && IsHairStrandsBindingEnable())
 	{
-		if (InBinding->Groom && !InBinding->Groom->IsValid())
+		if (InBinding->GetGroom() && !InBinding->GetGroom()->IsValid())
 		{
 			// The groom could be invalid if it's still being loaded asynchronously
 			return false;
 		}
 
-		if (!InBinding->Groom)
+		if (!InBinding->GetGroom())
 		{
 			UE_CLOG(bIssueWarning, LogHairStrands, Warning, TEXT("[Groom] The binding asset (%s) does not reference a groom. Falling back onto non-binding version."), *InBinding->GetName());
 			return false;
 		}
 
-		if (InGroom->GetPrimaryAssetId() != InBinding->Groom->GetPrimaryAssetId())
+		if (InGroom->GetPrimaryAssetId() != InBinding->GetGroom()->GetPrimaryAssetId())
 		{
 			UE_LOG(LogHairStrands, Warning, TEXT("[Groom] The Groom binding (%s) does not reference the same groom asset (BindingAsset's groom:%s vs. Groom:%s). The binding asset will not be used."), 
 				*InBinding->GetName(),
-				*InBinding->Groom->GetName(),
+				*InBinding->GetGroom()->GetName(),
 				*InGroom->GetName());
 			return false;
 		}
 
 		const uint32 GroupCount = InGroom->GetNumHairGroups();
-		if (GroupCount != InBinding->GroupInfos.Num())
+		if (GroupCount != InBinding->GetGroupInfos().Num())
 		{
 			UE_CLOG(bIssueWarning, LogHairStrands, Warning, TEXT("[Groom] The GroomBinding asset (%s) does not contains the same number of groups (%d vs. %d) than the groom (%s). The binding asset will not be used."),
 				*InBinding->GetName(),
 				GroupCount,
-				InBinding->GroupInfos.Num(),
+				InBinding->GetGroupInfos().Num(),
 				*InGroom->GetName());
 			return false;
 		}
@@ -603,7 +604,7 @@ bool UGroomBindingAsset::IsCompatible(const UGroomAsset* InGroom, const UGroomBi
 			// Guides
 			{
 				const uint32 GroomCount = InGroom->HairGroupsPlatformData[GroupIt].Guides.BulkData.GetNumCurves();
-				const uint32 BindingCount = InBinding->GroupInfos[GroupIt].SimRootCount;
+				const uint32 BindingCount = InBinding->GetGroupInfos()[GroupIt].SimRootCount;
 
 				if (GroomCount != 0 && GroomCount != BindingCount)
 				{
@@ -622,7 +623,7 @@ bool UGroomBindingAsset::IsCompatible(const UGroomAsset* InGroom, const UGroomBi
 			if (bNeedStrandsRoot)
 			{
 				const uint32 GroomCount = InGroom->HairGroupsPlatformData[GroupIt].Strands.BulkData.GetNumCurves();
-				const uint32 BindingCount = InBinding->GroupInfos[GroupIt].RenRootCount;
+				const uint32 BindingCount = InBinding->GetGroupInfos()[GroupIt].RenRootCount;
 
 				// Groom may have stripped strands data so GroomCount would be 0
 				if (GroomCount != 0 && GroomCount != BindingCount)
@@ -654,7 +655,7 @@ bool UGroomBindingAsset::IsBindingAssetValid(const UGroomBindingAsset* InBinding
 			return false;
 		}
 
-		if (!InBinding->Groom)
+		if (!InBinding->GetGroom())
 		{
 			if (bIssueWarning)
 			{
@@ -662,33 +663,33 @@ bool UGroomBindingAsset::IsBindingAssetValid(const UGroomBindingAsset* InBinding
 			}
 			return false;
 		}
-		if (InBinding->GroupInfos.Num() == 0)
+		if (InBinding->GetGroupInfos().Num() == 0)
 		{
 			if (bIssueWarning)
 			{
-				UE_LOG(LogHairStrands, Warning, TEXT("[Groom] The GroomBinding asset (Groom:%s) does not contain any groups. It is invalid and can't be assigned. The binding asset will not be used."), *InBinding->Groom->GetName());
+				UE_LOG(LogHairStrands, Warning, TEXT("[Groom] The GroomBinding asset (Groom:%s) does not contain any groups. It is invalid and can't be assigned. The binding asset will not be used."), *InBinding->GetGroom()->GetName());
 			}
 			return false;
 		}
 
 		uint32 GroupIt = 0;
-		for (const FGoomBindingGroupInfo& Info : InBinding->GroupInfos)
+		for (const FGoomBindingGroupInfo& Info : InBinding->GetGroupInfos())
 		{
 			if (Info.SimRootCount == 0)
 			{
 				if (bIssueWarning)
 				{
-					UE_LOG(LogHairStrands, Warning, TEXT("[Groom] The GroomBinding asset (Groom:%s) has group with 0 guides. It is invalid and can't be assigned. The binding asset will not be used."), *InBinding->Groom->GetName());
+					UE_LOG(LogHairStrands, Warning, TEXT("[Groom] The GroomBinding asset (Groom:%s) has group with 0 guides. It is invalid and can't be assigned. The binding asset will not be used."), *InBinding->GetGroom()->GetName());
 				}
 				return false;
 			}
 
-			const bool bNeedStrandsRoot = DoesGroomNeedStrandsBinding(InBinding->Groom, GroupIt);
+			const bool bNeedStrandsRoot = DoesGroomNeedStrandsBinding(InBinding->GetGroom(), GroupIt);
 			if (bNeedStrandsRoot && Info.RenRootCount == 0 && IsHairStrandsEnabled(EHairStrandsShaderType::Strands))
 			{
 				if (bIssueWarning)
 				{
-					UE_LOG(LogHairStrands, Warning, TEXT("[Groom] The GroomBinding asset (Groom:%s) has group with 0 curves. It is invalid and can't be assigned. The binding asset will not be used."), *InBinding->Groom->GetName());
+					UE_LOG(LogHairStrands, Warning, TEXT("[Groom] The GroomBinding asset (Groom:%s) has group with 0 curves. It is invalid and can't be assigned. The binding asset will not be used."), *InBinding->GetGroom()->GetName());
 				}
 				return false;
 			}
@@ -701,15 +702,15 @@ bool UGroomBindingAsset::IsBindingAssetValid(const UGroomBindingAsset* InBinding
 
 bool UGroomBindingAsset::HasValidTarget() const
 {
-	return (GroomBindingType == EGroomBindingMeshType::SkeletalMesh && TargetSkeletalMesh) ||
-		   (GroomBindingType == EGroomBindingMeshType::GeometryCache && TargetGeometryCache);
+	return (GetGroomBindingType() == EGroomBindingMeshType::SkeletalMesh && GetTargetSkeletalMesh()) ||
+		   (GetGroomBindingType() == EGroomBindingMeshType::GeometryCache && GetTargetGeometryCache());
 }
 
 #if WITH_EDITOR
 
 void UGroomBindingAsset::Build()
 {
-	if (Groom && HasValidTarget())
+	if (GetGroom() && HasValidTarget())
 	{
 		OnGroomBindingAssetChanged.Broadcast();
 		Reset();
@@ -729,12 +730,12 @@ void UpdateGroomBindingAssetInfos(UGroomBindingAsset* In)
 {
 	if (In)
 	{
-		const uint32 GroupCount = In->HairGroupsPlatformData.Num();
-		In->GroupInfos.SetNum(GroupCount);
+		const uint32 GroupCount = In->GetHairGroupsPlatformData().Num();
+		In->GetGroupInfos().SetNum(GroupCount);
 		for (uint32 GroupIt=0; GroupIt< GroupCount; ++GroupIt)
 		{
-			FGoomBindingGroupInfo& Info = In->GroupInfos[GroupIt];
-			const UGroomBindingAsset::FHairGroupPlatformData& BulkData = In->HairGroupsPlatformData[GroupIt];
+			FGoomBindingGroupInfo& Info = In->GetGroupInfos()[GroupIt];
+			const UGroomBindingAsset::FHairGroupPlatformData& BulkData = In->GetHairGroupsPlatformData()[GroupIt];
 			Info.SimRootCount = BulkData.SimRootBulkData.GetRootCount();
 			Info.SimLODCount  = BulkData.SimRootBulkData.GetLODCount();
 
@@ -775,23 +776,23 @@ static FString BuildDerivedDataKeySuffix(const UGroomBindingAsset& BindingAsset)
 	FString SourceKey;
 	FString TargetKey;
 
-	if (BindingAsset.GroomBindingType == EGroomBindingMeshType::SkeletalMesh)
+	if (BindingAsset.GetGroomBindingType() == EGroomBindingMeshType::SkeletalMesh)
 	{
 		// Binding type is implicitly SkeletalMesh so keep BindingType empty to prevent triggering rebuild of old binding for nothing
-		SourceKey = BindingAsset.SourceSkeletalMesh ? BindingAsset.SourceSkeletalMesh->GetDerivedDataKey() : FString();
-		TargetKey = BindingAsset.TargetSkeletalMesh ? BindingAsset.TargetSkeletalMesh->GetDerivedDataKey() : FString();
+		SourceKey = BindingAsset.GetSourceSkeletalMesh() ? BindingAsset.GetSourceSkeletalMesh()->GetDerivedDataKey() : FString();
+		TargetKey = BindingAsset.GetTargetSkeletalMesh() ? BindingAsset.GetTargetSkeletalMesh()->GetDerivedDataKey() : FString();
 	}
 	else
 	{
 		BindingType = "GEOCACHE_";
-		SourceKey = BindingAsset.SourceGeometryCache ? BindingAsset.SourceGeometryCache->GetHash() : FString();
-		TargetKey = BindingAsset.TargetGeometryCache ? BindingAsset.TargetGeometryCache->GetHash() : FString();
+		SourceKey = BindingAsset.GetSourceGeometryCache() ? BindingAsset.GetSourceGeometryCache()->GetHash() : FString();
+		TargetKey = BindingAsset.GetTargetGeometryCache() ? BindingAsset.GetTargetGeometryCache()->GetHash() : FString();
 	}
 	// When possible, use the GroomAsset 'cached DDC key'. This allows to avoid a bug where the DDC key would change 
 	// when loading GroomAsset's hair description, which would modify the hair description hash ID with legacy content.
-	FString GroomKey  = BindingAsset.Groom ? BindingAsset.Groom->GetDerivedDataKey(true /*bUseCachedKey*/) : FString();
-	FString PointKey  = FString::FromInt(BindingAsset.NumInterpolationPoints);
-	FString SectionKey = FString::FromInt(BindingAsset.MatchingSection);
+	FString GroomKey  = BindingAsset.GetGroom() ? BindingAsset.GetGroom()->GetDerivedDataKey(true /*bUseCachedKey*/) : FString();
+	FString PointKey  = FString::FromInt(BindingAsset.GetNumInterpolationPoints());
+	FString SectionKey = FString::FromInt(BindingAsset.GetMatchingSection());
 
 	uint32 KeyLength  = BindingType.Len() + SourceKey.Len() + TargetKey.Len() + GroomKey.Len() + PointKey.Len() + SectionKey.Len();
 
@@ -803,17 +804,16 @@ static FString BuildDerivedDataKeySuffix(const UGroomBindingAsset& BindingAsset)
 
 void UGroomBindingAsset::CacheDerivedDatas()
 {
-	if (!Groom || !Groom->IsValid())
+	if (!GetGroom() || !GetGroom()->IsValid())
 	{
-		// The groom could be invalid if it's still being loaded asynchronously
 		return;
 	}
 
 	// 1. Set group count to the groom target
-	const uint32 GroupCount = Groom->GetNumHairGroups();
-	HairGroupsPlatformData.SetNum(GroupCount);
+	const uint32 GroupCount = GetGroom()->GetNumHairGroups();
+	GetHairGroupsPlatformData().SetNum(GroupCount);
 	CachedDerivedDataKey.SetNum(GroupCount);
-	GroupInfos.SetNum(GroupCount);
+	GetGroupInfos().SetNum(GroupCount);
 
 	// 2. Prepare main cache key
 	const FString KeySuffix = BuildDerivedDataKeySuffix(*this);
@@ -864,7 +864,7 @@ void UGroomBindingAsset::CacheDerivedDatas(uint32 InGroupIndex, const FString Ke
 			Owner.Wait();
 		}
 
-		UGroomBindingAsset::FHairGroupPlatformData& PlatformData = HairGroupsPlatformData[InGroupIndex];
+		UGroomBindingAsset::FHairGroupPlatformData& PlatformData = GetHairGroupsPlatformData()[InGroupIndex];
 
 		// Populate key/name for streaming data request
 		auto FillDrivedDataKey = [&DerivedDataKey, &Name](UGroomBindingAsset::FHairGroupPlatformData& In)
@@ -872,7 +872,7 @@ void UGroomBindingAsset::CacheDerivedDatas(uint32 InGroupIndex, const FString Ke
 			In.RenRootBulkData.DerivedDataKey = DerivedDataKey + FString(TEXT("_Data"));
 		};
 
-		FGroomComponentRecreateRenderStateContext RecreateRenderContext(Groom);
+		FGroomComponentRecreateRenderStateContext RecreateRenderContext(GetGroom());
 
 		if (Data)
 		{
@@ -936,9 +936,9 @@ void UGroomBindingAsset::CacheDerivedDatas(uint32 InGroupIndex, const FString Ke
 
 void UGroomBindingAsset::GetResourceSizeEx(FResourceSizeEx& CumulativeResourceSize)
 {
-	CumulativeResourceSize.AddDedicatedSystemMemoryBytes(HairGroupsPlatformData.GetAllocatedSize());
+	CumulativeResourceSize.AddDedicatedSystemMemoryBytes(GetHairGroupsPlatformData().GetAllocatedSize());
 
-	for (const FHairGroupResource& Group : HairGroupResources)
+	for (const FHairGroupResource& Group : GetHairGroupResources())
 	{
 		if (Group.SimRootResources) CumulativeResourceSize.AddDedicatedVideoMemoryBytes(Group.SimRootResources->GetResourcesSize());
 		if (Group.RenRootResources) CumulativeResourceSize.AddDedicatedVideoMemoryBytes(Group.RenRootResources->GetResourcesSize());
@@ -964,3 +964,106 @@ FName UGroomBindingAsset::GetAssetPathName(int32 LODIndex)
 	return NAME_None;
 #endif
 }
+
+#define DEFINE_GROOMBINDING_MEMBER_ACCESSOR(Type, Access, Name, GetConst, SetConst)\
+	FName UGroomBindingAsset::Get##Name##MemberName()\
+	{\
+		PRAGMA_DISABLE_DEPRECATION_WARNINGS\
+		return FName(TEXT(#Name));\
+		PRAGMA_ENABLE_DEPRECATION_WARNINGS\
+	}\
+	UFUNCTION(BlueprintGetter)\
+	GetConst Type Access UGroomBindingAsset::Get##Name() const\
+	{\
+		PRAGMA_DISABLE_DEPRECATION_WARNINGS\
+		return Name;\
+		PRAGMA_ENABLE_DEPRECATION_WARNINGS\
+	}\
+	UFUNCTION(BlueprintSetter)\
+	void UGroomBindingAsset::Set##Name(SetConst Type Access In##Name)\
+	{\
+		PRAGMA_DISABLE_DEPRECATION_WARNINGS\
+		Name = In##Name;\
+		PRAGMA_ENABLE_DEPRECATION_WARNINGS\
+	}
+
+// Define most of the binding member accessor
+DEFINE_GROOMBINDING_MEMBER_ACCESSOR(EGroomBindingMeshType, , GroomBindingType, , );
+DEFINE_GROOMBINDING_MEMBER_ACCESSOR(UGroomAsset, *, Groom, , );
+DEFINE_GROOMBINDING_MEMBER_ACCESSOR(USkeletalMesh, *, SourceSkeletalMesh, , );
+DEFINE_GROOMBINDING_MEMBER_ACCESSOR(USkeletalMesh, *, TargetSkeletalMesh, , );
+DEFINE_GROOMBINDING_MEMBER_ACCESSOR(UGeometryCache, *, SourceGeometryCache, , );
+DEFINE_GROOMBINDING_MEMBER_ACCESSOR(UGeometryCache, *, TargetGeometryCache, , );
+DEFINE_GROOMBINDING_MEMBER_ACCESSOR(int32, , NumInterpolationPoints, , );
+DEFINE_GROOMBINDING_MEMBER_ACCESSOR(int32, , MatchingSection, , );
+DEFINE_GROOMBINDING_MEMBER_ACCESSOR(TArray<FGoomBindingGroupInfo>, &, GroupInfos, const, const);
+
+TArray<FGoomBindingGroupInfo>& UGroomBindingAsset::GetGroupInfos()
+{
+	PRAGMA_DISABLE_DEPRECATION_WARNINGS
+	return GroupInfos;
+	PRAGMA_ENABLE_DEPRECATION_WARNINGS
+}
+
+void UGroomBindingAsset::AddHairGroupResourcesToDelete(FHairGroupResource& In)
+{
+	PRAGMA_DISABLE_DEPRECATION_WARNINGS
+	HairGroupResourcesToDelete.Enqueue(In);
+	PRAGMA_ENABLE_DEPRECATION_WARNINGS
+}
+
+bool UGroomBindingAsset::RemoveHairGroupResourcesToDelete(FHairGroupResource& Out)
+{
+	PRAGMA_DISABLE_DEPRECATION_WARNINGS
+	return HairGroupResourcesToDelete.Dequeue(Out);
+	PRAGMA_ENABLE_DEPRECATION_WARNINGS
+}
+
+FName UGroomBindingAsset::GetHairGroupResourcesMemberName()
+{
+	PRAGMA_DISABLE_DEPRECATION_WARNINGS
+	return GET_MEMBER_NAME_CHECKED(UGroomBindingAsset, GroupInfos);
+	PRAGMA_ENABLE_DEPRECATION_WARNINGS
+}
+
+UGroomBindingAsset::FHairGroupResources& UGroomBindingAsset::GetHairGroupResources()
+{
+	PRAGMA_DISABLE_DEPRECATION_WARNINGS
+	return HairGroupResources;
+	PRAGMA_ENABLE_DEPRECATION_WARNINGS
+}
+
+const UGroomBindingAsset::FHairGroupResources& UGroomBindingAsset::GetHairGroupResources() const
+{
+	PRAGMA_DISABLE_DEPRECATION_WARNINGS
+	return HairGroupResources;
+	PRAGMA_ENABLE_DEPRECATION_WARNINGS
+}
+
+void UGroomBindingAsset::SetHairGroupResources(UGroomBindingAsset::FHairGroupResources InHairGroupResources)
+{
+	PRAGMA_DISABLE_DEPRECATION_WARNINGS
+	HairGroupResources = InHairGroupResources;
+	PRAGMA_ENABLE_DEPRECATION_WARNINGS
+}
+
+FName UGroomBindingAsset::GetHairGroupPlatformDataMemberName()
+{
+	PRAGMA_DISABLE_DEPRECATION_WARNINGS
+	return GET_MEMBER_NAME_CHECKED(UGroomBindingAsset, HairGroupsPlatformData);
+	PRAGMA_ENABLE_DEPRECATION_WARNINGS
+}
+
+const TArray<UGroomBindingAsset::FHairGroupPlatformData>& UGroomBindingAsset::GetHairGroupsPlatformData() const
+{
+	PRAGMA_DISABLE_DEPRECATION_WARNINGS
+	return HairGroupsPlatformData;
+	PRAGMA_ENABLE_DEPRECATION_WARNINGS
+}
+
+TArray<UGroomBindingAsset::FHairGroupPlatformData>& UGroomBindingAsset::GetHairGroupsPlatformData()
+{
+	PRAGMA_DISABLE_DEPRECATION_WARNINGS
+	return HairGroupsPlatformData;
+	PRAGMA_ENABLE_DEPRECATION_WARNINGS
+}
\ No newline at end of file
diff --git a/Engine/Plugins/Runtime/HairStrands/Source/HairStrandsCore/Private/GroomBindingBuilder.cpp b/Engine/Plugins/Runtime/HairStrands/Source/HairStrandsCore/Private/GroomBindingBuilder.cpp
--- a/Engine/Plugins/Runtime/HairStrands/Source/HairStrandsCore/Private/GroomBindingBuilder.cpp
+++ b/Engine/Plugins/Runtime/HairStrands/Source/HairStrandsCore/Private/GroomBindingBuilder.cpp
@@ -57,7 +57,7 @@ static FAutoConsoleVariableRef CVarHairStrandsBindingBuilderWarningEnable(TEXT("
 FString FGroomBindingBuilder::GetVersion()
 {
 	// Important to update the version when groom building changes
-	return TEXT("3e");
+	return TEXT("3f");
 }
 
 ///////////////////////////////////////////////////////////////////////////////////////////////////
@@ -2041,9 +2041,9 @@ static bool InternalBuildBinding_CPU(UGroomBindingAsset* BindingAsset, uint32 In
 {
 #if WITH_EDITORONLY_DATA
 	if (!BindingAsset ||
-		!BindingAsset->Groom ||
+		!BindingAsset->GetGroom() ||
 		!BindingAsset->HasValidTarget() ||
-		BindingAsset->Groom->GetNumHairGroups() == 0)
+		BindingAsset->GetGroom()->GetNumHairGroups() == 0)
 	{
 		UE_LOG(LogHairStrands, Error, TEXT("[Groom] Binding asset cannot be created/rebuilt."));
 		return false;
@@ -2052,34 +2052,34 @@ static bool InternalBuildBinding_CPU(UGroomBindingAsset* BindingAsset, uint32 In
 	// 1. Build groom root data
 	FHairRootGroupData OutData;
 	{
-		BindingAsset->Groom->ConditionalPostLoad();
+		BindingAsset->GetGroom()->ConditionalPostLoad();
 
-		const int32 NumInterpolationPoints = BindingAsset->NumInterpolationPoints;
-		UGroomAsset* GroomAsset = BindingAsset->Groom;
+		const int32 NumInterpolationPoints = BindingAsset->GetNumInterpolationPoints();
+		UGroomAsset* GroomAsset = BindingAsset->GetGroom();
 
 		TUniquePtr<IMeshData> SourceMeshData;
 		TUniquePtr<IMeshData> TargetMeshData;
-		if (BindingAsset->GroomBindingType == EGroomBindingMeshType::SkeletalMesh)
+		if (BindingAsset->GetGroomBindingType() == EGroomBindingMeshType::SkeletalMesh)
 		{
-			BindingAsset->TargetSkeletalMesh->ConditionalPostLoad();
-			if (BindingAsset->SourceSkeletalMesh)
+			BindingAsset->GetTargetSkeletalMesh()->ConditionalPostLoad();
+			if (BindingAsset->GetSourceSkeletalMesh())
 			{
-				BindingAsset->SourceSkeletalMesh->ConditionalPostLoad();
+				BindingAsset->GetSourceSkeletalMesh()->ConditionalPostLoad();
 			}
 
-			SourceMeshData = TUniquePtr<FSkeletalMeshData, TDefaultDelete<IMeshData>>(new FSkeletalMeshData(BindingAsset->SourceSkeletalMesh));
-			TargetMeshData = TUniquePtr<FSkeletalMeshData, TDefaultDelete<IMeshData>>(new FSkeletalMeshData(BindingAsset->TargetSkeletalMesh));
+			SourceMeshData = TUniquePtr<FSkeletalMeshData, TDefaultDelete<IMeshData>>(new FSkeletalMeshData(BindingAsset->GetSourceSkeletalMesh()));
+			TargetMeshData = TUniquePtr<FSkeletalMeshData, TDefaultDelete<IMeshData>>(new FSkeletalMeshData(BindingAsset->GetTargetSkeletalMesh()));
 		}
 		else
 		{
-			BindingAsset->TargetGeometryCache->ConditionalPostLoad();
-			if (BindingAsset->SourceGeometryCache)
+			BindingAsset->GetTargetGeometryCache()->ConditionalPostLoad();
+			if (BindingAsset->GetSourceGeometryCache())
 			{
-				BindingAsset->SourceGeometryCache->ConditionalPostLoad();
+				BindingAsset->GetSourceGeometryCache()->ConditionalPostLoad();
 			}
 
-			SourceMeshData = TUniquePtr<FGeometryCacheData, TDefaultDelete<IMeshData>>(new FGeometryCacheData(BindingAsset->SourceGeometryCache));
-			TargetMeshData = TUniquePtr<FGeometryCacheData, TDefaultDelete<IMeshData>>(new FGeometryCacheData(BindingAsset->TargetGeometryCache));
+			SourceMeshData = TUniquePtr<FGeometryCacheData, TDefaultDelete<IMeshData>>(new FGeometryCacheData(BindingAsset->GetSourceGeometryCache()));
+			TargetMeshData = TUniquePtr<FGeometryCacheData, TDefaultDelete<IMeshData>>(new FGeometryCacheData(BindingAsset->GetTargetGeometryCache()));
 		}
 
 		if (!TargetMeshData->IsValid())
@@ -2162,7 +2162,7 @@ static bool InternalBuildBinding_CPU(UGroomBindingAsset* BindingAsset, uint32 In
 			if (!GroomBinding_Transfer::Transfer(
 				SourceMeshData.Get(),
 				TargetMeshData.Get(),
-				TransferredPositions, BindingAsset->MatchingSection))
+				TransferredPositions, BindingAsset->GetMatchingSection()))
 			{
 				UE_LOG(LogHairStrands, Error, TEXT("[Groom] Binding asset could not be built. Positions transfer between source and target mesh failed."));
 				return false;
@@ -2205,7 +2205,7 @@ static bool InternalBuildBinding_CPU(UGroomBindingAsset* BindingAsset, uint32 In
 			const uint32 CardsLODCount = OutData.CardsRootData.Num();
 			for (uint32 CardsLODIt = 0; CardsLODIt < CardsLODCount; ++CardsLODIt)
 			{
-				if (BindingAsset->Groom->HairGroupsPlatformData[InGroupIndex].Cards.IsValid(CardsLODIt))
+				if (BindingAsset->GetGroom()->HairGroupsPlatformData[InGroupIndex].Cards.IsValid(CardsLODIt))
 				{
 					FHairStrandsDatas LODGuidesData;
 					const bool bIsValid = GroomAsset->GetHairCardsGuidesDatas(InGroupIndex, CardsLODIt, LODGuidesData);
@@ -2228,25 +2228,25 @@ static bool InternalBuildBinding_CPU(UGroomBindingAsset* BindingAsset, uint32 In
 		
 		// 1.5 RBF building
 		{
-			GroomBinding_RBFWeighting::ComputeInterpolationWeights(OutData, bNeedStrandsRoot, BindingAsset->NumInterpolationPoints, BindingAsset->MatchingSection, TargetMeshData.Get(), TransferredPositions);
+			GroomBinding_RBFWeighting::ComputeInterpolationWeights(OutData, bNeedStrandsRoot, BindingAsset->GetNumInterpolationPoints(), BindingAsset->GetMatchingSection(), TargetMeshData.Get(), TransferredPositions);
 			SlowTask.EnterProgressFrame();
 		}
 	}
 
 	// 2. Release existing resources data
-	UGroomBindingAsset::FHairGroupResources& OutHairGroupResources = BindingAsset->HairGroupResources;
-	if (BindingAsset->HairGroupResources.Num() > 0)
+	UGroomBindingAsset::FHairGroupResources& OutHairGroupResources = BindingAsset->GetHairGroupResources();
+	if (OutHairGroupResources.Num() > 0)
 	{
 		for (UGroomBindingAsset::FHairGroupResource& GroupResrouces : OutHairGroupResources)
 		{
-			BindingAsset->HairGroupResourcesToDelete.Enqueue(GroupResrouces);
+			BindingAsset->AddHairGroupResourcesToDelete(GroupResrouces);
 		}
 		OutHairGroupResources.Empty();
 	}
 	check(OutHairGroupResources.Num() == 0);
 
 	// 3. Convert data to bulk data
-	GroomBinding_BulkCopy::BuildRootBulkData(BindingAsset->HairGroupsPlatformData[InGroupIndex], OutData);
+	GroomBinding_BulkCopy::BuildRootBulkData(BindingAsset->GetHairGroupsPlatformData()[InGroupIndex], OutData);
 
 	BindingAsset->QueryStatus = UGroomBindingAsset::EQueryStatus::Completed;
 #endif
@@ -2264,8 +2264,8 @@ bool FGroomBindingBuilder::BuildBinding(UGroomBindingAsset* BindingAsset, bool b
 	bool bOutValid = true;
 
 	// 1. Build binding asset
-	const uint32 GroupCount = BindingAsset->HairGroupsPlatformData.Num();
-	BindingAsset->GroupInfos.SetNum(GroupCount);
+	const uint32 GroupCount = BindingAsset->GetHairGroupsPlatformData().Num();
+	BindingAsset->GetGroupInfos().SetNum(GroupCount);
 	for (uint32 InGroupIndex = 0; InGroupIndex < GroupCount; ++InGroupIndex)
 	{
 		bOutValid = InternalBuildBinding_CPU(BindingAsset, InGroupIndex) && bOutValid;
diff --git a/Engine/Plugins/Runtime/HairStrands/Source/HairStrandsCore/Private/GroomComponent.cpp b/Engine/Plugins/Runtime/HairStrands/Source/HairStrandsCore/Private/GroomComponent.cpp
--- a/Engine/Plugins/Runtime/HairStrands/Source/HairStrandsCore/Private/GroomComponent.cpp
+++ b/Engine/Plugins/Runtime/HairStrands/Source/HairStrandsCore/Private/GroomComponent.cpp
@@ -2250,10 +2250,10 @@ static USkeletalMeshComponent* ValidateBindingAsset(
 				const bool bIsCardsBindingCompatible =
 					bIsBindingCompatible &&
 					BindingAsset &&
-					GroupIt < BindingAsset->HairGroupResources.Num() &&
-					CardsLODIndex < uint32(BindingAsset->HairGroupResources[GroupIt].CardsRootResources.Num()) &&
-					BindingAsset->HairGroupResources[GroupIt].CardsRootResources[CardsLODIndex] != nullptr &&
-					((SkeletalMeshComponent && SkeletalMeshComponent->GetSkeletalMeshAsset()) ? SkeletalMeshComponent->GetSkeletalMeshAsset()->GetLODInfoArray().Num() == BindingAsset->HairGroupResources[GroupIt].CardsRootResources[CardsLODIndex]->GetLODCount() : false);
+					GroupIt < BindingAsset->GetHairGroupResources().Num() &&
+					CardsLODIndex < uint32(BindingAsset->GetHairGroupResources()[GroupIt].CardsRootResources.Num()) &&
+					BindingAsset->GetHairGroupResources()[GroupIt].CardsRootResources[CardsLODIndex] != nullptr &&
+					((SkeletalMeshComponent && SkeletalMeshComponent->GetSkeletalMeshAsset()) ? SkeletalMeshComponent->GetSkeletalMeshAsset()->GetLODInfoArray().Num() == BindingAsset->GetHairGroupResources()[GroupIt].CardsRootResources[CardsLODIndex]->GetLODCount() : false);
 
 				if (!bIsCardsBindingCompatible)
 				{
@@ -2280,7 +2280,7 @@ static UMeshComponent* ValidateBindingAsset(
 		return nullptr;
 	}
 
-	if (BindingAsset->GroomBindingType == EGroomBindingMeshType::SkeletalMesh)
+	if (BindingAsset->GetGroomBindingType() == EGroomBindingMeshType::SkeletalMesh)
 	{
 		return ValidateBindingAsset(GroomAsset, BindingAsset, Cast<USkeletalMeshComponent>(MeshComponent), bIsBindingReloading, bValidationEnable, Component);
 	}
@@ -2393,8 +2393,8 @@ void UGroomComponent::InitResources(bool bIsBindingReloading)
 		if (ValidatedMeshComponent)
 		{
 			if (BindingAsset && 
-				((BindingAsset->GroomBindingType == EGroomBindingMeshType::SkeletalMesh && Cast<USkeletalMeshComponent>(ValidatedMeshComponent)->GetSkeletalMeshAsset() == nullptr) ||
-				(BindingAsset->GroomBindingType == EGroomBindingMeshType::GeometryCache && Cast<UGeometryCacheComponent>(ValidatedMeshComponent)->GeometryCache == nullptr)))
+				((BindingAsset->GetGroomBindingType() == EGroomBindingMeshType::SkeletalMesh && Cast<USkeletalMeshComponent>(ValidatedMeshComponent)->GetSkeletalMeshAsset() == nullptr) ||
+				(BindingAsset->GetGroomBindingType() == EGroomBindingMeshType::GeometryCache && Cast<UGeometryCacheComponent>(ValidatedMeshComponent)->GeometryCache == nullptr)))
 			{
 				ValidatedMeshComponent = nullptr;
 			}
@@ -2481,7 +2481,7 @@ void UGroomComponent::InitResources(bool bIsBindingReloading)
 		if (bHasNeedBindingData)
 		{
 			LocalBindingAsset = BindingAsset;
-			if (LocalBindingAsset && LocalBindingAsset->HairGroupResources.Num() != GroomAsset->HairGroupsPlatformData.Num())
+			if (LocalBindingAsset && LocalBindingAsset->GetHairGroupResources().Num() != GroomAsset->HairGroupsPlatformData.Num())
 			{
 				LocalBindingAsset = nullptr;
 			}
@@ -2523,7 +2523,7 @@ void UGroomComponent::InitResources(bool bIsBindingReloading)
 		HairGroupInstance->Debug.GroomAssetName = GroomAsset->GetName();
 		HairGroupInstance->Debug.MeshComponentForDebug = RegisteredMeshComponent;
 		HairGroupInstance->Debug.GroomComponentForDebug = this;
-		HairGroupInstance->Debug.GroomBindingType = BindingAsset ? BindingAsset->GroomBindingType : EGroomBindingMeshType::SkeletalMesh;
+		HairGroupInstance->Debug.GroomBindingType = BindingAsset ? BindingAsset->GetGroomBindingType() : EGroomBindingMeshType::SkeletalMesh;
 		HairGroupInstance->Debug.GroomCacheType = GetEffectiveGroomCacheType(GroomCache, GroomAsset);
 		HairGroupInstance->Debug.GroomCacheBuffers = GroomCacheBuffers;
 		HairGroupInstance->Debug.LODForcedIndex = LODForcedIndex;
@@ -2618,13 +2618,13 @@ void UGroomComponent::InitResources(bool bIsBindingReloading)
 			{
 				check(RegisteredMeshComponent);
 				check(LocalBindingAsset);
-				check(GroupIt < LocalBindingAsset->HairGroupResources.Num());
+				check(GroupIt < LocalBindingAsset->GetHairGroupResources().Num());
 				if (USkeletalMeshComponent* SkeletalMeshComponent = Cast<USkeletalMeshComponent>(RegisteredMeshComponent))
 				{
-					check(SkeletalMeshComponent->GetSkeletalMeshAsset() ? SkeletalMeshComponent->GetSkeletalMeshAsset()->GetLODInfoArray().Num() == LocalBindingAsset->HairGroupResources[GroupIt].SimRootResources->GetLODCount() : false);
+					check(SkeletalMeshComponent->GetSkeletalMeshAsset() ? SkeletalMeshComponent->GetSkeletalMeshAsset()->GetLODInfoArray().Num() == LocalBindingAsset->GetHairGroupResources()[GroupIt].SimRootResources->GetLODCount() : false);
 				}
 
-				HairGroupInstance->Guides.RestRootResource = LocalBindingAsset->HairGroupResources[GroupIt].SimRootResources;
+				HairGroupInstance->Guides.RestRootResource = LocalBindingAsset->GetHairGroupResources()[GroupIt].SimRootResources;
 				HairGroupInstance->Guides.DeformedRootResource = new FHairStrandsDeformedRootResource(HairGroupInstance->Guides.RestRootResource, EHairStrandsResourcesType::Guides, ResourceName, OwnerName);
 			}
 
@@ -2742,13 +2742,13 @@ void UGroomComponent::InitResources(bool bIsBindingReloading)
 			{
 				check(LocalBindingAsset);
 				check(RegisteredMeshComponent);
-				check(GroupIt < LocalBindingAsset->HairGroupResources.Num());
+				check(GroupIt < LocalBindingAsset->GetHairGroupResources().Num());
 				if (USkeletalMeshComponent* SkeletalMeshComponent = Cast<USkeletalMeshComponent>(RegisteredMeshComponent))
 				{
-					check(SkeletalMeshComponent->GetSkeletalMeshAsset() ? SkeletalMeshComponent->GetSkeletalMeshAsset()->GetLODInfoArray().Num() == LocalBindingAsset->HairGroupResources[GroupIt].RenRootResources->GetLODCount() : false);
+					check(SkeletalMeshComponent->GetSkeletalMeshAsset() ? SkeletalMeshComponent->GetSkeletalMeshAsset()->GetLODInfoArray().Num() == LocalBindingAsset->GetHairGroupResources()[GroupIt].RenRootResources->GetLODCount() : false);
 				}
 
-				HairGroupInstance->Strands.RestRootResource = LocalBindingAsset->HairGroupResources[GroupIt].RenRootResources;
+				HairGroupInstance->Strands.RestRootResource = LocalBindingAsset->GetHairGroupResources()[GroupIt].RenRootResources;
 				HairGroupInstance->Strands.DeformedRootResource = new FHairStrandsDeformedRootResource(HairGroupInstance->Strands.RestRootResource, EHairStrandsResourcesType::Strands, ResourceName, OwnerName);
 			}
 
@@ -2847,13 +2847,13 @@ void UGroomComponent::InitResources(bool bIsBindingReloading)
 					if (bNeedRootData)
 					{
 						check(LocalBindingAsset);
-						check(GroupIt < LocalBindingAsset->HairGroupResources.Num());
+						check(GroupIt < LocalBindingAsset->GetHairGroupResources().Num());
 						if (USkeletalMeshComponent* SkeletalMeshComponent = Cast<USkeletalMeshComponent>(RegisteredMeshComponent))
 						{
-							check(SkeletalMeshComponent->GetSkeletalMeshAsset() ? SkeletalMeshComponent->GetSkeletalMeshAsset()->GetLODInfoArray().Num() == LocalBindingAsset->HairGroupResources[GroupIt].CardsRootResources[CardsLODIndex]->GetLODCount() : false);
+							check(SkeletalMeshComponent->GetSkeletalMeshAsset() ? SkeletalMeshComponent->GetSkeletalMeshAsset()->GetLODInfoArray().Num() == LocalBindingAsset->GetHairGroupResources()[GroupIt].CardsRootResources[CardsLODIndex]->GetLODCount() : false);
 						}
 
-						InstanceLOD.Guides.RestRootResource = LocalBindingAsset->HairGroupResources[GroupIt].CardsRootResources[CardsLODIndex];
+						InstanceLOD.Guides.RestRootResource = LocalBindingAsset->GetHairGroupResources()[GroupIt].CardsRootResources[CardsLODIndex];
 						InstanceLOD.Guides.DeformedRootResource = new FHairStrandsDeformedRootResource(InstanceLOD.Guides.RestRootResource, EHairStrandsResourcesType::Cards, LODResourceName, LODOwnerName);
 					}
 
diff --git a/Engine/Plugins/Runtime/HairStrands/Source/HairStrandsCore/Private/GroomRBFDeformer.cpp b/Engine/Plugins/Runtime/HairStrands/Source/HairStrandsCore/Private/GroomRBFDeformer.cpp
--- a/Engine/Plugins/Runtime/HairStrands/Source/HairStrandsCore/Private/GroomRBFDeformer.cpp
+++ b/Engine/Plugins/Runtime/HairStrands/Source/HairStrandsCore/Private/GroomRBFDeformer.cpp
@@ -584,18 +584,18 @@ static float SampleMaskTexture(const FVector2f& InUV, const FUintPoint& InResolu
 void FGroomRBFDeformer::GetRBFDeformedGroomAsset(const UGroomAsset* InGroomAsset, const UGroomBindingAsset* BindingAsset, FTextureSource* MaskTextureSource, const float MaskScale, UGroomAsset* OutGroomAsset)
 {
 #if WITH_EDITORONLY_DATA
-	if (InGroomAsset && BindingAsset && BindingAsset->TargetSkeletalMesh && BindingAsset->SourceSkeletalMesh)
+	if (InGroomAsset && BindingAsset && BindingAsset->GetTargetSkeletalMesh() && BindingAsset->GetSourceSkeletalMesh())
 	{
 		// Use the LOD0 skeletal mesh to extract the vertices used for the RBF weight computation
 		const int32 MeshLODIndex = 0;
 
 		// Get the target mesh vertices (source and target)
-		const FSkeletalMeshRenderData* SkeletalMeshData_Target = BindingAsset->TargetSkeletalMesh->GetResourceForRendering();
+		const FSkeletalMeshRenderData* SkeletalMeshData_Target = BindingAsset->GetTargetSkeletalMesh()->GetResourceForRendering();
 		TArray<FVector3f> MeshVertexPositionsBuffer_Target;
 		ExtractSkeletalVertexPosition(SkeletalMeshData_Target, MeshLODIndex, MeshVertexPositionsBuffer_Target);
 
 		// Apply RBF deformation to each group of guides and render strands
-		const int32 NumGroups = BindingAsset->HairGroupsPlatformData.Num();
+		const int32 NumGroups = BindingAsset->GetHairGroupsPlatformData().Num();
 
 		// Use the vertices positions from the HairDescription instead of the GroomAsset since the latter
 		// may contain decimated or auto-generated guides depending on the import settings
@@ -659,7 +659,7 @@ void FGroomRBFDeformer::GetRBFDeformedGroomAsset(const UGroomAsset* InGroomAsset
 			else
 			{
 				FHairStrandsRootData SimRootData;
-				FGroomBindingBuilder::GetRootData(SimRootData, BindingAsset->HairGroupsPlatformData[GroupIndex].SimRootBulkData);
+				FGroomBindingBuilder::GetRootData(SimRootData, BindingAsset->GetHairGroupsPlatformData()[GroupIndex].SimRootBulkData);
 
 				DeformedPositions[GroupIndex].GuideStrands = GetDeformedHairStrandsPositions(
 					MeshVertexPositionsBuffer_Target,
@@ -673,7 +673,7 @@ void FGroomRBFDeformer::GetRBFDeformedGroomAsset(const UGroomAsset* InGroomAsset
 			// Get deformed render strands
 			{
 				FHairStrandsRootData RenRootData;
-				FGroomBindingBuilder::GetRootData(RenRootData, BindingAsset->HairGroupsPlatformData[GroupIndex].RenRootBulkData);
+				FGroomBindingBuilder::GetRootData(RenRootData, BindingAsset->GetHairGroupsPlatformData()[GroupIndex].RenRootBulkData);
 
 				DeformedPositions[GroupIndex].RenderStrands = GetDeformedHairStrandsPositions(
 					MeshVertexPositionsBuffer_Target,
@@ -792,7 +792,7 @@ void FGroomRBFDeformer::GetRBFDeformedGroomAsset(const UGroomAsset* InGroomAsset
 				Mesh->ConditionalPostLoad();
 
 				FHairStrandsRootData RenRootData;
-				FGroomBindingBuilder::GetRootData(RenRootData, BindingAsset->HairGroupsPlatformData[Desc.GroupIndex].RenRootBulkData);
+				FGroomBindingBuilder::GetRootData(RenRootData, BindingAsset->GetHairGroupsPlatformData()[Desc.GroupIndex].RenRootBulkData);
 
 				DeformStaticMeshPositions(Mesh, MeshVertexPositionsBuffer_Target, RenRootData.MeshProjectionLODs[MeshLODIndex]);
 			}
@@ -811,7 +811,7 @@ void FGroomRBFDeformer::GetRBFDeformedGroomAsset(const UGroomAsset* InGroomAsset
 				Mesh->ConditionalPostLoad();
 				
 				FHairStrandsRootData RenRootData;
-				FGroomBindingBuilder::GetRootData(RenRootData, BindingAsset->HairGroupsPlatformData[Desc.GroupIndex].RenRootBulkData);
+				FGroomBindingBuilder::GetRootData(RenRootData, BindingAsset->GetHairGroupsPlatformData()[Desc.GroupIndex].RenRootBulkData);
 
 				DeformStaticMeshPositions(Mesh, MeshVertexPositionsBuffer_Target, RenRootData.MeshProjectionLODs[MeshLODIndex]);
 			}
diff --git a/Engine/Plugins/Runtime/HairStrands/Source/HairStrandsCore/Private/HairStrandsCore.cpp b/Engine/Plugins/Runtime/HairStrands/Source/HairStrandsCore/Private/HairStrandsCore.cpp
--- a/Engine/Plugins/Runtime/HairStrands/Source/HairStrandsCore/Private/HairStrandsCore.cpp
+++ b/Engine/Plugins/Runtime/HairStrands/Source/HairStrandsCore/Private/HairStrandsCore.cpp
@@ -219,21 +219,21 @@ UGroomBindingAsset* FHairStrandsCore::CreateGroomBindingAsset(EGroomBindingMeshT
 
 	if (UGroomBindingAsset* Out = NewObject<UGroomBindingAsset>(Package, *Name, RF_Public | RF_Standalone | RF_Transactional))
 	{
-		Out->GroomBindingType = BindingType;
-		Out->Groom = GroomAsset;
+		Out->SetGroomBindingType(BindingType);
+		Out->SetGroom(GroomAsset);
 		if (BindingType == EGroomBindingMeshType::SkeletalMesh)
 		{
-			Out->SourceSkeletalMesh = Cast<USkeletalMesh>(Source);
-			Out->TargetSkeletalMesh = Cast<USkeletalMesh>(Target);
+			Out->SetSourceSkeletalMesh(Cast<USkeletalMesh>(Source));
+			Out->SetTargetSkeletalMesh(Cast<USkeletalMesh>(Target));
 		}
 		else
 		{
-			Out->SourceGeometryCache = Cast<UGeometryCache>(Source);
-			Out->TargetGeometryCache = Cast<UGeometryCache>(Target);
+			Out->SetSourceGeometryCache(Cast<UGeometryCache>(Source));
+			Out->SetTargetGeometryCache(Cast<UGeometryCache>(Target));
 		}
-		Out->HairGroupsPlatformData.Reserve(GroomAsset->HairGroupsPlatformData.Num());
-		Out->NumInterpolationPoints = NumInterpolationPoints;
-		Out->MatchingSection = MatchingSection;
+		Out->GetHairGroupsPlatformData().Reserve(GroomAsset->HairGroupsPlatformData.Num());
+		Out->SetNumInterpolationPoints(NumInterpolationPoints);
+		Out->SetMatchingSection(MatchingSection);
 		Out->MarkPackageDirty();
 		AssetHelper().RegisterAsset(Out);
 		return Out;
diff --git a/Engine/Plugins/Runtime/HairStrands/Source/HairStrandsCore/Private/Niagara/NiagaraDataInterfaceHairStrands.cpp b/Engine/Plugins/Runtime/HairStrands/Source/HairStrandsCore/Private/Niagara/NiagaraDataInterfaceHairStrands.cpp
--- a/Engine/Plugins/Runtime/HairStrands/Source/HairStrandsCore/Private/Niagara/NiagaraDataInterfaceHairStrands.cpp
+++ b/Engine/Plugins/Runtime/HairStrands/Source/HairStrandsCore/Private/Niagara/NiagaraDataInterfaceHairStrands.cpp
@@ -384,8 +384,8 @@ void FNDIHairStrandsData::Update(UNiagaraDataInterfaceHairStrands* Interface, FN
 
 		GlobalInterpolation = bHasValidBindingAsset ? Interface->SourceComponent->GroomAsset->EnableGlobalInterpolation : false;
 		bSkinningTransfer = bHasValidBindingAsset ?
-			(Interface->SourceComponent->BindingAsset->SourceSkeletalMesh && Interface->SourceComponent->BindingAsset->TargetSkeletalMesh &&
-			 Interface->SourceComponent->BindingAsset->SourceSkeletalMesh != Interface->SourceComponent->BindingAsset->TargetSkeletalMesh) : false;
+			(Interface->SourceComponent->BindingAsset->GetSourceSkeletalMesh() && Interface->SourceComponent->BindingAsset->GetTargetSkeletalMesh() &&
+			 Interface->SourceComponent->BindingAsset->GetSourceSkeletalMesh() != Interface->SourceComponent->BindingAsset->GetTargetSkeletalMesh()) : false;
 		
 		TickingGroup = Interface->IsComponentValid() ? ComputeTickingGroup(Interface->SourceComponent) : NiagaraFirstTickGroup;
 
diff --git a/Engine/Plugins/Runtime/HairStrands/Source/HairStrandsCore/Public/GroomBindingAsset.h b/Engine/Plugins/Runtime/HairStrands/Source/HairStrandsCore/Public/GroomBindingAsset.h
--- a/Engine/Plugins/Runtime/HairStrands/Source/HairStrandsCore/Public/GroomBindingAsset.h
+++ b/Engine/Plugins/Runtime/HairStrands/Source/HairStrandsCore/Public/GroomBindingAsset.h
@@ -60,41 +60,104 @@ class HAIRSTRANDSCORE_API UGroomBindingAsset : public UObject
 	DECLARE_MULTICAST_DELEGATE(FOnGroomBindingAssetChanged);
 #endif
 
-public:
+private:
 	/** Type of mesh to create groom binding for */
-	UPROPERTY(VisibleAnywhere, BlueprintReadWrite, Category = "BuildSettings")
+	UE_DEPRECATED(5.3, "Please do not access this member directly; use UGroomBindingAsset::GetGroomBindingType() or UGroomBindingAsset::SetGroomBindingType().")
+	UPROPERTY(VisibleAnywhere, BlueprintGetter = GetGroomBindingType, BlueprintSetter = SetGroomBindingType, Category = "BuildSettings")
 	EGroomBindingMeshType GroomBindingType = EGroomBindingMeshType::SkeletalMesh;
 
+public:
+	static FName GetGroomBindingTypeMemberName();
+	UFUNCTION(BlueprintGetter) EGroomBindingMeshType GetGroomBindingType() const;
+	UFUNCTION(BlueprintSetter) void SetGroomBindingType(EGroomBindingMeshType InGroomBindingType);
+
+private:
 	/** Groom to bind. */
-	UPROPERTY(EditAnywhere, BlueprintReadWrite, Category = "BuildSettings", AssetRegistrySearchable)
+	UE_DEPRECATED(5.3, "Please do not access this member directly; use UGroomBindingAsset::GetGroom() or UGroomBindingAsset::SetGroom().")
+	UPROPERTY(EditAnywhere, BlueprintGetter = GetGroom, BlueprintSetter = SetGroom, Category = "BuildSettings", AssetRegistrySearchable)
 	TObjectPtr<UGroomAsset> Groom;
 
+public:
+	static FName GetGroomMemberName();
+	UFUNCTION(BlueprintGetter) UGroomAsset* GetGroom() const;
+	UFUNCTION(BlueprintSetter) void SetGroom(UGroomAsset* InGroom);
+
+private:
 	/** Skeletal mesh on which the groom has been authored. This is optional, and used only if the hair
 		binding is done a different mesh than the one which it has been authored */
-	UPROPERTY(EditAnywhere, BlueprintReadWrite, Category = "BuildSettings")
+	UE_DEPRECATED(5.3, "Please do not access this member directly; use UGroomBindingAsset::GetSourceSkeletalMesh() or UGroomBindingAsset::SetSourceSkeletalMesh().")
+	UPROPERTY(EditAnywhere, BlueprintGetter = GetSourceSkeletalMesh, BlueprintSetter = SetSourceSkeletalMesh, Category = "BuildSettings")
 	TObjectPtr<USkeletalMesh> SourceSkeletalMesh;
 
+public:
+	static FName GetSourceSkeletalMeshMemberName();
+	UFUNCTION(BlueprintGetter) USkeletalMesh* GetSourceSkeletalMesh() const;
+	UFUNCTION(BlueprintSetter) void SetSourceSkeletalMesh(USkeletalMesh* InSkeletalMesh);
+
+private:
 	/** Skeletal mesh on which the groom is attached to. */
-	UPROPERTY(EditAnywhere, BlueprintReadWrite, Category = "BuildSettings")
+	UE_DEPRECATED(5.3, "Please do not access this member directly; use UGroomBindingAsset::GetTargetSkeletalMesh() or UGroomBindingAsset::SetTargetSkeletalMesh().")
+	UPROPERTY(EditAnywhere, BlueprintGetter = GetTargetSkeletalMesh, BlueprintSetter = SetTargetSkeletalMesh, Category = "BuildSettings")
 	TObjectPtr<USkeletalMesh> TargetSkeletalMesh;
 
-	UPROPERTY(VisibleAnywhere, BlueprintReadWrite, Category = "BuildSettings")
+public:
+	static FName GetTargetSkeletalMeshMemberName();
+	UFUNCTION(BlueprintGetter) USkeletalMesh* GetTargetSkeletalMesh() const;
+	UFUNCTION(BlueprintSetter) void SetTargetSkeletalMesh(USkeletalMesh* InSkeletalMesh);
+
+private:
+	UE_DEPRECATED(5.3, "Please do not access this member directly; use UGroomBindingAsset::GetSourceGeometryCache() or UGroomBindingAsset::SetSourceGeometryCache().")
+	UPROPERTY(VisibleAnywhere, BlueprintGetter = GetSourceGeometryCache, BlueprintSetter = SetSourceGeometryCache, Category = "BuildSettings")
 	TObjectPtr<UGeometryCache> SourceGeometryCache;
 
-	UPROPERTY(VisibleAnywhere, BlueprintReadWrite, Category = "BuildSettings")
+public:
+	static FName GetSourceGeometryCacheMemberName();
+	UFUNCTION(BlueprintGetter) UGeometryCache* GetSourceGeometryCache() const;
+	UFUNCTION(BlueprintSetter) void SetSourceGeometryCache(UGeometryCache* InGeometryCache);
+
+private:
+	UE_DEPRECATED(5.3, "Please do not access this member directly; use UGroomBindingAsset::GetTargetGeometryCache() or UGroomBindingAsset::SetTargetGeometryCache().")
+	UPROPERTY(VisibleAnywhere, BlueprintGetter = GetTargetGeometryCache, BlueprintSetter = SetTargetGeometryCache, Category = "BuildSettings")
 	TObjectPtr<UGeometryCache> TargetGeometryCache;
 
+public:
+	static FName GetTargetGeometryCacheMemberName();
+	UFUNCTION(BlueprintGetter) UGeometryCache* GetTargetGeometryCache() const;
+	UFUNCTION(BlueprintSetter) void SetTargetGeometryCache(UGeometryCache* InGeometryCache);
+
+private:
 	/** Number of points used for the rbf interpolation */
-	UPROPERTY(VisibleAnywhere, BlueprintReadWrite, Category = "BuildSettings")
+	UE_DEPRECATED(5.3, "Please do not access this member directly; use UGroomBindingAsset::GetNumInterpolationPoints() or UGroomBindingAsset::SetNumInterpolationPoints().")
+	UPROPERTY(VisibleAnywhere, BlueprintGetter = GetNumInterpolationPoints, BlueprintSetter = SetNumInterpolationPoints, Category = "BuildSettings")
 	int32 NumInterpolationPoints = 100;
 
-	/** Section to pick to transfer the positions */
-	UPROPERTY(VisibleAnywhere, BlueprintReadWrite, Category = "BuildSettings")
+public:
+	static FName GetNumInterpolationPointsMemberName();
+	UFUNCTION(BlueprintGetter) int32 GetNumInterpolationPoints() const;
+	UFUNCTION(BlueprintSetter) void SetNumInterpolationPoints(int32 InNumInterpolationPoints);
+
+private:
+	/** Number of points used for the rbf interpolation */
+	UE_DEPRECATED(5.3, "Please do not access this member directly; use UGroomBindingAsset::GetMatchingSection() or UGroomBindingAsset::SetMatchingSection().")
+	UPROPERTY(VisibleAnywhere, BlueprintGetter = GetMatchingSection, BlueprintSetter = SetMatchingSection, Category = "BuildSettings")
 	int32 MatchingSection = 0;
 
-	UPROPERTY(EditAnywhere, EditFixedSize, BlueprintReadWrite, Category = "HairGroups", meta = (DisplayName = "Group"))
+public:
+	static FName GetMatchingSectionMemberName();
+	UFUNCTION(BlueprintGetter) int32 GetMatchingSection() const;
+	UFUNCTION(BlueprintSetter) void SetMatchingSection(int32 InMatchingSection);
+
+private:
+	UE_DEPRECATED(5.3, "Please do not access this member directly; use UGroomBindingAsset::GetGroupInfos() or UGroomBindingAsset::SetGroupInfos().")
+	UPROPERTY(EditAnywhere, EditFixedSize, BlueprintGetter = GetGroupInfos, BlueprintSetter = SetGroupInfos, Category = "HairGroups", meta = (DisplayName = "Group"))
 	TArray<FGoomBindingGroupInfo> GroupInfos;
 
+public:
+	static FName GetGroupInfosMemberName();
+	UFUNCTION(BlueprintGetter) const TArray<FGoomBindingGroupInfo>& GetGroupInfos() const;
+	UFUNCTION(BlueprintSetter) void SetGroupInfos(const TArray<FGoomBindingGroupInfo>& InGroupInfos);
+	TArray<FGoomBindingGroupInfo>& GetGroupInfos();
+
 	/** GPU and CPU binding data for both simulation and rendering. */
 	struct FHairGroupResource
 	{
@@ -103,11 +166,16 @@ public:
 		TArray<FHairStrandsRestRootResource*> CardsRootResources;
 	};
 	typedef TArray<FHairGroupResource> FHairGroupResources;
+
+private:
+	UE_DEPRECATED(5.3, "Please do not access this member directly; use UGroomBindingAsset::GetHairGroupResources() or UGroomBindingAsset::SetHairGroupResources().")
 	FHairGroupResources HairGroupResources;
 
-	/** Queue of resources which needs to be deleted. This queue is needed for keeping valid pointer on the group resources 
-	   when the binding asset is recomputed */
-	TQueue<FHairGroupResource> HairGroupResourcesToDelete;
+public:
+	static FName GetHairGroupResourcesMemberName();
+	FHairGroupResources& GetHairGroupResources();
+	const FHairGroupResources& GetHairGroupResources() const;
+	void SetHairGroupResources(FHairGroupResources InHairGroupResources);
 
 	/** Binding bulk data */
 	struct FHairGroupPlatformData
@@ -115,11 +183,27 @@ public:
 		FHairStrandsRootBulkData		SimRootBulkData;
 		FHairStrandsRootBulkData		RenRootBulkData;
 		TArray<FHairStrandsRootBulkData>CardsRootBulkData;
-
 	};
-	/** Root bulk data for each hair gruops */
+
+private:
+	/** Queue of resources which needs to be deleted. This queue is needed for keeping valid pointer on the group resources 
+	   when the binding asset is recomputed */
+	UE_DEPRECATED(5.3, "Please do not access this member directly; use UGroomBindingAsset::GetHairGroupPlatformData().")
+	TQueue<FHairGroupResource> HairGroupResourcesToDelete;
+
+	/** Platform data for each hair groups */
+	UE_DEPRECATED(5.3, "Please do not access this member directly; use UGroomBindingAsset::GetHairGroupPlatformData().")
 	TArray<FHairGroupPlatformData> HairGroupsPlatformData;
 
+public:
+	void AddHairGroupResourcesToDelete(FHairGroupResource& In);
+	bool RemoveHairGroupResourcesToDelete(FHairGroupResource& Out);
+
+	static FName GetHairGroupPlatformDataMemberName();
+	const TArray<FHairGroupPlatformData>& GetHairGroupsPlatformData() const;
+	TArray<FHairGroupPlatformData>& GetHairGroupsPlatformData();
+
+public:
 	//~ Begin UObject Interface.
 	virtual void GetResourceSizeEx(FResourceSizeEx& CumulativeResourceSize) override;
 	virtual void PostLoad() override;
diff --git a/Engine/Plugins/Runtime/HairStrands/Source/HairStrandsEditor/Private/GroomActions.cpp b/Engine/Plugins/Runtime/HairStrands/Source/HairStrandsEditor/Private/GroomActions.cpp
--- a/Engine/Plugins/Runtime/HairStrands/Source/HairStrandsEditor/Private/GroomActions.cpp
+++ b/Engine/Plugins/Runtime/HairStrands/Source/HairStrandsEditor/Private/GroomActions.cpp
@@ -223,7 +223,7 @@ void ExecuteCreateBindingAsset(const FToolMenuContext& InContext)
 				if (BindingAsset)
 				{
 					BindingAsset->Build();
-					if (BindingAsset->bIsValid)
+					if (BindingAsset->IsValid())
 					{
 						TArray<UObject*> CreatedObjects;
 						CreatedObjects.Add(BindingAsset);
diff --git a/Engine/Plugins/Runtime/HairStrands/Source/HairStrandsEditor/Private/GroomBindingActions.cpp b/Engine/Plugins/Runtime/HairStrands/Source/HairStrandsEditor/Private/GroomBindingActions.cpp
--- a/Engine/Plugins/Runtime/HairStrands/Source/HairStrandsEditor/Private/GroomBindingActions.cpp
+++ b/Engine/Plugins/Runtime/HairStrands/Source/HairStrandsEditor/Private/GroomBindingActions.cpp
@@ -28,23 +28,23 @@ void ExecuteRebuildBindingAsset(const FToolMenuContext& InContext)
 	{
 		for (UGroomBindingAsset* BindingAsset : Context->LoadSelectedObjects<UGroomBindingAsset>())
 		{
-			if (BindingAsset->Groom && BindingAsset->HasValidTarget())
+			if (BindingAsset->GetGroom() && BindingAsset->HasValidTarget())
 			{
-				BindingAsset->Groom->ConditionalPostLoad();
-				if (BindingAsset->GroomBindingType == EGroomBindingMeshType::SkeletalMesh)
+				BindingAsset->GetGroom()->ConditionalPostLoad();
+				if (BindingAsset->GetGroomBindingType() == EGroomBindingMeshType::SkeletalMesh)
 				{
-					BindingAsset->TargetSkeletalMesh->ConditionalPostLoad();
-					if (BindingAsset->SourceSkeletalMesh)
+					BindingAsset->GetTargetSkeletalMesh()->ConditionalPostLoad();
+					if (BindingAsset->GetSourceSkeletalMesh())
 					{
-						BindingAsset->SourceSkeletalMesh->ConditionalPostLoad();
+						BindingAsset->GetSourceSkeletalMesh()->ConditionalPostLoad();
 					}
 				}
 				else
 				{
-					BindingAsset->TargetGeometryCache->ConditionalPostLoad();
-					if (BindingAsset->SourceGeometryCache)
+					BindingAsset->GetTargetGeometryCache()->ConditionalPostLoad();
+					if (BindingAsset->GetSourceGeometryCache())
 					{
-						BindingAsset->SourceGeometryCache->ConditionalPostLoad();
+						BindingAsset->GetSourceGeometryCache()->ConditionalPostLoad();
 					}
 				}
 				FGroomBindingBuilder::BuildBinding(BindingAsset, true);
diff --git a/Engine/Plugins/Runtime/HairStrands/Source/HairStrandsEditor/Private/GroomBindingDetailsCustomization.cpp b/Engine/Plugins/Runtime/HairStrands/Source/HairStrandsEditor/Private/GroomBindingDetailsCustomization.cpp
--- a/Engine/Plugins/Runtime/HairStrands/Source/HairStrandsEditor/Private/GroomBindingDetailsCustomization.cpp
+++ b/Engine/Plugins/Runtime/HairStrands/Source/HairStrandsEditor/Private/GroomBindingDetailsCustomization.cpp
@@ -8,19 +8,19 @@
 // Customization for UGroomBindingAsset
 void FGroomBindingDetailsCustomization::CustomizeDetails(IDetailLayoutBuilder& LayoutBuilder)
 {
-	TSharedRef<IPropertyHandle> GroomBindingType = LayoutBuilder.GetProperty(GET_MEMBER_NAME_CHECKED(UGroomBindingAsset, GroomBindingType));
+	TSharedRef<IPropertyHandle> GroomBindingType = LayoutBuilder.GetProperty(UGroomBindingAsset::GetGroomBindingTypeMemberName());
 	
 	uint8 EnumValue;
 	GroomBindingType->GetValue(EnumValue);
 	if (EnumValue == (uint8) EGroomBindingMeshType::SkeletalMesh)
 	{
-		LayoutBuilder.HideProperty(GET_MEMBER_NAME_CHECKED(UGroomBindingAsset, SourceGeometryCache));
-		LayoutBuilder.HideProperty(GET_MEMBER_NAME_CHECKED(UGroomBindingAsset, TargetGeometryCache));
+		LayoutBuilder.HideProperty(UGroomBindingAsset::GetSourceGeometryCacheMemberName());
+		LayoutBuilder.HideProperty(UGroomBindingAsset::GetTargetGeometryCacheMemberName());
 	}
 	else
 	{
-		LayoutBuilder.HideProperty(GET_MEMBER_NAME_CHECKED(UGroomBindingAsset, SourceSkeletalMesh));
-		LayoutBuilder.HideProperty(GET_MEMBER_NAME_CHECKED(UGroomBindingAsset, TargetSkeletalMesh));
+		LayoutBuilder.HideProperty(UGroomBindingAsset::GetSourceSkeletalMeshMemberName());
+		LayoutBuilder.HideProperty(UGroomBindingAsset::GetTargetSkeletalMeshMemberName());
 	}
 }
 
diff --git a/Engine/Plugins/Runtime/HairStrands/Source/HairStrandsEditor/Private/GroomBindingFactory.cpp b/Engine/Plugins/Runtime/HairStrands/Source/HairStrandsEditor/Private/GroomBindingFactory.cpp
--- a/Engine/Plugins/Runtime/HairStrands/Source/HairStrandsEditor/Private/GroomBindingFactory.cpp
+++ b/Engine/Plugins/Runtime/HairStrands/Source/HairStrandsEditor/Private/GroomBindingFactory.cpp
@@ -18,11 +18,11 @@ UGroomBindingFactory::UGroomBindingFactory()
 UObject* UGroomBindingFactory::FactoryCreateNew(UClass* InClass, UObject* InParent, FName InName, EObjectFlags Flags, UObject* Context, FFeedbackContext* Warn)
 {
 	UGroomBindingAsset* GroomBinding = NewObject<UGroomBindingAsset>(InParent, InName, Flags | RF_Transactional);
-	GroomBinding->Groom = nullptr;
-	GroomBinding->TargetSkeletalMesh = nullptr;
-	GroomBinding->SourceSkeletalMesh = nullptr;
-	GroomBinding->NumInterpolationPoints = 100;
-	GroomBinding->MatchingSection = 0;
+	GroomBinding->SetGroom(nullptr);
+	GroomBinding->SetTargetSkeletalMesh(nullptr);
+	GroomBinding->SetSourceSkeletalMesh(nullptr);
+	GroomBinding->SetNumInterpolationPoints(100);
+	GroomBinding->SetMatchingSection(0);
 
 	return GroomBinding;
 }
diff --git a/Engine/Plugins/Runtime/HairStrands/Source/HairStrandsEditor/Private/GroomCustomAssetEditorToolkit.cpp b/Engine/Plugins/Runtime/HairStrands/Source/HairStrandsEditor/Private/GroomCustomAssetEditorToolkit.cpp
--- a/Engine/Plugins/Runtime/HairStrands/Source/HairStrandsEditor/Private/GroomCustomAssetEditorToolkit.cpp
+++ b/Engine/Plugins/Runtime/HairStrands/Source/HairStrandsEditor/Private/GroomCustomAssetEditorToolkit.cpp
@@ -331,7 +331,7 @@ static void ListAllBindingAssets(const UGroomAsset* InGroomAsset, TWeakObjectPtr
 
 	// 1. Use tagged properties for searching compatible binding assets (for assets saved after AssetRegistrySearchable has been added)
 	{
-		const FName GroomAssetProperty = GET_MEMBER_NAME_CHECKED(UGroomBindingAsset, Groom);
+		const FName GroomAssetProperty = UGroomBindingAsset::GetGroomMemberName();
 		FString GroomAssetName = FAssetData(InGroomAsset).GetExportTextName();
 		FARFilter Filter;
 		Filter.ClassPaths.Add(UGroomBindingAsset::StaticClass()->GetClassPathName());
@@ -342,7 +342,7 @@ static void ListAllBindingAssets(const UGroomAsset* InGroomAsset, TWeakObjectPtr
 		{
 			if (UGroomBindingAsset* Binding = (UGroomBindingAsset*)Asset.GetAsset(/*bLoad*/))
 			{
-				if (Binding->Groom == InGroomAsset)
+				if (Binding->GetGroom() == InGroomAsset)
 				{
 					Out->Bindings.Add(Binding);
 				}
@@ -364,7 +364,7 @@ static void ListAllBindingAssets(const UGroomAsset* InGroomAsset, TWeakObjectPtr
 			UGroomBindingAsset* Binding = (UGroomBindingAsset*)Asset.FastGetAsset(false /*bLoad*/);
 			if (Binding)
 			{
-				if (Binding->Groom == InGroomAsset)
+				if (Binding->GetGroom() == InGroomAsset)
 				{
 					Out->Bindings.Add(Binding);
 				}
@@ -377,7 +377,7 @@ static void ListAllBindingAssets(const UGroomAsset* InGroomAsset, TWeakObjectPtr
 				if (TCString<TCHAR>::Strfind(*BindingAssetName, *GroomAssetName, true))
 				{
 					Binding = (UGroomBindingAsset*)Asset.GetAsset(/*bLoad*/);
-					if (Binding->Groom == InGroomAsset)
+					if (Binding->GetGroom() == InGroomAsset)
 					{
 						Out->Bindings.Add(Binding);
 					}
@@ -959,7 +959,7 @@ void FGroomCustomAssetEditorToolkit::PreviewBinding(int32 BindingIndex)
 	if (GroomBindingAsset.Get())
 	{
 		PreviewSkeletalMeshComponent = NewObject<USkeletalMeshComponent>(GetTransientPackage(), NAME_None, RF_Transient);
-		PreviewSkeletalMeshComponent->SetSkeletalMesh(GroomBindingAsset->TargetSkeletalMesh);
+		PreviewSkeletalMeshComponent->SetSkeletalMesh(GroomBindingAsset->GetTargetSkeletalMesh());
 		PreviewSkeletalMeshComponent->SetVisibility(true);
 		PreviewSkeletalMeshComponent->Activate(true);
 		PreviewGroomComponent->AttachToComponent(PreviewSkeletalMeshComponent.Get(), FAttachmentTransformRules::KeepRelativeTransform);
