Deprecate direct access to groom members.

This prepares for async building/loading.

diff --git a/Engine/Plugins/Importers/USDImporter/Source/USDSchemas/Private/USDGroomTranslator.cpp b/Engine/Plugins/Importers/USDImporter/Source/USDSchemas/Private/USDGroomTranslator.cpp
--- a/Engine/Plugins/Importers/USDImporter/Source/USDSchemas/Private/USDGroomTranslator.cpp
+++ b/Engine/Plugins/Importers/USDImporter/Source/USDSchemas/Private/USDGroomTranslator.cpp
@@ -247,7 +247,7 @@ protected:
 					GroomCacheProcessor = MakeUnique<FGroomCacheProcessor>(EGroomCacheType::Strands, AnimInfo.Attributes);
 
 					// ref. FGroomCacheImporter::ImportGroomCache
-					const TArray<FHairGroupPlatformData>& GroomHairGroupsData = GroomAsset->HairGroupsPlatformData;
+					const TArray<FHairGroupPlatformData>& GroomHairGroupsData = GroomAsset->GetHairGroupsPlatformData();
 
 					// Each frame is translated into a HairDescription and processed into HairGroupData
 					// Sample one extra frame so that we can interpolate between EndFrame - 1 and EndFrame
@@ -270,14 +270,14 @@ protected:
 
 						const uint32 GroupCount = HairDescriptionGroups.HairGroups.Num();
 
-						TArray<FHairGroupInfoWithVisibility> HairGroupsInfo = GroomAsset->HairGroupsInfo;
+						TArray<FHairGroupInfoWithVisibility> HairGroupsInfo = GroomAsset->GetHairGroupsInfo();
 						TArray<FHairDescriptionGroup> HairGroupsData;
 						HairGroupsData.SetNum(GroupCount);
 						for (uint32 GroupIndex = 0; GroupIndex < GroupCount; ++GroupIndex)
 						{
 							const FHairDescriptionGroup& HairGroup = HairDescriptionGroups.HairGroups[GroupIndex];
 							FHairDescriptionGroup& HairGroupData = HairGroupsData[GroupIndex];
-							FGroomBuilder::BuildData(HairGroup, GroomAsset->HairGroupsInterpolation[GroupIndex], HairGroupsInfo[GroupIndex], HairGroupData.Strands, HairGroupData.Guides);
+							FGroomBuilder::BuildData(HairGroup, GroomAsset->GetHairGroupsInterpolation()[GroupIndex], HairGroupsInfo[GroupIndex], HairGroupData.Strands, HairGroupData.Guides);
 						}
 
 						// Validate that the GroomCache has the same topology as the static groom
diff --git a/Engine/Plugins/Importers/USDImporter/Source/USDSchemas/Private/USDGroomTranslatorUtils.cpp b/Engine/Plugins/Importers/USDImporter/Source/USDSchemas/Private/USDGroomTranslatorUtils.cpp
--- a/Engine/Plugins/Importers/USDImporter/Source/USDSchemas/Private/USDGroomTranslatorUtils.cpp
+++ b/Engine/Plugins/Importers/USDImporter/Source/USDSchemas/Private/USDGroomTranslatorUtils.cpp
@@ -123,7 +123,7 @@ namespace UE::UsdGroomTranslatorUtils::Private
 				GroomBinding->SetSourceGeometryCache(Cast<UGeometryCache>(Settings.SourceMesh));
 				GroomBinding->SetTargetGeometryCache(Cast<UGeometryCache>(Settings.TargetMesh));
 			}
-			GroomBinding->GetHairGroupsPlatformData().Reserve(Settings.Groom->HairGroupsPlatformData.Num());
+			GroomBinding->GetHairGroupsPlatformData().Reserve(Settings.Groom->GetHairGroupsPlatformData().Num());
 			GroomBinding->SetNumInterpolationPoints(Settings.NumInterpolationPoints);
 			GroomBinding->SetMatchingSection(Settings.MatchingSection);
 
diff --git a/Engine/Plugins/Runtime/HairStrands/Source/HairStrandsCore/Private/GroomAsset.cpp b/Engine/Plugins/Runtime/HairStrands/Source/HairStrandsCore/Private/GroomAsset.cpp
--- a/Engine/Plugins/Runtime/HairStrands/Source/HairStrandsCore/Private/GroomAsset.cpp
+++ b/Engine/Plugins/Runtime/HairStrands/Source/HairStrandsCore/Private/GroomAsset.cpp
@@ -160,18 +160,18 @@ void DumpLoadedGroomAssets(IConsoleVariable* InCVarPakTesterEnabled)
 	{
 		if (AssetIt)
 		{
-			const uint32 GroupCount = AssetIt->HairGroupsPlatformData.Num();
+			const uint32 GroupCount = AssetIt->GetHairGroupsPlatformData().Num();
 			for (uint32 GroupIt = 0; GroupIt < GroupCount; ++GroupIt)
 			{
-				const uint32 CPUMemorySize_Guides  = AssetIt->HairGroupsPlatformData[GroupIt].Guides.GetDataSize();
-				const uint32 CPUMemorySize_Strands = AssetIt->HairGroupsPlatformData[GroupIt].Strands.GetDataSize();
-				const uint32 CPUMemorySize_Cards   = AssetIt->HairGroupsPlatformData[GroupIt].Cards.GetDataSize();
-				const uint32 CPUMemorySize_Meshes  = AssetIt->HairGroupsPlatformData[GroupIt].Meshes.GetDataSize();
+				const uint32 CPUMemorySize_Guides  = AssetIt->GetHairGroupsPlatformData()[GroupIt].Guides.GetDataSize();
+				const uint32 CPUMemorySize_Strands = AssetIt->GetHairGroupsPlatformData()[GroupIt].Strands.GetDataSize();
+				const uint32 CPUMemorySize_Cards   = AssetIt->GetHairGroupsPlatformData()[GroupIt].Cards.GetDataSize();
+				const uint32 CPUMemorySize_Meshes  = AssetIt->GetHairGroupsPlatformData()[GroupIt].Meshes.GetDataSize();
 
-				const uint32 GPUMemorySize_Guides  = AssetIt->HairGroupsPlatformData[GroupIt].Guides.GetResourcesSize();
-				const uint32 GPUMemorySize_Strands = AssetIt->HairGroupsPlatformData[GroupIt].Strands.GetResourcesSize();
-				const uint32 GPUMemorySize_Cards   = AssetIt->HairGroupsPlatformData[GroupIt].Cards.GetResourcesSize();
-				const uint32 GPUMemorySize_Meshes  = AssetIt->HairGroupsPlatformData[GroupIt].Meshes.GetResourcesSize();
+				const uint32 GPUMemorySize_Guides  = AssetIt->GetHairGroupsPlatformData()[GroupIt].Guides.GetResourcesSize();
+				const uint32 GPUMemorySize_Strands = AssetIt->GetHairGroupsPlatformData()[GroupIt].Strands.GetResourcesSize();
+				const uint32 GPUMemorySize_Cards   = AssetIt->GetHairGroupsPlatformData()[GroupIt].Cards.GetResourcesSize();
+				const uint32 GPUMemorySize_Meshes  = AssetIt->GetHairGroupsPlatformData()[GroupIt].Meshes.GetResourcesSize();
 
 				uint32 CPUMemorySize = 0;
 				uint32 GPUMemorySize = 0;
@@ -194,7 +194,7 @@ void DumpLoadedGroomAssets(IConsoleVariable* InCVarPakTesterEnabled)
 				Total_GPUMemorySize_Cards  += GPUMemorySize_Cards;
 				Total_GPUMemorySize_Meshes += GPUMemorySize_Meshes;
 
-				const uint32 LODCount = AssetIt->HairGroupsLOD[GroupIt].LODs.Num();
+				const uint32 LODCount = AssetIt->GetHairGroupsLOD()[GroupIt].LODs.Num();
 				if (bDetails)
 				{
 //					UE_LOG(LogHairStrands, Log, TEXT("--  No.  - LOD -    CPU Total (     Guides|    Strands|      Cards|     Meshes) -    GPU Total (     Guides|    Strands|      Cards|     Meshes) - Asset Name "));
@@ -495,8 +495,8 @@ static bool BuildHairGroup(
 	const int32 GroupIndex,
 	const bool bNeedInterpolationData,
 	const FHairDescriptionGroups& HairDescriptionGroups, 
-	const TArray<FHairGroupsInterpolation>& HairGroupsInterpolation,
-	const TArray<FHairGroupsLOD>& HairGroupsLOD,
+	const TArray<FHairGroupsInterpolation>& InHairGroupsInterpolation,
+	const TArray<FHairGroupsLOD>& InHairGroupsLOD,
 	TArray<FHairGroupInfoWithVisibility>& OutHairGroupsInfo,
 	TArray<FHairGroupPlatformData>& OutHairGroupsData)
 {
@@ -509,7 +509,7 @@ static bool BuildHairGroup(
 
 		FHairStrandsDatas StrandsData;
 		FHairStrandsDatas GuidesData;
-		FGroomBuilder::BuildData(HairGroup, HairGroupsInterpolation[GroupIndex], OutHairGroupsInfo[GroupIndex], StrandsData, GuidesData);
+		FGroomBuilder::BuildData(HairGroup, InHairGroupsInterpolation[GroupIndex], OutHairGroupsInfo[GroupIndex], StrandsData, GuidesData);
 
 		FGroomBuilder::BuildBulkData(HairGroup.Info, GuidesData, OutHairGroupsData[GroupIndex].Guides.BulkData);
 		FGroomBuilder::BuildBulkData(HairGroup.Info, StrandsData, OutHairGroupsData[GroupIndex].Strands.BulkData);
@@ -518,7 +518,7 @@ static bool BuildHairGroup(
 		if (bNeedInterpolationData)
 		{
 			FHairStrandsInterpolationDatas InterpolationData;
-			FGroomBuilder::BuildInterplationData(HairGroup.Info, StrandsData, GuidesData, HairGroupsInterpolation[GroupIndex].InterpolationSettings, InterpolationData);
+			FGroomBuilder::BuildInterplationData(HairGroup.Info, StrandsData, GuidesData, InHairGroupsInterpolation[GroupIndex].InterpolationSettings, InterpolationData);
 			FGroomBuilder::BuildInterplationBulkData(GuidesData, InterpolationData, OutHairGroupsData[GroupIndex].Strands.InterpolationBulkData);
 		}
 		else
@@ -526,7 +526,7 @@ static bool BuildHairGroup(
 			OutHairGroupsData[GroupIndex].Strands.InterpolationBulkData.Reset();
 		}
 
-		FGroomBuilder::BuildClusterBulkData(StrandsData, HairDescriptionGroups.Bounds.SphereRadius, HairGroupsLOD[GroupIndex], OutHairGroupsData[GroupIndex].Strands.ClusterCullingBulkData);
+		FGroomBuilder::BuildClusterBulkData(StrandsData, HairDescriptionGroups.Bounds.SphereRadius, InHairGroupsLOD[GroupIndex], OutHairGroupsData[GroupIndex].Strands.ClusterCullingBulkData);
 	}
 	return bIsValid;
 }
@@ -537,23 +537,23 @@ static bool BuildHairGroup(
 // Guides and strands data are all transient
 static bool BuildHairGroupCluster(
 	const int32 GroupIndex,
-	const FHairDescriptionGroups& HairDescriptionGroups,
-	const TArray<FHairGroupsInterpolation>& HairGroupsInterpolation,
-	const TArray<FHairGroupsLOD>& HairGroupsLOD,
+	const FHairDescriptionGroups& InHairDescriptionGroups,
+	const TArray<FHairGroupsInterpolation>& InHairGroupsInterpolation,
+	const TArray<FHairGroupsLOD>& InHairGroupsLOD,
 	TArray<FHairGroupInfoWithVisibility>& OutHairGroupsInfo,
 	TArray<FHairGroupPlatformData>& OutHairGroupsData)
 {
-	const bool bIsValid = HairDescriptionGroups.IsValid();
+	const bool bIsValid = InHairDescriptionGroups.IsValid();
 	if (bIsValid)
 	{
-		const FHairDescriptionGroup& HairGroup = HairDescriptionGroups.HairGroups[GroupIndex];
-		check(GroupIndex <= HairDescriptionGroups.HairGroups.Num());
+		const FHairDescriptionGroup& HairGroup = InHairDescriptionGroups.HairGroups[GroupIndex];
+		check(GroupIndex <= InHairDescriptionGroups.HairGroups.Num());
 		check(GroupIndex == HairGroup.Info.GroupID);
 
 		FHairStrandsDatas StrandsData;
 		FHairStrandsDatas GuidesData;
-		FGroomBuilder::BuildData(HairGroup, HairGroupsInterpolation[GroupIndex], OutHairGroupsInfo[GroupIndex], StrandsData, GuidesData);
-		FGroomBuilder::BuildClusterBulkData(StrandsData, HairDescriptionGroups.Bounds.SphereRadius, HairGroupsLOD[GroupIndex], OutHairGroupsData[GroupIndex].Strands.ClusterCullingBulkData);
+		FGroomBuilder::BuildData(HairGroup, InHairGroupsInterpolation[GroupIndex], OutHairGroupsInfo[GroupIndex], StrandsData, GuidesData);
+		FGroomBuilder::BuildClusterBulkData(StrandsData, InHairDescriptionGroups.Bounds.SphereRadius, InHairGroupsLOD[GroupIndex], OutHairGroupsData[GroupIndex].Strands.ClusterCullingBulkData);
 	}
 	return bIsValid;
 }
@@ -583,17 +583,17 @@ uint8 UGroomAsset::GenerateClassStripFlags(FArchive& Ar)
 		}
 
 		// Determine the platform min LOD if the stripping hasn't been disabled
-		const bool bDisableMinLODStrip = DisableBelowMinLodStripping.GetValueForPlatform(*CookTarget->IniPlatformName());
+		const bool bDisableMinLODStrip = GetDisableBelowMinLodStripping().GetValueForPlatform(*CookTarget->IniPlatformName());
 		int32 PlatformMinLOD = 0;
 		if (!bDisableMinLODStrip)
 		{
-			PlatformMinLOD = MinLOD.GetValueForPlatform(*CookTarget->IniPlatformName());
+			PlatformMinLOD = GetMinLOD().GetValueForPlatform(*CookTarget->IniPlatformName());
 		}
 
 		bIsLODStripped = (PlatformMinLOD > 0) && !bDisableMinLODStrip;
 		if (bIsLODStripped)
 		{
-			for (const FHairGroupsLOD& GroupsLOD : HairGroupsLOD)
+			for (const FHairGroupsLOD& GroupsLOD : GetHairGroupsLOD())
 			{
 				for (int32 LODIndex = 0; LODIndex < PlatformMinLOD && LODIndex < GroupsLOD.LODs.Num(); ++LODIndex)
 				{
@@ -638,19 +638,19 @@ void UGroomAsset::ApplyStripFlags(uint8 StripFlags, const ITargetPlatform* CookT
 		return;
 	}
 
-	const bool bDisableMinLodStrip = DisableBelowMinLodStripping.GetValueForPlatform(*CookTarget->IniPlatformName()); 
+	const bool bDisableMinLodStrip = GetDisableBelowMinLodStripping().GetValueForPlatform(*CookTarget->IniPlatformName()); 
 	int32 PlatformMinLOD = 0;
 	if (!bDisableMinLodStrip)
 	{
-		PlatformMinLOD = MinLOD.GetValueForPlatform(*CookTarget->IniPlatformName());
+		PlatformMinLOD = GetMinLOD().GetValueForPlatform(*CookTarget->IniPlatformName());
 	}
 
 	const bool bIsStrandsStrippedForCook = !!(StripFlags & CDSF_StrandsStripped);
 
 	// Set the CookedOut flags as appropriate and compute the value for EffectiveLODBias
-	for (int32 GroupIndex = 0; GroupIndex < HairGroupsPlatformData.Num(); ++GroupIndex)
+	for (int32 GroupIndex = 0; GroupIndex < GetHairGroupsPlatformData().Num(); ++GroupIndex)
 	{
-		FHairGroupPlatformData& GroupData = HairGroupsPlatformData[GroupIndex];
+		FHairGroupPlatformData& GroupData = GetHairGroupsPlatformData()[GroupIndex];
 		GroupData.Strands.bIsCookedOut = bIsStrandsStrippedForCook;
 
 		// Determine the max LOD for strands since they are cooked out if the platform doesn't support them
@@ -658,7 +658,7 @@ void UGroomAsset::ApplyStripFlags(uint8 StripFlags, const ITargetPlatform* CookT
 		int32 StrandsMaxLOD = -1;
 		if (bIsStrandsStrippedForCook)
 		{
-			const FHairGroupsLOD& GroupsLOD = HairGroupsLOD[GroupIndex];
+			const FHairGroupsLOD& GroupsLOD = GetHairGroupsLOD()[GroupIndex];
 			for (int32 LODIndex = 0; LODIndex < GroupsLOD.LODs.Num(); ++LODIndex)
 			{
 				const FHairLODSettings& LODSetting = GroupsLOD.LODs[LODIndex];
@@ -671,7 +671,7 @@ void UGroomAsset::ApplyStripFlags(uint8 StripFlags, const ITargetPlatform* CookT
 		
 		// The MinLOD for stripping this group, taking into account the possibility of strands being cooked out
 		const int32 MinStrippedLOD = FMath::Max(StrandsMaxLOD + 1, PlatformMinLOD);
-		EffectiveLODBias[GroupIndex] = MinStrippedLOD;
+		GetEffectiveLODBias()[GroupIndex] = MinStrippedLOD;
 
 		for (int32 Index = 0; Index < GroupData.Cards.LODs.Num(); ++Index)
 		{
@@ -790,7 +790,7 @@ void UGroomAsset::Serialize(FArchive& Ar)
 	// For older assets, force width override to preserve prior to fixing the imported groom's width.
 	if (Ar.CustomVer(FUE5MainStreamObjectVersion::GUID) < FUE5MainStreamObjectVersion::GroomAssetWidthOverride && Ar.IsLoading())
 	{
-		for (FHairGroupsRendering& Group : HairGroupsRendering)
+		for (FHairGroupsRendering& Group : GetHairGroupsRendering())
 		{
 			Group.GeometrySettings.HairWidth_Override = true;
 		}
@@ -812,7 +812,7 @@ void UGroomAsset::Serialize(FArchive& Ar)
 		if (StripFlags.IsEditorDataStripped())
 		{
 			// When cooking data or loading cooked data
-			InternalSerializePlatformDatas(Ar, this, HairGroupsPlatformData);
+			InternalSerializePlatformDatas(Ar, this, GetHairGroupsPlatformData());
 		}
 		else if (Ar.CustomVer(FUE5MainStreamObjectVersion::GUID) < FUE5MainStreamObjectVersion::GroomAssetRemoveInAssetSerialization && Ar.IsLoading() && !bHasImportedStrands)
 		{
@@ -850,17 +850,17 @@ UGroomAsset::UGroomAsset(const FObjectInitializer& ObjectInitializer)
 {
 	bIsInitialized = false;
 
-	MinLOD.Default = 0;
-	DisableBelowMinLodStripping.Default = false;
+	SetMinLOD(0);
+	SetDisableBelowMinLodStripping(false);
 }
 
 bool UGroomAsset::HasGeometryType(uint32 GroupIndex, EGroomGeometryType Type) const
 {
-	check(GroupIndex < uint32(HairGroupsLOD.Num()));
-	const uint32 LODCount = HairGroupsLOD[GroupIndex].LODs.Num();
+	check(GroupIndex < uint32(GetHairGroupsLOD().Num()));
+	const uint32 LODCount = GetHairGroupsLOD()[GroupIndex].LODs.Num();
 	for (uint32 LODIt=0;LODIt<LODCount;++LODIt)
 	{
-		if (HairGroupsLOD[GroupIndex].LODs[LODIt].GeometryType == Type)
+		if (GetHairGroupsLOD()[GroupIndex].LODs[LODIt].GeometryType == Type)
 			return true;
 	}
 	return false;
@@ -868,7 +868,7 @@ bool UGroomAsset::HasGeometryType(uint32 GroupIndex, EGroomGeometryType Type) co
 
 bool UGroomAsset::HasGeometryType(EGroomGeometryType Type) const
 {
-	const uint32 GroupCount = HairGroupsLOD.Num();
+	const uint32 GroupCount = GetHairGroupsLOD().Num();
 	for (uint32 GroupIt = 0; GroupIt < GroupCount; ++GroupIt)
 	{
 		if (HasGeometryType(GroupIt, Type))
@@ -901,11 +901,11 @@ TArray<FHairVertexFactoryTypesPerMaterialData> UGroomAsset::CollectVertexFactory
 	const int32 GroupCount = GetNumHairGroups();
 	for (int32 GroupIt = 0; GroupIt < GroupCount; GroupIt++)
 	{
-		const FHairGroupPlatformData& InGroupData = HairGroupsPlatformData[GroupIt];
+		const FHairGroupPlatformData& InGroupData = GetHairGroupsPlatformData()[GroupIt];
 
 		// Strands
 		{
-			const int32 MaterialIndex = GetMaterialIndex(HairGroupsRendering[GroupIt].MaterialSlotName);
+			const int32 MaterialIndex = GetMaterialIndex(GetHairGroupsRendering()[GroupIt].MaterialSlotName);
 			AddVFPerMaterial(MaterialIndex, EHairGeometryType::Strands, &FHairStrandsVertexFactory::StaticType);
 		}
 
@@ -919,7 +919,7 @@ TArray<FHairVertexFactoryTypesPerMaterialData> UGroomAsset::CollectVertexFactory
 				{
 					// Material
 					int32 MaterialIndex = INDEX_NONE;
-					for (const FHairGroupsCardsSourceDescription& Desc : HairGroupsCards)
+					for (const FHairGroupsCardsSourceDescription& Desc : GetHairGroupsCards())
 					{
 						if (Desc.GroupIndex == GroupIt && Desc.LODIndex == CardsLODIndex)
 						{
@@ -942,7 +942,7 @@ TArray<FHairVertexFactoryTypesPerMaterialData> UGroomAsset::CollectVertexFactory
 				{
 					// Material
 					int32 MaterialIndex = INDEX_NONE;
-					for (const FHairGroupsMeshesSourceDescription& Desc : HairGroupsMeshes)
+					for (const FHairGroupsMeshesSourceDescription& Desc : GetHairGroupsMeshes())
 					{
 						if (Desc.GroupIndex == GroupIt && Desc.LODIndex == MeshesLODIndex)
 						{
@@ -1000,11 +1000,11 @@ void UGroomAsset::UpdateResource()
 	uint32 AllChangeType = 0;
 	for (uint32 GroupIndex = 0, GroupCount = GetNumHairGroups(); GroupIndex < GroupCount; ++GroupIndex)
 	{
-		uint32 ChangeType = (CachedHairGroupsLOD[GroupIndex] == HairGroupsLOD[GroupIndex] ? 0 : GroomChangeType_LOD);
+		uint32 ChangeType = (CachedHairGroupsLOD[GroupIndex] == GetHairGroupsLOD()[GroupIndex] ? 0 : GroomChangeType_LOD);
 		AllChangeType = AllChangeType | ChangeType;
 
 		check(GroupIndex < uint32(GetNumHairGroups()));
-		FHairGroupPlatformData& GroupData = HairGroupsPlatformData[GroupIndex];
+		FHairGroupPlatformData& GroupData = GetHairGroupsPlatformData()[GroupIndex];
 		if (GroupData.Guides.RestResource)
 		{
 			InternalUpdateResource(GroupData.Guides.RestResource);
@@ -1026,10 +1026,10 @@ void UGroomAsset::UpdateResource()
 				BuildHairGroupCluster(
 					GroupIndex,
 					LocalHairDescriptionGroups,
-					HairGroupsInterpolation,
-					HairGroupsLOD,
-					HairGroupsInfo,
-					HairGroupsPlatformData);
+					GetHairGroupsInterpolation(),
+					GetHairGroupsLOD(),
+					GetHairGroupsInfo(),
+					GetHairGroupsPlatformData());
 
 				GroupData.Strands.ClusterCullingResource = new FHairStrandsClusterCullingResource(GroupData.Strands.ClusterCullingBulkData, FHairResourceName(GetFName(), GroupIndex), GetAssetPathName());
 			}
@@ -1093,8 +1093,8 @@ void UGroomAsset::ReleaseResource()
 
 void UGroomAsset::ReleaseGuidesResource(uint32 GroupIndex)
 {
-	if (GroupIndex >= uint32(HairGroupsPlatformData.Num())) return;
-	FHairGroupPlatformData& GroupData = HairGroupsPlatformData[GroupIndex];
+	if (GroupIndex >= uint32(GetHairGroupsPlatformData().Num())) return;
+	FHairGroupPlatformData& GroupData = GetHairGroupsPlatformData()[GroupIndex];
 	if (GroupData.Guides.IsValid())
 	{
 		InternalReleaseResource(GroupData.Guides.RestResource);
@@ -1103,8 +1103,8 @@ void UGroomAsset::ReleaseGuidesResource(uint32 GroupIndex)
 
 void UGroomAsset::ReleaseStrandsResource(uint32 GroupIndex)
 {
-	if (GroupIndex >= uint32(HairGroupsPlatformData.Num())) return;
-	FHairGroupPlatformData& GroupData = HairGroupsPlatformData[GroupIndex];
+	if (GroupIndex >= uint32(GetHairGroupsPlatformData().Num())) return;
+	FHairGroupPlatformData& GroupData = GetHairGroupsPlatformData()[GroupIndex];
 	if (GroupData.Strands.IsValid())
 	{
 		InternalReleaseResource(GroupData.Strands.RestResource);
@@ -1118,8 +1118,8 @@ void UGroomAsset::ReleaseStrandsResource(uint32 GroupIndex)
 
 void UGroomAsset::ReleaseCardsResource(uint32 GroupIndex)
 {
-	if (GroupIndex >= uint32(HairGroupsPlatformData.Num())) return;
-	FHairGroupPlatformData& GroupData = HairGroupsPlatformData[GroupIndex];
+	if (GroupIndex >= uint32(GetHairGroupsPlatformData().Num())) return;
+	FHairGroupPlatformData& GroupData = GetHairGroupsPlatformData()[GroupIndex];
 	for (FHairGroupPlatformData::FCards::FLOD& LOD : GroupData.Cards.LODs)
 	{
 		InternalReleaseResource(LOD.RestResource);
@@ -1134,8 +1134,8 @@ void UGroomAsset::ReleaseCardsResource(uint32 GroupIndex)
 
 void UGroomAsset::ReleaseMeshesResource(uint32 GroupIndex)
 {
-	if (GroupIndex >= uint32(HairGroupsPlatformData.Num())) return;
-	FHairGroupPlatformData& GroupData = HairGroupsPlatformData[GroupIndex];
+	if (GroupIndex >= uint32(GetHairGroupsPlatformData().Num())) return;
+	FHairGroupPlatformData& GroupData = GetHairGroupsPlatformData()[GroupIndex];
 	for (FHairGroupPlatformData::FMeshes::FLOD& LOD : GroupData.Meshes.LODs)
 	{
 		InternalReleaseResource(LOD.RestResource);
@@ -1151,13 +1151,13 @@ void UGroomAsset::ReleaseMeshesResource(uint32 GroupIndex)
 void UGroomAsset::UpdateHairGroupsInfo()
 {
 	const uint32 GroupCount = GetNumHairGroups();
-	const bool bForceReset = HairGroupsInfo.Num() != GroupCount;
-	HairGroupsInfo.SetNum(GroupCount);
+	const bool bForceReset = GetHairGroupsInfo().Num() != GroupCount;
+	GetHairGroupsInfo().SetNum(GroupCount);
 
 	uint32 GroupIndex = 0;
-	for (FHairGroupPlatformData& Data: HairGroupsPlatformData)
+	for (FHairGroupPlatformData& Data: GetHairGroupsPlatformData())
 	{
-		FHairGroupInfoWithVisibility& Info = HairGroupsInfo[GroupIndex];
+		FHairGroupInfoWithVisibility& Info = GetHairGroupsInfo()[GroupIndex];
 		Info.GroupID = GroupIndex;
 		Info.NumCurves = Data.Strands.BulkData.GetNumCurves();
 		Info.NumGuides = Data.Guides.BulkData.GetNumCurves();
@@ -1173,7 +1173,7 @@ void UGroomAsset::UpdateHairGroupsInfo()
 }
 
 template<typename T>
-bool ConvertMaterial(TArray<T>& Groups, TArray<FHairGroupsMaterial>& HairGroupsMaterials)
+bool ConvertMaterial(TArray<T>& Groups, TArray<FHairGroupsMaterial>& InHairGroupsMaterials)
 {
 	bool bNeedSaving = false;
 	for (T& Group : Groups)
@@ -1182,7 +1182,7 @@ bool ConvertMaterial(TArray<T>& Groups, TArray<FHairGroupsMaterial>& HairGroupsM
 		{
 			FName MaterialSlotName = Group.Material->GetFName();
 			bool bFound = false;
-			for (FHairGroupsMaterial& Material : HairGroupsMaterials)
+			for (FHairGroupsMaterial& Material : InHairGroupsMaterials)
 			{
 				if (Material.SlotName == MaterialSlotName)
 				{
@@ -1194,7 +1194,7 @@ bool ConvertMaterial(TArray<T>& Groups, TArray<FHairGroupsMaterial>& HairGroupsM
 
 			if (!bFound)
 			{
-				FHairGroupsMaterial& MaterialEntry = HairGroupsMaterials.AddDefaulted_GetRef();
+				FHairGroupsMaterial& MaterialEntry = InHairGroupsMaterials.AddDefaulted_GetRef();
 				MaterialEntry.Material = Group.Material;
 				MaterialEntry.SlotName = MaterialSlotName;
 				Group.MaterialSlotName = MaterialSlotName;
@@ -1215,7 +1215,7 @@ void UGroomAsset::PostLoad()
 
 	// Legacy asset are missing rendering or interpolation settings
 #if WITH_EDITORONLY_DATA
-	const bool bIsLegacyAsset = HairGroupsInterpolation.Num() == 0;
+	const bool bIsLegacyAsset = GetHairGroupsInfo().Num() == 0;
 	if (bIsLegacyAsset)
 	{
 		if (HairDescriptionBulkData[HairDescriptionType])
@@ -1234,7 +1234,7 @@ void UGroomAsset::PostLoad()
 	}
 #endif
 	// Convert old rigging/guide override
-	for (FHairGroupsInterpolation& Group : HairGroupsInterpolation)
+	for (FHairGroupsInterpolation& Group : GetHairGroupsInterpolation())
 	{
 		const bool bNeedLegacyConversion = Group.InterpolationSettings.bOverrideGuides_DEPRECATED || Group.RiggingSettings.bEnableRigging_DEPRECATED;
 		if (bNeedLegacyConversion)
@@ -1263,7 +1263,7 @@ void UGroomAsset::PostLoad()
 	bool bSucceed = true;
 	{
 		// Interpolation settings are used for building the interpolation data, and per se defined the number of groups
-		const uint32 GroupCount = HairGroupsInterpolation.Num();
+		const uint32 GroupCount = GetHairGroupsInterpolation().Num();
 		if (uint32(GetNumHairGroups()) != GroupCount)
 		{
 			SetNumGroup(GroupCount);
@@ -1341,25 +1341,25 @@ void UGroomAsset::PostLoad()
 	{
 		const FHairDescriptionGroups& LocalHairDescriptionGroups = GetHairDescriptionGroups();
 
-		const int32 GroupCount = HairGroupsInterpolation.Num();
+		const int32 GroupCount = GetHairGroupsInterpolation().Num();
 		check(LocalHairDescriptionGroups.HairGroups.Num() == GroupCount);
 		for (int32 GroupIndex = 0; GroupIndex < GroupCount; ++GroupIndex)
 		{
 			const bool bNeedInterpoldationData = NeedsInterpolationData(GroupIndex);
 			const bool bNeedToBuildData = 
-				(HairGroupsPlatformData[GroupIndex].Guides.BulkData.GetNumCurves() == 0 ||
-				 (bNeedInterpoldationData && HairGroupsPlatformData[GroupIndex].Strands.InterpolationBulkData.GetPointCount() == 0)) &&
-				 HairGroupsPlatformData[GroupIndex].Strands.BulkData.GetNumCurves() > 0; // Empty groom has no data to build
+				(GetHairGroupsPlatformData()[GroupIndex].Guides.BulkData.GetNumCurves() == 0 ||
+				 (bNeedInterpoldationData && GetHairGroupsPlatformData()[GroupIndex].Strands.InterpolationBulkData.GetPointCount() == 0)) &&
+				 GetHairGroupsPlatformData()[GroupIndex].Strands.BulkData.GetNumCurves() > 0; // Empty groom has no data to build
 			if (bNeedToBuildData)
 			{
 				BuildHairGroup(
 					GroupIndex,
 					bNeedInterpoldationData,
 					LocalHairDescriptionGroups,
-					HairGroupsInterpolation,
-					HairGroupsLOD,
-					HairGroupsInfo,
-					HairGroupsPlatformData);
+					GetHairGroupsInterpolation(),
+					GetHairGroupsLOD(),
+					GetHairGroupsInfo(),
+					GetHairGroupsPlatformData());
 			}
 		}
 	}
@@ -1369,9 +1369,9 @@ void UGroomAsset::PostLoad()
 	{
 		// Strands
 		bool bNeedSaving = false;
-		bNeedSaving = bNeedSaving || ConvertMaterial(HairGroupsRendering,	HairGroupsMaterials);
-		bNeedSaving = bNeedSaving || ConvertMaterial(HairGroupsCards,		HairGroupsMaterials);
-		bNeedSaving = bNeedSaving || ConvertMaterial(HairGroupsMeshes,		HairGroupsMaterials);
+		bNeedSaving = bNeedSaving || ConvertMaterial(GetHairGroupsRendering(),	GetHairGroupsMaterials());
+		bNeedSaving = bNeedSaving || ConvertMaterial(GetHairGroupsCards(),		GetHairGroupsMaterials());
+		bNeedSaving = bNeedSaving || ConvertMaterial(GetHairGroupsMeshes(),		GetHairGroupsMaterials());
 
 		if (bNeedSaving)
 		{
@@ -1402,7 +1402,7 @@ void UGroomAsset::PostLoad()
 		FPSOPrecacheParams PrecachePSOParams;
 		for (FHairVertexFactoryTypesPerMaterialData& VFsPerMaterial : VFsPerMaterials)
 		{
-			UMaterialInterface* MaterialInterface = HairGroupsMaterials[VFsPerMaterial.MaterialIndex].Material;
+			UMaterialInterface* MaterialInterface = GetHairGroupsMaterials()[VFsPerMaterial.MaterialIndex].Material;
 			if (MaterialInterface)
 			{
 				MaterialInterface->PrecachePSOs(VFsPerMaterial.VertexFactoryDataList, PrecachePSOParams, EPSOPrecachePriority::Medium, RequestIDs);
@@ -1428,7 +1428,7 @@ void UGroomAsset::PreSave(FObjectPreSaveContext ObjectSaveContext)
 
 	for (uint32 GroupIt = 0; GroupIt < GroupCount; ++GroupIt)
 	{
-		if (!(CachedHairGroupsInterpolation[GroupIt] == HairGroupsInterpolation[GroupIt]))
+		if (!(CachedHairGroupsInterpolation[GroupIt] == GetHairGroupsInterpolation()[GroupIt]))
 		{
 			ChangeType = ChangeType | GroomChangeType_Interpolation;
 			break;
@@ -1437,7 +1437,7 @@ void UGroomAsset::PreSave(FObjectPreSaveContext ObjectSaveContext)
 
 	for (uint32 GroupIt = 0; GroupIt < GroupCount; ++GroupIt)
 	{
-		if (!(CachedHairGroupsLOD[GroupIt] == HairGroupsLOD[GroupIt]))
+		if (!(CachedHairGroupsLOD[GroupIt] == GetHairGroupsLOD()[GroupIt]))
 		{
 			ChangeType = ChangeType | GroomChangeType_LOD;
 			break;
@@ -1453,8 +1453,8 @@ void UGroomAsset::PreSave(FObjectPreSaveContext ObjectSaveContext)
 			FGroomComponentRecreateRenderStateContext RecreateRenderContext(this);
 			for (uint32 GroupIt = 0; GroupIt < GroupCount; ++GroupIt)
 			{
-				const bool bHasInterpolationChanged = !(CachedHairGroupsInterpolation[GroupIt] == HairGroupsInterpolation[GroupIt]);
-				const bool bHasLODChanged = !(CachedHairGroupsLOD[GroupIt] == HairGroupsLOD[GroupIt]);
+				const bool bHasInterpolationChanged = !(CachedHairGroupsInterpolation[GroupIt] == GetHairGroupsInterpolation()[GroupIt]);
+				const bool bHasLODChanged = !(CachedHairGroupsLOD[GroupIt] == GetHairGroupsLOD()[GroupIt]);
 				const bool bNeedInterpoldationData = NeedsInterpolationData(GroupIt);
 
 				if (bHasInterpolationChanged)
@@ -1463,20 +1463,20 @@ void UGroomAsset::PreSave(FObjectPreSaveContext ObjectSaveContext)
 						GroupIt,
 						bNeedInterpoldationData,
 						LocalHairDescriptionGroups,
-						HairGroupsInterpolation,
-						HairGroupsLOD,
-						HairGroupsInfo,
-						HairGroupsPlatformData);
+						GetHairGroupsInterpolation(),
+						GetHairGroupsLOD(),
+						GetHairGroupsInfo(),
+						GetHairGroupsPlatformData());
 				}
 				else if (bHasLODChanged)
 				{
 					BuildHairGroupCluster(
 						GroupIt,
 						LocalHairDescriptionGroups,
-						HairGroupsInterpolation,
-						HairGroupsLOD,
-						HairGroupsInfo,
-						HairGroupsPlatformData);
+						GetHairGroupsInterpolation(),
+						GetHairGroupsLOD(),
+						GetHairGroupsInfo(),
+						GetHairGroupsPlatformData());
 				}
 			}
 		}
@@ -1538,8 +1538,8 @@ static bool IsStrandsInterpolationAttributes(const FName PropertyName)
 
 
 		// Add dependency on simulation and per LOD-simulation/global-interpolation to strip-out interpolation data if there are not needed
-		|| PropertyName == GET_MEMBER_NAME_CHECKED(UGroomAsset, EnableGlobalInterpolation)
-		|| PropertyName == GET_MEMBER_NAME_CHECKED(UGroomAsset, EnableSimulationCache)
+		|| PropertyName == UGroomAsset::GetEnableGlobalInterpolationMemberName()
+		|| PropertyName == UGroomAsset::GetEnableSimulationCacheMemberName()
 		|| PropertyName == GET_MEMBER_NAME_CHECKED(FHairSolverSettings, EnableSimulation)
 
 		// Decimation control needs to force a rebuild
@@ -1566,7 +1566,7 @@ static bool IsStrandsLODAttributes(const FName PropertyName)
 {
 	return
 		// LOD count needs for cluster rebuilding
-		PropertyName == GET_MEMBER_NAME_CHECKED(UGroomAsset, HairGroupsLOD)
+		PropertyName == UGroomAsset::GetHairGroupsLODMemberName()
 
 		// LOD groups needs for cluster rebuilding
 		|| PropertyName == GET_MEMBER_NAME_CHECKED(FHairGroupsLOD, ClusterWorldSize)
@@ -1596,12 +1596,12 @@ static bool IsRiggingSettings(const FName PropertyName)
 
 static bool IsCardsAttributes(const FName PropertyName)
 {
-	return PropertyName == GET_MEMBER_NAME_CHECKED(UGroomAsset, HairGroupsCards);
+	return PropertyName == UGroomAsset::GetHairGroupsCardsMemberName();
 }
 
 static bool IsMeshesAttributes(const FName PropertyName)
 {
-	return PropertyName == GET_MEMBER_NAME_CHECKED(UGroomAsset, HairGroupsMeshes);
+	return PropertyName == UGroomAsset::GetHairGroupsMeshesMemberName();
 }
 
 void UGroomAsset::PostEditChangeProperty(FPropertyChangedEvent& PropertyChangedEvent)
@@ -1614,19 +1614,19 @@ void UGroomAsset::PostEditChangeProperty(FPropertyChangedEvent& PropertyChangedE
 	if (bGeometryTypeChanged)
 	{
 		// If the asset didn't have any mesh or cards, we need to create/allocate the cards/mesh groups
-		const uint32 GroupCount = HairGroupsPlatformData.Num();
-		if (IsHairStrandsEnabled(EHairStrandsShaderType::Cards) && HasGeometryType(EGroomGeometryType::Cards) && HairGroupsCards.Num() == 0)
+		const uint32 GroupCount = GetHairGroupsPlatformData().Num();
+		if (IsHairStrandsEnabled(EHairStrandsShaderType::Cards) && HasGeometryType(EGroomGeometryType::Cards) && GetHairGroupsCards().Num() == 0)
 		{
-			HairGroupsCards.Init(FHairGroupsCardsSourceDescription(), GroupCount);
+			GetHairGroupsCards().Init(FHairGroupsCardsSourceDescription(), GroupCount);
 			FHairGroupsCardsSourceDescription Dirty;
 			Dirty.ProceduralSettings.ClusterSettings.ClusterDecimation = 0;
 			Dirty.SourceType = HasImportedStrandsData() ? EHairCardsSourceType::Procedural : EHairCardsSourceType::Imported;
 			CachedHairGroupsCards.Init(Dirty, GroupCount);
 		}
 
-		if (IsHairStrandsEnabled(EHairStrandsShaderType::Meshes) && HasGeometryType(EGroomGeometryType::Meshes) && HairGroupsMeshes.Num() == 0)
+		if (IsHairStrandsEnabled(EHairStrandsShaderType::Meshes) && HasGeometryType(EGroomGeometryType::Meshes) && GetHairGroupsMeshes().Num() == 0)
 		{
-			HairGroupsMeshes.Init(FHairGroupsMeshesSourceDescription(), GroupCount);
+			GetHairGroupsMeshes().Init(FHairGroupsMeshesSourceDescription(), GroupCount);
 			FHairGroupsMeshesSourceDescription Dirty;
 			Dirty.ImportedMesh = nullptr;
 			CachedHairGroupsMeshes.Init(Dirty, GroupCount);
@@ -1644,10 +1644,10 @@ void UGroomAsset::PostEditChangeProperty(FPropertyChangedEvent& PropertyChangedE
 		CacheDerivedDatas();
 	}
 	
-	const bool bCardsArrayChanged = PropertyName == GET_MEMBER_NAME_CHECKED(UGroomAsset, HairGroupsCards);
+	const bool bCardsArrayChanged = PropertyName == UGroomAsset::GetHairGroupsCardsMemberName();
 	if (bCardsArrayChanged && PropertyChangedEvent.ChangeType == EPropertyChangeType::ArrayAdd)
 	{
-		HairGroupsCards.Last().SourceType = HasImportedStrandsData() ? EHairCardsSourceType::Procedural : EHairCardsSourceType::Imported;
+		GetHairGroupsCards().Last().SourceType = HasImportedStrandsData() ? EHairCardsSourceType::Procedural : EHairCardsSourceType::Imported;
 	}
 
 	SavePendingProceduralAssets();
@@ -1678,13 +1678,13 @@ void UGroomAsset::PostEditChangeProperty(FPropertyChangedEvent& PropertyChangedE
 	if (bRiggingChanged)
 	{
 		bool bUpdateRiggedMesh = false;
-		for (const FHairGroupsInterpolation& Interpolation : HairGroupsInterpolation)
+		for (const FHairGroupsInterpolation& Interpolation : GetHairGroupsInterpolation())
 		{
 			bUpdateRiggedMesh |= Interpolation.InterpolationSettings.GuideType == EGroomGuideType::Rigged;
 		}
 		if (bUpdateRiggedMesh)
 		{
-			RiggedSkeletalMesh = FGroomDeformerBuilder::CreateSkeletalMesh(this);
+			SetRiggedSkeletalMesh(FGroomDeformerBuilder::CreateSkeletalMesh(this));
 		}
 	}
 
@@ -1733,7 +1733,7 @@ void UGroomAsset::PostInitProperties()
 
 int32 UGroomAsset::GetNumHairGroups() const
 {
-	return HairGroupsPlatformData.Num();
+	return GetHairGroupsPlatformData().Num();
 }
 
 ////////////////////////////////////////////////////////////////////////////////////////////////////////////
@@ -1957,13 +1957,13 @@ bool UGroomAsset::CanRebuildFromDescription() const
 
 bool UGroomAsset::AreGroupsValid() const
 {
-	const uint32 GroupCount = HairGroupsInterpolation.Num();
+	const uint32 GroupCount = GetHairGroupsInterpolation().Num();
 	return
 		GroupCount > 0 &&
-		HairGroupsPlatformData.Num() == GroupCount &&
-		HairGroupsPhysics.Num() == GroupCount &&
-		HairGroupsRendering.Num() == GroupCount &&
-		HairGroupsLOD.Num() == GroupCount;
+		GetHairGroupsPlatformData().Num() == GroupCount &&
+		GetHairGroupsPhysics().Num() == GroupCount &&
+		GetHairGroupsRendering().Num() == GroupCount &&
+		GetHairGroupsLOD().Num() == GroupCount;
 }
 
 void UGroomAsset::SetNumGroup(uint32 InGroupCount, bool bResetGroupData, bool bResetOtherData)
@@ -1971,64 +1971,64 @@ void UGroomAsset::SetNumGroup(uint32 InGroupCount, bool bResetGroupData, bool bR
 	ReleaseResource();
 	if (bResetGroupData)
 	{
-		HairGroupsPlatformData.Reset();
+		GetHairGroupsPlatformData().Reset();
 	}
 
 	// In order to preserve the existing asset settings, the settings are only reset if the group count has changed.
-	if (InGroupCount != HairGroupsPlatformData.Num())
+	if (InGroupCount != GetHairGroupsPlatformData().Num())
 	{
-		HairGroupsPlatformData.Init(FHairGroupPlatformData(), InGroupCount);
+		GetHairGroupsPlatformData().Init(FHairGroupPlatformData(), InGroupCount);
 	}
 
-	if (InGroupCount != HairGroupsPhysics.Num())
+	if (InGroupCount != GetHairGroupsPhysics().Num())
 	{
 		if (bResetOtherData)
 		{
-			HairGroupsPhysics.Init(FHairGroupsPhysics(), InGroupCount);
+			GetHairGroupsPhysics().Init(FHairGroupsPhysics(), InGroupCount);
 		}
 		else
 		{
-			HairGroupsPhysics.SetNum(InGroupCount);
+			GetHairGroupsPhysics().SetNum(InGroupCount);
 		}
 	}
 
-	if (InGroupCount != HairGroupsRendering.Num())
+	if (InGroupCount != GetHairGroupsRendering().Num())
 	{
 		if (bResetOtherData)
 		{
-			HairGroupsRendering.Init(FHairGroupsRendering(), InGroupCount);
+			GetHairGroupsRendering().Init(FHairGroupsRendering(), InGroupCount);
 		}
 		else
 		{
-			HairGroupsRendering.SetNum(InGroupCount);
+			GetHairGroupsRendering().SetNum(InGroupCount);
 		}
 	}
 
-	if (InGroupCount != HairGroupsInterpolation.Num())
+	if (InGroupCount != GetHairGroupsInterpolation().Num())
 	{
 		if (bResetOtherData)
 		{
-			HairGroupsInterpolation.Init(FHairGroupsInterpolation(), InGroupCount);
+			GetHairGroupsInterpolation().Init(FHairGroupsInterpolation(), InGroupCount);
 		}
 		else
 		{
-			HairGroupsInterpolation.SetNum(InGroupCount);
+			GetHairGroupsInterpolation().SetNum(InGroupCount);
 		}
 	}
 
-	if (InGroupCount != HairGroupsLOD.Num())
+	if (InGroupCount != GetHairGroupsLOD().Num())
 	{
 		if (bResetOtherData)
 		{
-			HairGroupsLOD.Init(FHairGroupsLOD(), InGroupCount);
+			GetHairGroupsLOD().Init(FHairGroupsLOD(), InGroupCount);
 		}
 		else
 		{
-			HairGroupsLOD.SetNum(InGroupCount);
+			GetHairGroupsLOD().SetNum(InGroupCount);
 		}
 
 		// Insure that each group has at least one LOD
-		for (FHairGroupsLOD& GroupLOD : HairGroupsLOD)
+		for (FHairGroupsLOD& GroupLOD : GetHairGroupsLOD())
 		{
 			if (GroupLOD.LODs.IsEmpty())
 			{
@@ -2039,24 +2039,24 @@ void UGroomAsset::SetNumGroup(uint32 InGroupCount, bool bResetGroupData, bool bR
 		}
 	}
 
-	if (InGroupCount != HairGroupsInfo.Num())
+	if (InGroupCount != GetHairGroupsInfo().Num())
 	{
 		if (bResetOtherData)
 		{
-			HairGroupsInfo.Init(FHairGroupInfoWithVisibility(), InGroupCount);
+			GetHairGroupsInfo().Init(FHairGroupInfoWithVisibility(), InGroupCount);
 		}
 		else
 		{
-			HairGroupsInfo.SetNum(InGroupCount);
+			GetHairGroupsInfo().SetNum(InGroupCount);
 		}
 	}
 
 	// The following groups remain unchanged:
 	// * HairGroupsCards
-	// * HairGroupsMeshes
-	// * HairGroupsMaterials
+	// * GetHairGroupsMeshes()
+	// * GetHairGroupsMaterials()
 
-	EffectiveLODBias.SetNum(InGroupCount);
+	GetEffectiveLODBias().SetNum(InGroupCount);
 
 #if WITH_EDITORONLY_DATA
 	StrandsDerivedDataKey.SetNum(InGroupCount);
@@ -2067,7 +2067,7 @@ void UGroomAsset::SetNumGroup(uint32 InGroupCount, bool bResetGroupData, bool bR
 
 void UGroomAsset::SetStableRasterization(bool bEnable)
 {
-	for (FHairGroupsRendering& Group : HairGroupsRendering)
+	for (FHairGroupsRendering& Group : GetHairGroupsRendering())
 	{
 		Group.AdvancedSettings.bUseStableRasterization = bEnable;
 	}
@@ -2075,7 +2075,7 @@ void UGroomAsset::SetStableRasterization(bool bEnable)
 
 void UGroomAsset::SetScatterSceneLighting(bool bEnable)
 {
-	for (FHairGroupsRendering& Group : HairGroupsRendering)
+	for (FHairGroupsRendering& Group : GetHairGroupsRendering())
 	{
 		Group.AdvancedSettings.bScatterSceneLighting = bEnable;
 	}
@@ -2083,7 +2083,7 @@ void UGroomAsset::SetScatterSceneLighting(bool bEnable)
 
 void UGroomAsset::SetHairWidth(float Width)
 {
-	for (FHairGroupsRendering& Group : HairGroupsRendering)
+	for (FHairGroupsRendering& Group : GetHairGroupsRendering())
 	{
 		Group.GeometrySettings.HairWidth = Width;
 	}
@@ -2131,7 +2131,7 @@ namespace GroomDerivedDataCacheUtils
 			return FString();
 		}
 
-		if (Desc.GroupIndex >= GroomAsset->HairGroupsPlatformData.Num() || Desc.LODIndex >= GroomAsset->HairGroupsPlatformData[Desc.GroupIndex].Cards.LODs.Num())
+		if (Desc.GroupIndex >= GroomAsset->GetHairGroupsPlatformData().Num() || Desc.LODIndex >= GroomAsset->GetHairGroupsPlatformData()[Desc.GroupIndex].Cards.LODs.Num())
 		{
 			return FString();
 		}
@@ -2305,7 +2305,7 @@ bool UGroomAsset::IsFullyCached()
 {
 	// Check if all the groom derived data for strands, cards and meshes are already stored in the DDC
 	TArray<FString> CacheKeys;
-	const uint32 GroupCount = HairGroupsInterpolation.Num();
+	const uint32 GroupCount = GetHairGroupsInterpolation().Num();
 	for (uint32 GroupIndex = 0; GroupIndex < GroupCount; ++GroupIndex)
 	{
 		const FString StrandsDDCKey = UGroomAsset::GetDerivedDataKeyForStrands(GroupIndex);
@@ -2314,17 +2314,17 @@ bool UGroomAsset::IsFullyCached()
 		// Some cards and meshes LOD settings may not produce any derived data so those must be excluded
 		bool bHasCardsLODs = false;
 		bool bHasMeshesLODs = false;
-		const uint32 LODCount = HairGroupsLOD[GroupIndex].LODs.Num();
+		const uint32 LODCount = GetHairGroupsLOD()[GroupIndex].LODs.Num();
 		for (uint32 LODIt = 0; LODIt < LODCount; ++LODIt)
 		{
 			int32 SourceIt = 0;
 			// GetSourceDescription will cross-check the LOD settings with the cards/meshes settings to see if they would produce any data
-			if (const FHairGroupsCardsSourceDescription* Desc = GetSourceDescription(HairGroupsCards, GroupIndex, LODIt, SourceIt))
+			if (const FHairGroupsCardsSourceDescription* Desc = GetSourceDescription(GetHairGroupsCards(), GroupIndex, LODIt, SourceIt))
 			{
 				Desc->HasMeshChanged(); // this query will trigger a load of the mesh dependency, which has to be done in the game thread
 				bHasCardsLODs = true;
 			}
-			if (const FHairGroupsMeshesSourceDescription* Desc = GetSourceDescription(HairGroupsMeshes, GroupIndex, LODIt, SourceIt))
+			if (const FHairGroupsMeshesSourceDescription* Desc = GetSourceDescription(GetHairGroupsMeshes(), GroupIndex, LODIt, SourceIt))
 			{
 				Desc->HasMeshChanged(); // this query will trigger a load of the mesh dependency, which has to be done in the game thread
 				bHasMeshesLODs = true;
@@ -2333,14 +2333,14 @@ bool UGroomAsset::IsFullyCached()
 
 		if (bHasCardsLODs)
 		{
-			const FString CardsKeySuffix = GroomDerivedDataCacheUtils::BuildCardsDerivedDataKeySuffix(GroupIndex, HairGroupsLOD[GroupIndex].LODs, HairGroupsCards);
+			const FString CardsKeySuffix = GroomDerivedDataCacheUtils::BuildCardsDerivedDataKeySuffix(GroupIndex, GetHairGroupsLOD()[GroupIndex].LODs, GetHairGroupsCards());
 			const FString CardsDDCKey = StrandsDDCKey + CardsKeySuffix;
 			CacheKeys.Add(CardsDDCKey);
 		}
 
 		if (bHasMeshesLODs)
 		{
-			const FString MeshesKeySuffix = GroomDerivedDataCacheUtils::BuildMeshesDerivedDataKeySuffix(GroupIndex, HairGroupsLOD[GroupIndex].LODs, HairGroupsMeshes);
+			const FString MeshesKeySuffix = GroomDerivedDataCacheUtils::BuildMeshesDerivedDataKeySuffix(GroupIndex, GetHairGroupsLOD()[GroupIndex].LODs, GetHairGroupsMeshes());
 			const FString MeshesDDCKey = GroomDerivedDataCacheUtils::BuildGroomDerivedDataKey(MeshesKeySuffix);
 			CacheKeys.Add(MeshesDDCKey);
 		}
@@ -2370,22 +2370,22 @@ FString UGroomAsset::GetDerivedDataKey(bool bUseStrandCachedKey)
 
 FString UGroomAsset::GetDerivedDataKeyForCards(uint32 GroupIndex, const FString& StrandsKey)
 {
-	const FString KeySuffix = GroomDerivedDataCacheUtils::BuildCardsDerivedDataKeySuffix(GroupIndex, HairGroupsLOD[GroupIndex].LODs, HairGroupsCards);
+	const FString KeySuffix = GroomDerivedDataCacheUtils::BuildCardsDerivedDataKeySuffix(GroupIndex, GetHairGroupsLOD()[GroupIndex].LODs, GetHairGroupsCards());
 	const FString DerivedDataKey = StrandsKey + KeySuffix;
 	return DerivedDataKey;
 }
 
 FString UGroomAsset::GetDerivedDataKeyForMeshes(uint32 GroupIndex)
 {
-	const FString KeySuffix = GroomDerivedDataCacheUtils::BuildMeshesDerivedDataKeySuffix(GroupIndex, HairGroupsLOD[GroupIndex].LODs, HairGroupsMeshes);
+	const FString KeySuffix = GroomDerivedDataCacheUtils::BuildMeshesDerivedDataKeySuffix(GroupIndex, GetHairGroupsLOD()[GroupIndex].LODs, GetHairGroupsMeshes());
 	const FString DerivedDataKey = GroomDerivedDataCacheUtils::BuildGroomDerivedDataKey(KeySuffix);
 	return DerivedDataKey;
 }
 
 FString UGroomAsset::GetDerivedDataKeyForStrands(uint32 GroupIndex)
 {
-	const FHairGroupsInterpolation& InterpolationSettings = HairGroupsInterpolation[GroupIndex];
-	const FHairGroupsLOD& LODSettings = HairGroupsLOD[GroupIndex];
+	const FHairGroupsInterpolation& InterpolationSettings = GetHairGroupsInterpolation()[GroupIndex];
+	const FHairGroupsLOD& LODSettings = GetHairGroupsLOD()[GroupIndex];
 
 	// If simulation or global interpolation is enabled, then interpolation data are required. Otherwise they can be skipped
 	const bool bNeedInterpolationData = NeedsInterpolationData(GroupIndex);
@@ -2450,7 +2450,7 @@ bool UGroomAsset::CacheDerivedDatas()
 	const bool bIsGameThread = IsInGameThread();
 	FGroomComponentRecreateRenderStateContext RecreateContext(bIsGameThread ? this : nullptr);
 
-	const uint32 GroupCount = HairGroupsInterpolation.Num();
+	const uint32 GroupCount = GetHairGroupsInterpolation().Num();
 	for (uint32 GroupIndex = 0; GroupIndex < GroupCount; ++GroupIndex)
 	{
 		bool bSucceed = CacheDerivedData(GroupIndex);
@@ -2476,7 +2476,7 @@ bool UGroomAsset::CacheDerivedData(uint32 GroupIndex)
 		return false;
 	}
 
-	const uint32 GroupCount = HairGroupsInterpolation.Num();
+	const uint32 GroupCount = GetHairGroupsInterpolation().Num();
 	check(GroupIndex < GroupCount);
 	if (GroupIndex >= GroupCount)
 	{
@@ -2500,7 +2500,7 @@ bool UGroomAsset::CacheStrandsData(uint32 GroupIndex, FString& OutDerivedDataKey
 	{
 		// Groom having only cards/meshes (i.e., not containing imported strands) don't have hair description
 		bool bHasStrands = false;
-		for (const FHairLODSettings& LOD : HairGroupsLOD[GroupIndex].LODs)
+		for (const FHairLODSettings& LOD : GetHairGroupsLOD()[GroupIndex].LODs)
 		{
 			if (LOD.GeometryType == EGroomGeometryType::Strands)
 			{
@@ -2539,7 +2539,7 @@ bool UGroomAsset::CacheStrandsData(uint32 GroupIndex, FString& OutDerivedDataKey
 		In.Strands.ClusterCullingBulkData.DerivedDataKey = DerivedDataKey + FString(TEXT("_ClusterCullingData"));
 	};
 
-	FHairGroupPlatformData& PlatformData = HairGroupsPlatformData[GroupIndex];
+	FHairGroupPlatformData& PlatformData = GetHairGroupsPlatformData()[GroupIndex];
 
 	bool bSuccess = true;
 	if (Data)
@@ -2579,10 +2579,10 @@ bool UGroomAsset::CacheStrandsData(uint32 GroupIndex, FString& OutDerivedDataKey
 			GroupIndex,
 			NeedsInterpolationData(GroupIndex),
 			LocalHairDescriptionGroups,
-			HairGroupsInterpolation,
-			HairGroupsLOD,
-			HairGroupsInfo,
-			HairGroupsPlatformData);
+			GetHairGroupsInterpolation(),
+			GetHairGroupsLOD(),
+			GetHairGroupsInfo(),
+			GetHairGroupsPlatformData());
 		
 		if (bSuccess)
 		{
@@ -2641,14 +2641,14 @@ bool UGroomAsset::CacheStrandsData(uint32 GroupIndex, FString& OutDerivedDataKey
 
 bool UGroomAsset::HasValidMeshesData(uint32 GroupIndex) const
 {
-	const FHairGroupPlatformData& HairGroupData = HairGroupsPlatformData[GroupIndex];
-	const FHairGroupsLOD& GroupsLOD = HairGroupsLOD[GroupIndex];
+	const FHairGroupPlatformData& HairGroupData = GetHairGroupsPlatformData()[GroupIndex];
+	const FHairGroupsLOD& GroupsLOD = GetHairGroupsLOD()[GroupIndex];
 	for (int32 LODIt = 0; LODIt < GroupsLOD.LODs.Num(); ++LODIt)
 	{
 		if (GroupsLOD.LODs[LODIt].GeometryType == EGroomGeometryType::Meshes)
 		{
 			int32 SourceIt = 0;
-			if (const FHairGroupsMeshesSourceDescription* Desc = GetSourceDescription(HairGroupsMeshes, GroupIndex, LODIt, SourceIt))
+			if (const FHairGroupsMeshesSourceDescription* Desc = GetSourceDescription(GetHairGroupsMeshes(), GroupIndex, LODIt, SourceIt))
 			{
 				if (Desc->ImportedMesh)
 				{
@@ -2662,15 +2662,15 @@ bool UGroomAsset::HasValidMeshesData(uint32 GroupIndex) const
 
 bool UGroomAsset::HasValidCardsData(uint32 GroupIndex) const
 {
-	const FHairGroupPlatformData& HairGroupData = HairGroupsPlatformData[GroupIndex];
-	const FHairGroupsLOD& GroupsLOD = HairGroupsLOD[GroupIndex];
+	const FHairGroupPlatformData& HairGroupData = GetHairGroupsPlatformData()[GroupIndex];
+	const FHairGroupsLOD& GroupsLOD = GetHairGroupsLOD()[GroupIndex];
 	for (int32 LODIt = 0; LODIt < GroupsLOD.LODs.Num(); ++LODIt)
 	{
 		// Note: We don't condition/filter resource allocation based on GeometryType == EGroomGeometryType::Cards, because 
 		// geometry type can be switched at runtime (between cards/strands)
 		{
 			int32 SourceIt = 0;
-			if (const FHairGroupsCardsSourceDescription* Desc = GetSourceDescription(HairGroupsCards, GroupIndex, LODIt, SourceIt))
+			if (const FHairGroupsCardsSourceDescription* Desc = GetSourceDescription(GetHairGroupsCards(), GroupIndex, LODIt, SourceIt))
 			{
 				if (Desc->GetMesh())
 				{
@@ -2693,7 +2693,7 @@ bool UGroomAsset::CacheCardsData(uint32 GroupIndex, const FString& StrandsKey)
 		return bNeedReset;
 	}
 
-	const FString KeySuffix = GroomDerivedDataCacheUtils::BuildCardsDerivedDataKeySuffix(GroupIndex, HairGroupsLOD[GroupIndex].LODs, HairGroupsCards);
+	const FString KeySuffix = GroomDerivedDataCacheUtils::BuildCardsDerivedDataKeySuffix(GroupIndex, GetHairGroupsLOD()[GroupIndex].LODs, GetHairGroupsCards());
 	const FString DerivedDataKey = StrandsKey + KeySuffix;
 
 	if (CardsDerivedDataKey[GroupIndex] == DerivedDataKey)
@@ -2702,17 +2702,17 @@ bool UGroomAsset::CacheCardsData(uint32 GroupIndex, const FString& StrandsKey)
 		return false;
 	}
 
-	FHairGroupPlatformData& HairGroupData = HairGroupsPlatformData[GroupIndex];
+	FHairGroupPlatformData& HairGroupData = GetHairGroupsPlatformData()[GroupIndex];
 
 	// We should never query the DDC for something that we won't be able to build. So check that first.
 	// The goal is to avoid paying the price of a guaranteed failure on high-latency DDC backends every run.
 	bool bDataCanBeBuilt = false;
 
-	const FHairGroupsLOD& GroupsLOD = HairGroupsLOD[GroupIndex];
+	const FHairGroupsLOD& GroupsLOD = GetHairGroupsLOD()[GroupIndex];
 	for (int32 LODIt = 0; LODIt < GroupsLOD.LODs.Num(); ++LODIt)
 	{
 		int32 SourceIt = 0;
-		if (FHairGroupsCardsSourceDescription* Desc = GetSourceDescription(HairGroupsCards, GroupIndex, LODIt, SourceIt))
+		if (FHairGroupsCardsSourceDescription* Desc = GetSourceDescription(GetHairGroupsCards(), GroupIndex, LODIt, SourceIt))
 		{
 			// Note: We don't condition/filter resource allocation based on GeometryType == EGroomGeometryType::Cards, because 
 			// geometry type can be switched at runtime (between cards/strands)
@@ -2772,13 +2772,13 @@ bool UGroomAsset::CacheCardsData(uint32 GroupIndex, const FString& StrandsKey)
 
 	// Handle the case where the cards data is already cached in the DDC
 	// Need to populate the strands data with it
-	FHairGroupPlatformData& GroupData = HairGroupsPlatformData[GroupIndex];
+	FHairGroupPlatformData& GroupData = GetHairGroupsPlatformData()[GroupIndex];
 	if (StrandsDerivedDataKey[GroupIndex].IsEmpty() && GroupData.Strands.BulkData.GetNumPoints() == 0)
 	{
 		for (int32 LODIt = 0; LODIt < GroupData.Cards.LODs.Num(); ++LODIt)
 		{
 			int32 SourceIt = 0;
-			if (FHairGroupsCardsSourceDescription* Desc = GetSourceDescription(HairGroupsCards, GroupIndex, LODIt, SourceIt))
+			if (FHairGroupsCardsSourceDescription* Desc = GetSourceDescription(GetHairGroupsCards(), GroupIndex, LODIt, SourceIt))
 			{
 				FHairGroupPlatformData::FCards::FLOD& LOD = GroupData.Cards.LODs[LODIt];
 				GroupData.Strands.BulkData = LOD.Guides.BulkData;
@@ -2792,7 +2792,7 @@ bool UGroomAsset::CacheCardsData(uint32 GroupIndex, const FString& StrandsKey)
 	for (int32 LODIt = 0; LODIt < GroupData.Cards.LODs.Num(); ++LODIt)
 	{
 		int32 SourceIt = 0;
-		if (FHairGroupsCardsSourceDescription* Desc = GetSourceDescription(HairGroupsCards, GroupIndex, LODIt, SourceIt))
+		if (FHairGroupsCardsSourceDescription* Desc = GetSourceDescription(GetHairGroupsCards(), GroupIndex, LODIt, SourceIt))
 		{
 			Desc->UpdateMeshKey();
 		}
@@ -2814,7 +2814,7 @@ bool UGroomAsset::CacheMeshesData(uint32 GroupIndex)
 		return bNeedReset;
 	}
 
-	const FString KeySuffix = GroomDerivedDataCacheUtils::BuildMeshesDerivedDataKeySuffix(GroupIndex, HairGroupsLOD[GroupIndex].LODs, HairGroupsMeshes);
+	const FString KeySuffix = GroomDerivedDataCacheUtils::BuildMeshesDerivedDataKeySuffix(GroupIndex, GetHairGroupsLOD()[GroupIndex].LODs, GetHairGroupsMeshes());
 	const FString DerivedDataKey = GroomDerivedDataCacheUtils::BuildGroomDerivedDataKey(KeySuffix);
 
 	if (MeshesDerivedDataKey[GroupIndex] == DerivedDataKey)
@@ -2823,17 +2823,17 @@ bool UGroomAsset::CacheMeshesData(uint32 GroupIndex)
 		return false;
 	}
 
-	FHairGroupPlatformData& HairGroupData = HairGroupsPlatformData[GroupIndex];
+	FHairGroupPlatformData& HairGroupData = GetHairGroupsPlatformData()[GroupIndex];
 
 	// We should never query the DDC for something that we won't be able to build. So check that first.
 	// The goal is to avoid paying the price of a guaranteed failure on high-latency DDC backends every run.
 	bool bDataCanBeBuilt = false;
 
-	const FHairGroupsLOD& GroupsLOD = HairGroupsLOD[GroupIndex];
+	const FHairGroupsLOD& GroupsLOD = GetHairGroupsLOD()[GroupIndex];
 	for (int32 LODIt = 0; LODIt < GroupsLOD.LODs.Num(); ++LODIt)
 	{
 		int32 SourceIt = 0;
-		if (FHairGroupsMeshesSourceDescription* Desc = GetSourceDescription(HairGroupsMeshes, GroupIndex, LODIt, SourceIt))
+		if (FHairGroupsMeshesSourceDescription* Desc = GetSourceDescription(GetHairGroupsMeshes(), GroupIndex, LODIt, SourceIt))
 		{
 			if (GroupsLOD.LODs[LODIt].GeometryType == EGroomGeometryType::Meshes)
 			{
@@ -2895,24 +2895,24 @@ bool UGroomAsset::BuildCardsData(uint32 GroupIndex)
 {
 	LLM_SCOPE_BYTAG(Groom);
 
-	if (!IsHairStrandsEnabled(EHairStrandsShaderType::Cards) || HairGroupsCards.Num() == 0 || !CanRebuildFromDescription())
+	if (!IsHairStrandsEnabled(EHairStrandsShaderType::Cards) || GetHairGroupsCards().Num() == 0 || !CanRebuildFromDescription())
 	{
 		return false;
 	}
 
-	bool bHasChanged = HairGroupsCards.Num() != CachedHairGroupsCards.Num();
+	bool bHasChanged = GetHairGroupsCards().Num() != CachedHairGroupsCards.Num();
 
 	check(GroupIndex < uint32(GetNumHairGroups()));
-	FHairGroupPlatformData& GroupData = HairGroupsPlatformData[GroupIndex];
+	FHairGroupPlatformData& GroupData = GetHairGroupsPlatformData()[GroupIndex];
 
 	// The settings might have been previously cached without the data having been built
-	const uint32 LODCount = HairGroupsLOD[GroupIndex].LODs.Num();
+	const uint32 LODCount = GetHairGroupsLOD()[GroupIndex].LODs.Num();
 	TArray<bool> bIsAlreadyBuilt;
 	bIsAlreadyBuilt.SetNum(LODCount);
 	for (uint32 LODIt = 0; LODIt < LODCount; ++LODIt)
 	{
 		int32 SourceIt = 0;
-		if (const FHairGroupsCardsSourceDescription* Desc = GetSourceDescription(HairGroupsCards, GroupIndex, LODIt, SourceIt))
+		if (const FHairGroupsCardsSourceDescription* Desc = GetSourceDescription(GetHairGroupsCards(), GroupIndex, LODIt, SourceIt))
 		{
 			bIsAlreadyBuilt[LODIt] = GroupData.Strands.HasValidData() && Desc->CardsInfo.NumCardVertices > 0 && !Desc->HasMeshChanged();
 			bHasChanged |= !bIsAlreadyBuilt[LODIt];
@@ -2921,9 +2921,9 @@ bool UGroomAsset::BuildCardsData(uint32 GroupIndex)
 
 	if (!bHasChanged)
 	{
-		for (uint32 SourceIt = 0, SourceCount = HairGroupsCards.Num(); SourceIt < SourceCount; ++SourceIt)
+		for (uint32 SourceIt = 0, SourceCount = GetHairGroupsCards().Num(); SourceIt < SourceCount; ++SourceIt)
 		{
-			const bool bEquals = CachedHairGroupsCards[SourceIt] == HairGroupsCards[SourceIt];
+			const bool bEquals = CachedHairGroupsCards[SourceIt] == GetHairGroupsCards()[SourceIt];
 			if (!bEquals)
 			{
 				bHasChanged = true;
@@ -2934,7 +2934,7 @@ bool UGroomAsset::BuildCardsData(uint32 GroupIndex)
 				// content is up to date with what has been saved.
 				if (IsInGameThread())
 				{
-					FHairGroupCardsTextures& Textures = HairGroupsCards[SourceIt].Textures;
+					FHairGroupCardsTextures& Textures = GetHairGroupsCards()[SourceIt].Textures;
 					if (Textures.DepthTexture != nullptr)		Textures.DepthTexture->UpdateResource();
 					if (Textures.TangentTexture != nullptr)		Textures.TangentTexture->UpdateResource();
 					if (Textures.AttributeTexture != nullptr)	Textures.AttributeTexture->UpdateResource();
@@ -2953,11 +2953,11 @@ bool UGroomAsset::BuildCardsData(uint32 GroupIndex)
 
 	bool bDataBuilt = false;
 	GroupData.Cards.LODs.SetNum(LODCount);
-	const FHairGroupsLOD& GroupsLOD = HairGroupsLOD[GroupIndex];
+	const FHairGroupsLOD& GroupsLOD = GetHairGroupsLOD()[GroupIndex];
 	for (uint32 LODIt = 0; LODIt < LODCount; ++LODIt)
 	{
 		int32 SourceIt = 0;
-		if (FHairGroupsCardsSourceDescription* Desc = GetSourceDescription(HairGroupsCards, GroupIndex, LODIt, SourceIt))
+		if (FHairGroupsCardsSourceDescription* Desc = GetSourceDescription(GetHairGroupsCards(), GroupIndex, LODIt, SourceIt))
 		{
 			// #hair_todo: add finer culling granularity to rebuild only what is necessary
 			const FHairGroupsCardsSourceDescription* CachedDesc = SourceIt < CachedHairGroupsCards.Num() ? &CachedHairGroupsCards[SourceIt] : nullptr;
@@ -3051,7 +3051,7 @@ bool UGroomAsset::BuildCardsData(uint32 GroupIndex)
 
 				// Create own interpolation settings for cards.
 				// Force closest guides as this is the most relevant matching metric for cards, due to their coarse geometry
-				FHairInterpolationSettings CardsInterpolationSettings = HairGroupsInterpolation[GroupIndex].InterpolationSettings;
+				FHairInterpolationSettings CardsInterpolationSettings = GetHairGroupsInterpolation()[GroupIndex].InterpolationSettings;
 				CardsInterpolationSettings.GuideType = EGroomGuideType::Imported;
 				CardsInterpolationSettings.bUseUniqueGuide = true;
 				CardsInterpolationSettings.bRandomizeGuide = false;
@@ -3078,7 +3078,7 @@ bool UGroomAsset::BuildCardsData(uint32 GroupIndex)
 				check(LocalHairDescriptionGroups.IsValid());
 				FHairStrandsDatas SimGuidesData;
 				FHairStrandsDatas RenStrandsData; // Dummy, not used
-				FGroomBuilder::BuildData(LocalHairDescriptionGroups.HairGroups[GroupIndex], HairGroupsInterpolation[GroupIndex], DummyInfo, RenStrandsData, SimGuidesData);
+				FGroomBuilder::BuildData(LocalHairDescriptionGroups.HairGroups[GroupIndex], GetHairGroupsInterpolation()[GroupIndex], DummyInfo, RenStrandsData, SimGuidesData);
 
 				// 3. Compute interpolation data & bulk data between the sim guides and the cards guides
 				FHairStrandsInterpolationDatas LODGuideInterpolationData;
@@ -3118,7 +3118,7 @@ bool UGroomAsset::BuildCardsData()
 		InitGuideResources();
 
 		// When building cards in an empty groom, the strands resources need to be initialized once
-		FHairGroupPlatformData& GroupData = HairGroupsPlatformData[0];
+		FHairGroupPlatformData& GroupData = GetHairGroupsPlatformData()[0];
 		if (!GroupData.Strands.ClusterCullingResource)
 		{
 			InitStrandsResources();
@@ -3148,23 +3148,23 @@ bool UGroomAsset::BuildMeshesData()
 
 bool UGroomAsset::BuildMeshesData(uint32 GroupIndex)
 {
-	if (!IsHairStrandsEnabled(EHairStrandsShaderType::Meshes) || HairGroupsMeshes.Num() == 0)
+	if (!IsHairStrandsEnabled(EHairStrandsShaderType::Meshes) || GetHairGroupsMeshes().Num() == 0)
 	{
 		return false;
 	}
 
-	bool bHasChanged = HairGroupsMeshes.Num() != CachedHairGroupsMeshes.Num();
+	bool bHasChanged = GetHairGroupsMeshes().Num() != CachedHairGroupsMeshes.Num();
 
-	FHairGroupPlatformData& GroupData = HairGroupsPlatformData[GroupIndex];
+	FHairGroupPlatformData& GroupData = GetHairGroupsPlatformData()[GroupIndex];
 
 	// The settings might have been previously cached without the data having been built
-	const uint32 LODCount = HairGroupsLOD[GroupIndex].LODs.Num();
+	const uint32 LODCount = GetHairGroupsLOD()[GroupIndex].LODs.Num();
 	TArray<bool> bIsAlreadyBuilt;
 	bIsAlreadyBuilt.SetNum(LODCount);
 	for (uint32 LODIt = 0; LODIt < LODCount; ++LODIt)
 	{
 		int32 SourceIt = 0;
-		if (const FHairGroupsMeshesSourceDescription* Desc = GetSourceDescription(HairGroupsMeshes, GroupIndex, LODIt, SourceIt))
+		if (const FHairGroupsMeshesSourceDescription* Desc = GetSourceDescription(GetHairGroupsMeshes(), GroupIndex, LODIt, SourceIt))
 		{
 			bIsAlreadyBuilt[LODIt] = GroupData.Meshes.LODs.IsValidIndex(LODIt) && GroupData.Meshes.LODs[LODIt].BulkData.GetNumVertices() > 0 && !Desc->HasMeshChanged();
 			bHasChanged |= !bIsAlreadyBuilt[LODIt];
@@ -3173,9 +3173,9 @@ bool UGroomAsset::BuildMeshesData(uint32 GroupIndex)
 
 	if (!bHasChanged)
 	{
-		for (uint32 SourceIt = 0, SourceCount = HairGroupsMeshes.Num(); SourceIt < SourceCount; ++SourceIt)
+		for (uint32 SourceIt = 0, SourceCount = GetHairGroupsMeshes().Num(); SourceIt < SourceCount; ++SourceIt)
 		{
-			const bool bEquals = CachedHairGroupsMeshes[SourceIt] == HairGroupsMeshes[SourceIt];
+			const bool bEquals = CachedHairGroupsMeshes[SourceIt] == GetHairGroupsMeshes()[SourceIt];
 			if (!bEquals)
 			{
 				bHasChanged = true;
@@ -3186,9 +3186,9 @@ bool UGroomAsset::BuildMeshesData(uint32 GroupIndex)
 
 	if (!bHasChanged)
 	{
-		for (uint32 SourceIt = 0, SourceCount = HairGroupsMeshes.Num(); SourceIt < SourceCount; ++SourceIt)
+		for (uint32 SourceIt = 0, SourceCount = GetHairGroupsMeshes().Num(); SourceIt < SourceCount; ++SourceIt)
 		{
-			const bool bEquals = CachedHairGroupsMeshes[SourceIt] == HairGroupsMeshes[SourceIt];
+			const bool bEquals = CachedHairGroupsMeshes[SourceIt] == GetHairGroupsMeshes()[SourceIt];
 			if (!bEquals)
 			{
 				bHasChanged = true;
@@ -3199,7 +3199,7 @@ bool UGroomAsset::BuildMeshesData(uint32 GroupIndex)
 				// content is up to date with what has been saved.
 				if (IsInGameThread())
 				{
-					FHairGroupCardsTextures& Textures = HairGroupsMeshes[SourceIt].Textures;
+					FHairGroupCardsTextures& Textures = GetHairGroupsMeshes()[SourceIt].Textures;
 					if (Textures.DepthTexture != nullptr)		Textures.DepthTexture->UpdateResource();
 					if (Textures.TangentTexture != nullptr)		Textures.TangentTexture->UpdateResource();
 					if (Textures.AttributeTexture != nullptr)	Textures.AttributeTexture->UpdateResource();
@@ -3218,11 +3218,11 @@ bool UGroomAsset::BuildMeshesData(uint32 GroupIndex)
 
 	bool bDataBuilt = false;
 	GroupData.Meshes.LODs.SetNum(LODCount);
-	const FHairGroupsLOD& GroupsLOD = HairGroupsLOD[GroupIndex];
+	const FHairGroupsLOD& GroupsLOD = GetHairGroupsLOD()[GroupIndex];
 	for (uint32 LODIt = 0; LODIt < LODCount; ++LODIt)
 	{
 		int32 SourceIt = 0;
-		if (const FHairGroupsMeshesSourceDescription* Desc = GetSourceDescription(HairGroupsMeshes, GroupIndex, LODIt, SourceIt))
+		if (const FHairGroupsMeshesSourceDescription* Desc = GetSourceDescription(GetHairGroupsMeshes(), GroupIndex, LODIt, SourceIt))
 		{
 			const FHairGroupsMeshesSourceDescription* CachedDesc = SourceIt < CachedHairGroupsMeshes.Num() ? &CachedHairGroupsMeshes[SourceIt] : nullptr;
 			const bool bLODHasChanged = CachedDesc == nullptr || !(*CachedDesc == *Desc);
@@ -3264,7 +3264,7 @@ bool UGroomAsset::BuildMeshesData(uint32 GroupIndex)
 	for (int32 LODIt = 0; LODIt < GroupData.Meshes.LODs.Num(); ++LODIt)
 	{
 		int32 SourceIt = 0;
-		if (FHairGroupsMeshesSourceDescription* Desc = GetSourceDescription(HairGroupsMeshes, GroupIndex, LODIt, SourceIt))
+		if (FHairGroupsMeshesSourceDescription* Desc = GetSourceDescription(GetHairGroupsMeshes(), GroupIndex, LODIt, SourceIt))
 		{
 			Desc->UpdateMeshKey();
 		}
@@ -3276,12 +3276,12 @@ bool UGroomAsset::BuildMeshesData(uint32 GroupIndex)
 void UGroomAsset::UpdateCachedSettings()
 {
 #if WITH_EDITORONLY_DATA
-	CachedHairGroupsRendering		= HairGroupsRendering;
-	CachedHairGroupsPhysics			= HairGroupsPhysics;
-	CachedHairGroupsInterpolation	= HairGroupsInterpolation;
-	CachedHairGroupsLOD				= HairGroupsLOD;
-	CachedHairGroupsCards			= HairGroupsCards;
-	CachedHairGroupsMeshes			= HairGroupsMeshes;
+	CachedHairGroupsRendering		= GetHairGroupsRendering();
+	CachedHairGroupsPhysics			= GetHairGroupsPhysics();
+	CachedHairGroupsInterpolation	= GetHairGroupsInterpolation();
+	CachedHairGroupsLOD				= GetHairGroupsLOD();
+	CachedHairGroupsCards			= GetHairGroupsCards();
+	CachedHairGroupsMeshes			= GetHairGroupsMeshes();
 #endif
 }
 #endif // WITH_EDITORONLY_DATA
@@ -3300,7 +3300,7 @@ FHairStrandsRestResource* UGroomAsset::AllocateGuidesResources(uint32 GroupIndex
 {
 	if (GroupIndex < uint32(GetNumHairGroups()))
 	{
-		FHairGroupPlatformData& GroupData = HairGroupsPlatformData[GroupIndex];
+		FHairGroupPlatformData& GroupData = GetHairGroupsPlatformData()[GroupIndex];
 		if (GroupData.Guides.HasValidData())
 		{
 			if (GroupData.Guides.RestResource == nullptr)
@@ -3317,7 +3317,7 @@ FHairStrandsInterpolationResource* UGroomAsset::AllocateInterpolationResources(u
 {
 	if (GroupIndex < uint32(GetNumHairGroups()))
 	{
-		FHairGroupPlatformData& GroupData = HairGroupsPlatformData[GroupIndex];
+		FHairGroupPlatformData& GroupData = GetHairGroupsPlatformData()[GroupIndex];
 		check(GroupData.Guides.IsValid());
 		if (GroupData.Strands.InterpolationResource == nullptr)
 		{
@@ -3331,9 +3331,9 @@ FHairStrandsInterpolationResource* UGroomAsset::AllocateInterpolationResources(u
 #if RHI_RAYTRACING
 FHairStrandsRaytracingResource* UGroomAsset::AllocateCardsRaytracingResources(uint32 GroupIndex, uint32 LODIndex)
 {
-	if (IsRayTracingAllowed() && GroupIndex < uint32(GetNumHairGroups()) && LODIndex < uint32(HairGroupsLOD[GroupIndex].LODs.Num()))
+	if (IsRayTracingAllowed() && GroupIndex < uint32(GetNumHairGroups()) && LODIndex < uint32(GetHairGroupsLOD()[GroupIndex].LODs.Num()))
 	{
-		FHairGroupPlatformData& GroupData = HairGroupsPlatformData[GroupIndex];
+		FHairGroupPlatformData& GroupData = GetHairGroupsPlatformData()[GroupIndex];
 		FHairGroupPlatformData::FCards::FLOD& LOD = GroupData.Cards.LODs[LODIndex];
 		check(LOD.BulkData.IsValid());
 
@@ -3348,9 +3348,9 @@ FHairStrandsRaytracingResource* UGroomAsset::AllocateCardsRaytracingResources(ui
 
 FHairStrandsRaytracingResource* UGroomAsset::AllocateMeshesRaytracingResources(uint32 GroupIndex, uint32 LODIndex)
 {
-	if (IsRayTracingAllowed() && GroupIndex < uint32(GetNumHairGroups()) && LODIndex < uint32(HairGroupsLOD[GroupIndex].LODs.Num()))
+	if (IsRayTracingAllowed() && GroupIndex < uint32(GetNumHairGroups()) && LODIndex < uint32(GetHairGroupsLOD()[GroupIndex].LODs.Num()))
 	{
-		FHairGroupPlatformData& GroupData = HairGroupsPlatformData[GroupIndex];
+		FHairGroupPlatformData& GroupData = GetHairGroupsPlatformData()[GroupIndex];
 		FHairGroupPlatformData::FMeshes::FLOD& LOD = GroupData.Meshes.LODs[LODIndex];
 		check(LOD.BulkData.IsValid());
 
@@ -3367,7 +3367,7 @@ FHairStrandsRaytracingResource* UGroomAsset::AllocateStrandsRaytracingResources(
 {
 	if (IsRayTracingAllowed() && GroupIndex < uint32(GetNumHairGroups()))
 	{
-		FHairGroupPlatformData& GroupData = HairGroupsPlatformData[GroupIndex];
+		FHairGroupPlatformData& GroupData = GetHairGroupsPlatformData()[GroupIndex];
 		check(GroupData.Strands.HasValidData());
 
 		if (GroupData.Strands.RaytracingResource == nullptr)
@@ -3398,7 +3398,7 @@ void UGroomAsset::InitStrandsResources()
 
 	for (uint32 GroupIndex = 0, GroupCount = GetNumHairGroups(); GroupIndex < GroupCount; ++GroupIndex)
 	{
-		FHairGroupPlatformData& GroupData = HairGroupsPlatformData[GroupIndex];
+		FHairGroupPlatformData& GroupData = GetHairGroupsPlatformData()[GroupIndex];
 
 		if (GroupData.Strands.HasValidData())
 		{
@@ -3418,16 +3418,16 @@ void UGroomAsset::InitStrandsResources()
 
 static void InitCardsTextureResources(UGroomAsset* GroomAsset)
 {
-	if (!GroomAsset || !IsHairStrandsEnabled(EHairStrandsShaderType::Cards) || GroomAsset->HairGroupsCards.Num() == 0)
+	if (!GroomAsset || !IsHairStrandsEnabled(EHairStrandsShaderType::Cards) || GroomAsset->GetHairGroupsCards().Num() == 0)
 	{
 		return;
 	}
 
 	for (uint32 GroupIndex = 0, GroupCount = GroomAsset->GetNumHairGroups(); GroupIndex < GroupCount; ++GroupIndex)
 	{
-		FHairGroupPlatformData& GroupData = GroomAsset->HairGroupsPlatformData[GroupIndex];
+		FHairGroupPlatformData& GroupData = GroomAsset->GetHairGroupsPlatformData()[GroupIndex];
 
-		const uint32 LODCount = GroomAsset->HairGroupsLOD[GroupIndex].LODs.Num();
+		const uint32 LODCount = GroomAsset->GetHairGroupsLOD()[GroupIndex].LODs.Num();
 		GroupData.Cards.LODs.SetNum(LODCount);
 
 		for (uint32 LODIt = 0; LODIt < LODCount; ++LODIt)
@@ -3435,7 +3435,7 @@ static void InitCardsTextureResources(UGroomAsset* GroomAsset)
 			FHairGroupPlatformData::FCards::FLOD& LOD = GroupData.Cards.LODs[LODIt];
 
 			int32 SourceIt = 0;
-			const FHairGroupsCardsSourceDescription* Desc = GetSourceDescription(GroomAsset->HairGroupsCards, GroupIndex, LODIt, SourceIt);
+			const FHairGroupsCardsSourceDescription* Desc = GetSourceDescription(GroomAsset->GetHairGroupsCards(), GroupIndex, LODIt, SourceIt);
 			if (!Desc)
 			{
 				continue;
@@ -3465,16 +3465,16 @@ void UGroomAsset::InitCardsResources()
 {
 	LLM_SCOPE_BYTAG(Groom);
 
-	if (!IsHairStrandsEnabled(EHairStrandsShaderType::Cards) || HairGroupsCards.Num() == 0)
+	if (!IsHairStrandsEnabled(EHairStrandsShaderType::Cards) || GetHairGroupsCards().Num() == 0)
 	{
 		return;
 	}
 
 	for (uint32 GroupIndex = 0, GroupCount = GetNumHairGroups(); GroupIndex < GroupCount; ++GroupIndex)
 	{
-		FHairGroupPlatformData& GroupData = HairGroupsPlatformData[GroupIndex];
+		FHairGroupPlatformData& GroupData = GetHairGroupsPlatformData()[GroupIndex];
 
-		const uint32 LODCount = HairGroupsLOD[GroupIndex].LODs.Num();
+		const uint32 LODCount = GetHairGroupsLOD()[GroupIndex].LODs.Num();
 		GroupData.Cards.LODs.SetNum(LODCount);
 
 		for (uint32 LODIt = 0; LODIt < LODCount; ++LODIt)
@@ -3482,7 +3482,7 @@ void UGroomAsset::InitCardsResources()
 			FHairGroupPlatformData::FCards::FLOD& LOD = GroupData.Cards.LODs[LODIt];
 
 			int32 SourceIt = 0;
-			const FHairGroupsCardsSourceDescription* Desc = GetSourceDescription(HairGroupsCards, GroupIndex, LODIt, SourceIt);
+			const FHairGroupsCardsSourceDescription* Desc = GetSourceDescription(GetHairGroupsCards(), GroupIndex, LODIt, SourceIt);
 			if (!Desc)
 			{
 				continue;
@@ -3536,13 +3536,13 @@ void UGroomAsset::InitMeshesResources()
 
 	for (uint32 GroupIndex = 0, GroupCount = GetNumHairGroups(); GroupIndex < GroupCount; ++GroupIndex)
 	{
-		FHairGroupPlatformData& GroupData = HairGroupsPlatformData[GroupIndex];
+		FHairGroupPlatformData& GroupData = GetHairGroupsPlatformData()[GroupIndex];
 
 		const uint32 LODCount = GroupData.Meshes.LODs.Num();
 		for (uint32 LODIt = 0; LODIt < LODCount; ++LODIt)
 		{
 			int32 SourceIt = 0;
-			const FHairGroupsMeshesSourceDescription* Desc = GetSourceDescription(HairGroupsMeshes, GroupIndex, LODIt, SourceIt);
+			const FHairGroupsMeshesSourceDescription* Desc = GetSourceDescription(GetHairGroupsMeshes(), GroupIndex, LODIt, SourceIt);
 			if (!Desc)
 			{
 				continue;
@@ -3572,74 +3572,74 @@ void UGroomAsset::InitMeshesResources()
 
 EGroomGeometryType UGroomAsset::GetGeometryType(int32 GroupIndex, int32 LODIndex) const
 {
-	if (GroupIndex < 0 || GroupIndex >= HairGroupsLOD.Num())
+	if (GroupIndex < 0 || GroupIndex >= GetHairGroupsLOD().Num())
 	{
 		return EGroomGeometryType::Strands;
 	}
 
-	if (LODIndex < 0 || LODIndex >= HairGroupsLOD[GroupIndex].LODs.Num())
+	if (LODIndex < 0 || LODIndex >= GetHairGroupsLOD()[GroupIndex].LODs.Num())
 	{
 		return EGroomGeometryType::Strands;
 	}
 
-	return HairGroupsLOD[GroupIndex].LODs[LODIndex].GeometryType;
+	return GetHairGroupsLOD()[GroupIndex].LODs[LODIndex].GeometryType;
 }
 
 EGroomBindingType UGroomAsset::GetBindingType(int32 GroupIndex, int32 LODIndex) const
 {
-	if (GroupIndex < 0 || GroupIndex >= HairGroupsLOD.Num() || !IsHairStrandsBindingEnable())
+	if (GroupIndex < 0 || GroupIndex >= GetHairGroupsLOD().Num() || !IsHairStrandsBindingEnable())
 	{
 		return EGroomBindingType::Rigid; // Fallback to rigid by default
 	}
 
-	if (LODIndex < 0 || LODIndex >= HairGroupsLOD[GroupIndex].LODs.Num())
+	if (LODIndex < 0 || LODIndex >= GetHairGroupsLOD()[GroupIndex].LODs.Num())
 	{
 		return EGroomBindingType::Rigid; // Fallback to rigid by default
 	}
 
-	return HairGroupsLOD[GroupIndex].LODs[LODIndex].BindingType;
+	return GetHairGroupsLOD()[GroupIndex].LODs[LODIndex].BindingType;
 }
 
 bool UGroomAsset::IsVisible(int32 GroupIndex, int32 LODIndex) const
 {
-	if (GroupIndex < 0 || GroupIndex >= HairGroupsLOD.Num())
+	if (GroupIndex < 0 || GroupIndex >= GetHairGroupsLOD().Num())
 	{
 		return false;
 	}
 
-	if (LODIndex < 0 || LODIndex >= HairGroupsLOD[GroupIndex].LODs.Num())
+	if (LODIndex < 0 || LODIndex >= GetHairGroupsLOD()[GroupIndex].LODs.Num())
 	{
 		return false;
 	}
 
-	return HairGroupsLOD[GroupIndex].LODs[LODIndex].bVisible;
+	return GetHairGroupsLOD()[GroupIndex].LODs[LODIndex].bVisible;
 }
 
 bool UGroomAsset::IsDeformationEnable(int32 GroupIndex) const
 {
 	return 
-		HairGroupsInterpolation.IsValidIndex(GroupIndex) && 
-		HairGroupsInterpolation[GroupIndex].InterpolationSettings.GuideType == EGroomGuideType::Rigged;
+		GetHairGroupsInterpolation().IsValidIndex(GroupIndex) && 
+		GetHairGroupsInterpolation()[GroupIndex].InterpolationSettings.GuideType == EGroomGuideType::Rigged;
 }
 
 bool UGroomAsset::IsSimulationEnable(int32 GroupIndex, int32 LODIndex) const 
 {
-	if (GroupIndex < 0 || GroupIndex >= HairGroupsLOD.Num() || !IsHairStrandsSimulationEnable())
+	if (GroupIndex < 0 || GroupIndex >= GetHairGroupsLOD().Num() || !IsHairStrandsSimulationEnable())
 	{
 		return false;
 	}
 
-	if (LODIndex >= HairGroupsLOD[GroupIndex].LODs.Num())
+	if (LODIndex >= GetHairGroupsLOD()[GroupIndex].LODs.Num())
 	{
 		return false;
 	}
 
-	check(HairGroupsPhysics.Num() == HairGroupsLOD.Num());
+	check(GetHairGroupsPhysics().Num() == GetHairGroupsLOD().Num());
 
 	// If the LOD index is not forced, then its value could be -1. In this case we return the 'global' asset value
 	if (LODIndex < 0)
 	{
-		return HairGroupsPhysics[GroupIndex].SolverSettings.EnableSimulation;
+		return GetHairGroupsPhysics()[GroupIndex].SolverSettings.EnableSimulation;
 	}
 
 	if (!IsHairLODSimulationEnabled(LODIndex))
@@ -3647,19 +3647,19 @@ bool UGroomAsset::IsSimulationEnable(int32 GroupIndex, int32 LODIndex) const
 		return false;
 	}
 
-	return HairGroupsLOD[GroupIndex].LODs[LODIndex].bVisible && (
-		HairGroupsLOD[GroupIndex].LODs[LODIndex].Simulation == EGroomOverrideType::Enable ||
-		(HairGroupsLOD[GroupIndex].LODs[LODIndex].Simulation == EGroomOverrideType::Auto && HairGroupsPhysics[GroupIndex].SolverSettings.EnableSimulation));
+	return GetHairGroupsLOD()[GroupIndex].LODs[LODIndex].bVisible && (
+		GetHairGroupsLOD()[GroupIndex].LODs[LODIndex].Simulation == EGroomOverrideType::Enable ||
+		(GetHairGroupsLOD()[GroupIndex].LODs[LODIndex].Simulation == EGroomOverrideType::Auto && GetHairGroupsPhysics()[GroupIndex].SolverSettings.EnableSimulation));
 }
 
 bool UGroomAsset::IsGlobalInterpolationEnable(int32 GroupIndex, int32 LODIndex) const
 {
-	if (GroupIndex < 0 || GroupIndex >= HairGroupsLOD.Num())
+	if (GroupIndex < 0 || GroupIndex >= GetHairGroupsLOD().Num())
 	{
 		return false;
 	}
 
-	if (LODIndex < 0 || LODIndex >= HairGroupsLOD[GroupIndex].LODs.Num())
+	if (LODIndex < 0 || LODIndex >= GetHairGroupsLOD()[GroupIndex].LODs.Num())
 	{
 		return false;
 	}
@@ -3667,19 +3667,19 @@ bool UGroomAsset::IsGlobalInterpolationEnable(int32 GroupIndex, int32 LODIndex)
 	// If the LOD index is not forced, then its value could be -1. In this case we return the 'global' asset value
 	if (LODIndex < 0)
 	{
-		return EnableGlobalInterpolation;
+		return GetEnableGlobalInterpolation();
 	}
 
 	return
-		HairGroupsLOD[GroupIndex].LODs[LODIndex].GlobalInterpolation == EGroomOverrideType::Enable ||
-		(HairGroupsLOD[GroupIndex].LODs[LODIndex].GlobalInterpolation == EGroomOverrideType::Auto && EnableGlobalInterpolation);
+		GetHairGroupsLOD()[GroupIndex].LODs[LODIndex].GlobalInterpolation == EGroomOverrideType::Enable ||
+		(GetHairGroupsLOD()[GroupIndex].LODs[LODIndex].GlobalInterpolation == EGroomOverrideType::Auto && GetEnableGlobalInterpolation());
 }
 
 bool UGroomAsset::IsSimulationEnable() const
 {
-	for (int32 GroupIndex = 0; GroupIndex < HairGroupsPlatformData.Num(); ++GroupIndex)
+	for (int32 GroupIndex = 0; GroupIndex < GetHairGroupsPlatformData().Num(); ++GroupIndex)
 	{
-		for (int32 LODIt = 0; LODIt < HairGroupsLOD.Num(); ++LODIt)
+		for (int32 LODIt = 0; LODIt < GetHairGroupsLOD().Num(); ++LODIt)
 		{
 			if (IsSimulationEnable(GroupIndex, LODIt))
 			{
@@ -3693,9 +3693,9 @@ bool UGroomAsset::IsSimulationEnable() const
 
 bool UGroomAsset::NeedsInterpolationData(int32 GroupIndex) const
 {
-	for (int32 LODIt = 0; LODIt < HairGroupsLOD.Num(); ++LODIt)
+	for (int32 LODIt = 0; LODIt < GetHairGroupsLOD().Num(); ++LODIt)
 	{
-		if (IsSimulationEnable(GroupIndex, LODIt) || IsGlobalInterpolationEnable(GroupIndex, LODIt) || EnableSimulationCache)
+		if (IsSimulationEnable(GroupIndex, LODIt) || IsGlobalInterpolationEnable(GroupIndex, LODIt) || GetEnableSimulationCache())
 		{
 			return true;
 		}
@@ -3706,7 +3706,7 @@ bool UGroomAsset::NeedsInterpolationData(int32 GroupIndex) const
 
 bool UGroomAsset::NeedsInterpolationData() const
 {
-	for (int32 GroupIndex = 0; GroupIndex < HairGroupsPlatformData.Num(); ++GroupIndex)
+	for (int32 GroupIndex = 0; GroupIndex < GetHairGroupsPlatformData().Num(); ++GroupIndex)
 	{
 		if (NeedsInterpolationData(GroupIndex))
 		{
@@ -3733,7 +3733,7 @@ bool FHairDescriptionGroups::IsValid() const
 int32 UGroomAsset::GetLODCount() const
 {
 	int32 MaxLODCount = -1;
-	for (const FHairGroupsLOD& S : HairGroupsLOD)
+	for (const FHairGroupsLOD& S : GetHairGroupsLOD())
 	{
 		MaxLODCount = FMath::Max(MaxLODCount, S.LODs.Num());
 	}
@@ -3750,7 +3750,7 @@ void UGroomAsset::StripLODs(const TArray<int32>& LODsToKeep, bool bRebuildResour
 		return;
 	}
 
-	const int32 GroupCount = HairGroupsLOD.Num();
+	const int32 GroupCount = GetHairGroupsLOD().Num();
 	int32 LODsTOKeepIndex = LODsToKeep.Num()-1;
 
 	// Remove the LOD settings prior to rebuild the LOD data
@@ -3764,7 +3764,7 @@ void UGroomAsset::StripLODs(const TArray<int32>& LODsToKeep, bool bRebuildResour
 
 		for (int32 GroupIt = 0; GroupIt < GroupCount; ++GroupIt)
 		{
-			HairGroupsLOD[GroupIt].LODs.RemoveAt(LODIt);
+			GetHairGroupsLOD()[GroupIt].LODs.RemoveAt(LODIt);
 		}
 
 		--LODsTOKeepIndex;
@@ -3776,16 +3776,16 @@ void UGroomAsset::StripLODs(const TArray<int32>& LODsToKeep, bool bRebuildResour
 		const FHairDescriptionGroups& LocalHairDescriptionGroups = GetHairDescriptionGroups();
 		for (int32 GroupIt = 0; GroupIt < GroupCount; ++GroupIt)
 		{
-			FHairGroupPlatformData& GroupData = HairGroupsPlatformData[GroupIt];
+			FHairGroupPlatformData& GroupData = GetHairGroupsPlatformData()[GroupIt];
 			InternalReleaseResource(GroupData.Strands.ClusterCullingResource);
 
 			BuildHairGroupCluster(
 				GroupIt,
 				LocalHairDescriptionGroups,
-				HairGroupsInterpolation,
-				HairGroupsLOD,
-				HairGroupsInfo,
-				HairGroupsPlatformData);
+				GetHairGroupsInterpolation(),
+				GetHairGroupsLOD(),
+				GetHairGroupsInfo(),
+				GetHairGroupsPlatformData());
 		}
 	}
 }
@@ -3795,7 +3795,7 @@ bool UGroomAsset::HasDebugData() const
 {
 	for (uint32 GroupIndex = 0, GroupCount = GetNumHairGroups(); GroupIndex < GroupCount; ++GroupIndex)
 	{
-		if (!HairGroupsPlatformData[GroupIndex].Debug.Data.IsValid())
+		if (!GetHairGroupsPlatformData()[GroupIndex].Debug.Data.IsValid())
 		{
 			return false;
 		}
@@ -3817,9 +3817,9 @@ void UGroomAsset::CreateDebugData()
 		FHairGroupInfo DummyInfo;
 		FHairStrandsDatas StrandsData;
 		FHairStrandsDatas GuidesData;
-		FGroomBuilder::BuildData(LocalHairDescriptionGroups.HairGroups[GroupIndex], HairGroupsInterpolation[GroupIndex], DummyInfo, StrandsData, GuidesData);
+		FGroomBuilder::BuildData(LocalHairDescriptionGroups.HairGroups[GroupIndex], GetHairGroupsInterpolation()[GroupIndex], DummyInfo, StrandsData, GuidesData);
 
-		FHairGroupPlatformData& GroupData = HairGroupsPlatformData[GroupIndex];
+		FHairGroupPlatformData& GroupData = GetHairGroupsPlatformData()[GroupIndex];
 		CreateHairStrandsDebugDatas(StrandsData, GroupData.Debug.Data);
 
 		if (GroupData.Debug.Data.IsValid())
@@ -3842,10 +3842,10 @@ void UGroomAsset::CreateDebugData()
 
 int32 UGroomAsset::GetMaterialIndex(FName MaterialSlotName) const
 {
-	const int32 SlotCount = HairGroupsMaterials.Num();
+	const int32 SlotCount = GetHairGroupsMaterials().Num();
 	for (int32 SlotIt = 0; SlotIt < SlotCount; ++SlotIt)
 	{
-		if (HairGroupsMaterials[SlotIt].SlotName == MaterialSlotName)
+		if (GetHairGroupsMaterials()[SlotIt].SlotName == MaterialSlotName)
 		{
 			return SlotIt;
 		}
@@ -3862,7 +3862,7 @@ bool UGroomAsset::IsMaterialSlotNameValid(FName MaterialSlotName) const
 TArray<FName> UGroomAsset::GetMaterialSlotNames() const
 {
 	TArray<FName> MaterialNames;
-	for (const FHairGroupsMaterial& Material : HairGroupsMaterials)
+	for (const FHairGroupsMaterial& Material : GetHairGroupsMaterials())
 	{
 		MaterialNames.Add(Material.SlotName);
 	}
@@ -3887,14 +3887,14 @@ bool InternalIsMaterialUsed(const TArray<T>& Groups, const FName& MaterialSlotNa
 
 bool UGroomAsset::IsMaterialUsed(int32 MaterialIndex) const
 {
-	if (MaterialIndex < 0 || MaterialIndex >= HairGroupsMaterials.Num())
+	if (MaterialIndex < 0 || MaterialIndex >= GetHairGroupsMaterials().Num())
 		return false;
 
-	const FName MaterialSlotName = HairGroupsMaterials[MaterialIndex].SlotName;
+	const FName MaterialSlotName = GetHairGroupsMaterials()[MaterialIndex].SlotName;
 	return 
-		InternalIsMaterialUsed(HairGroupsRendering, MaterialSlotName) ||
-		InternalIsMaterialUsed(HairGroupsCards, MaterialSlotName) ||
-		InternalIsMaterialUsed(HairGroupsMeshes, MaterialSlotName);
+		InternalIsMaterialUsed(GetHairGroupsRendering(), MaterialSlotName) ||
+		InternalIsMaterialUsed(GetHairGroupsCards(), MaterialSlotName) ||
+		InternalIsMaterialUsed(GetHairGroupsMeshes(), MaterialSlotName);
 }
 
 #if WITH_EDITOR
@@ -3926,14 +3926,14 @@ void UGroomAsset::SaveProceduralCards(uint32 DescIndex)
 
 	LLM_SCOPE_BYTAG(Groom);
 
-	if (DescIndex >= uint32(HairGroupsCards.Num()))
+	if (DescIndex >= uint32(GetHairGroupsCards().Num()))
 		return;
 
-	FHairGroupsCardsSourceDescription* Desc = &HairGroupsCards[DescIndex];
+	FHairGroupsCardsSourceDescription* Desc = &GetHairGroupsCards()[DescIndex];
 
 	const int32 GroupIndex = Desc->GroupIndex;
 	const int32 LODIndex = Desc->LODIndex;
-	if (GroupIndex >= HairGroupsPlatformData.Num())
+	if (GroupIndex >= GetHairGroupsPlatformData().Num())
 		return;
 
 	// 1. Convert old parameters (ClusterDecimation, bUseCards) to new parameters (GenerationType & CardsCount)
@@ -3943,7 +3943,7 @@ void UGroomAsset::SaveProceduralCards(uint32 DescIndex)
 		const bool bNeedConversion = ClusterSettings.ClusterDecimation > 0;
 		if (bNeedConversion)
 		{
-			const int32 MaxCardCount = HairGroupsPlatformData[GroupIndex].Strands.BulkData.GetNumCurves();
+			const int32 MaxCardCount = GetHairGroupsPlatformData()[GroupIndex].Strands.BulkData.GetNumCurves();
 
 			GeometrySettings.GenerationType = ClusterSettings.bUseGuide ? EHairCardsGenerationType::UseGuides : EHairCardsGenerationType::CardsCount;
 			GeometrySettings.CardsCount = FMath::Clamp(FMath::CeilToInt(ClusterSettings.ClusterDecimation * MaxCardCount), 1, MaxCardCount);
@@ -3966,7 +3966,7 @@ void UGroomAsset::SaveProceduralCards(uint32 DescIndex)
 		FHairGroupInfo DummyInfo;
 		FHairStrandsDatas StrandsData;
 		FHairStrandsDatas GuidesData;
-		FGroomBuilder::BuildData(LocalHairDescriptionGroups.HairGroups[GroupIndex], HairGroupsInterpolation[GroupIndex], DummyInfo, StrandsData, GuidesData);
+		FGroomBuilder::BuildData(LocalHairDescriptionGroups.HairGroups[GroupIndex], GetHairGroupsInterpolation()[GroupIndex], DummyInfo, StrandsData, GuidesData);
 
 		// 2.2 Build cards geometry
 		FHairCardsBuilder::BuildGeometry(
@@ -3994,11 +3994,11 @@ void UGroomAsset::SaveProceduralCards(uint32 DescIndex)
 	// 4. Save output asset (geometry, and enqueue texture saving)
 	{
 		// Create a static meshes with the vertex data
-		if (HairGroupsCards[DescIndex].ProceduralMesh == nullptr)
+		if (GetHairGroupsCards()[DescIndex].ProceduralMesh == nullptr)
 		{
 			const FString PackageName = GetOutermost()->GetName();
 			const FString SuffixName = FText::Format(LOCTEXT("CardsStatisMesh", "_CardsMesh_Group{0}_LOD{1}"), FText::AsNumber(GroupIndex), FText::AsNumber(LODIndex)).ToString();
-			HairGroupsCards[DescIndex].ProceduralMesh = FHairStrandsCore::CreateStaticMesh(PackageName, SuffixName);
+			GetHairGroupsCards()[DescIndex].ProceduralMesh = FHairStrandsCore::CreateStaticMesh(PackageName, SuffixName);
 		}
 
 		// Convert procedural cards data to cards data prior to export
@@ -4007,9 +4007,9 @@ void UGroomAsset::SaveProceduralCards(uint32 DescIndex)
 			CardData.Cards = Q.ProceduralData.Cards;
 		}
 
-		FHairCardsBuilder::ExportGeometry(CardData, HairGroupsCards[DescIndex].ProceduralMesh);
-		FHairStrandsCore::SaveAsset(HairGroupsCards[DescIndex].ProceduralMesh);
-		HairGroupsCards[DescIndex].ProceduralMeshKey = GroomDerivedDataCacheUtils::BuildCardsDerivedDataKeySuffix(this, *Desc);
+		FHairCardsBuilder::ExportGeometry(CardData, GetHairGroupsCards()[DescIndex].ProceduralMesh);
+		FHairStrandsCore::SaveAsset(GetHairGroupsCards()[DescIndex].ProceduralMesh);
+		GetHairGroupsCards()[DescIndex].ProceduralMeshKey = GroomDerivedDataCacheUtils::BuildCardsDerivedDataKeySuffix(this, *Desc);
 	}
 }
 
@@ -4071,7 +4071,7 @@ bool UGroomAsset::GetHairStrandsDatas(
 		check(GroupIndex == HairGroup.Info.GroupID);
 
 		FHairGroupInfo DummyInfo;
-		FGroomBuilder::BuildData(HairGroup, HairGroupsInterpolation[GroupIndex], DummyInfo, OutStrandsData, OutGuidesData);
+		FGroomBuilder::BuildData(HairGroup, GetHairGroupsInterpolation()[GroupIndex], DummyInfo, OutStrandsData, OutGuidesData);
 	}
 	return bIsValid;
 }
@@ -4090,7 +4090,7 @@ bool UGroomAsset::GetHairCardsGuidesDatas(
 	}
 
 	int32 SourceIt = 0;
-	if (FHairGroupsCardsSourceDescription* Desc = GetSourceDescription(HairGroupsCards, GroupIndex, LODIndex, SourceIt))
+	if (FHairGroupsCardsSourceDescription* Desc = GetSourceDescription(GetHairGroupsCards(), GroupIndex, LODIndex, SourceIt))
 	{
 		// 1. Load geometry data, if any
 		UStaticMesh* CardsMesh = nullptr;
@@ -4121,16 +4121,16 @@ bool UGroomAsset::GetHairCardsGuidesDatas(
 
 void UGroomAsset::GetResourceSizeEx(FResourceSizeEx& CumulativeResourceSize)
 {
-	CumulativeResourceSize.AddDedicatedSystemMemoryBytes(HairGroupsRendering.GetAllocatedSize());
-	CumulativeResourceSize.AddDedicatedSystemMemoryBytes(HairGroupsPhysics.GetAllocatedSize());
-	CumulativeResourceSize.AddDedicatedSystemMemoryBytes(HairGroupsInterpolation.GetAllocatedSize());
-	CumulativeResourceSize.AddDedicatedSystemMemoryBytes(HairGroupsLOD.GetAllocatedSize());
-	CumulativeResourceSize.AddDedicatedSystemMemoryBytes(HairGroupsCards.GetAllocatedSize());
-	CumulativeResourceSize.AddDedicatedSystemMemoryBytes(HairGroupsMeshes.GetAllocatedSize());
-	CumulativeResourceSize.AddDedicatedSystemMemoryBytes(HairGroupsMaterials.GetAllocatedSize());
-	CumulativeResourceSize.AddDedicatedSystemMemoryBytes(HairGroupsPlatformData.GetAllocatedSize());
+	CumulativeResourceSize.AddDedicatedSystemMemoryBytes(GetHairGroupsRendering().GetAllocatedSize());
+	CumulativeResourceSize.AddDedicatedSystemMemoryBytes(GetHairGroupsPhysics().GetAllocatedSize());
+	CumulativeResourceSize.AddDedicatedSystemMemoryBytes(GetHairGroupsInterpolation().GetAllocatedSize());
+	CumulativeResourceSize.AddDedicatedSystemMemoryBytes(GetHairGroupsLOD().GetAllocatedSize());
+	CumulativeResourceSize.AddDedicatedSystemMemoryBytes(GetHairGroupsCards().GetAllocatedSize());
+	CumulativeResourceSize.AddDedicatedSystemMemoryBytes(GetHairGroupsMeshes().GetAllocatedSize());
+	CumulativeResourceSize.AddDedicatedSystemMemoryBytes(GetHairGroupsMaterials().GetAllocatedSize());
+	CumulativeResourceSize.AddDedicatedSystemMemoryBytes(GetHairGroupsPlatformData().GetAllocatedSize());
 	
-	for (const FHairGroupPlatformData & GroupData : HairGroupsPlatformData)
+	for (const FHairGroupPlatformData & GroupData : GetHairGroupsPlatformData())
 	{		
 		CumulativeResourceSize.AddDedicatedVideoMemoryBytes(GroupData.Guides.GetResourcesSize());
 		CumulativeResourceSize.AddDedicatedVideoMemoryBytes(GroupData.Strands.GetResourcesSize());
@@ -4162,3 +4162,73 @@ FName UGroomAsset::GetAssetPathName(int32 LODIndex)
 
 #undef LOCTEXT_NAMESPACE
 
+#define GROOMASSET_DEFINE_MEMBER_NAME(Name) \
+	FName UGroomAsset::Get##Name##MemberName()\
+	{\
+		PRAGMA_DISABLE_DEPRECATION_WARNINGS\
+		return FName(TEXT(#Name));\
+		PRAGMA_ENABLE_DEPRECATION_WARNINGS\
+	}
+
+#define GROOMASSET_DEFINE_GET(Type, Name, Access, Const0, Const1) \
+	Const0 Type Access UGroomAsset::Get##Name() Const1\
+	{\
+		PRAGMA_DISABLE_DEPRECATION_WARNINGS\
+		return Name;\
+		PRAGMA_ENABLE_DEPRECATION_WARNINGS\
+	}
+
+#define GROOMASSET_DEFINE_SET(Type, Name, Access, Const) \
+	void UGroomAsset::Set##Name(Const Type Access In##Name)\
+	{\
+		PRAGMA_DISABLE_DEPRECATION_WARNINGS\
+		Name = In##Name;\
+		PRAGMA_ENABLE_DEPRECATION_WARNINGS\
+	}
+
+#define DEFINE_GROOM_MEMBER_ACCESSOR_BP(Type, Name)\
+	GROOMASSET_DEFINE_MEMBER_NAME(Name)\
+	UFUNCTION(BlueprintGetter)\
+	GROOMASSET_DEFINE_GET(Type, Name, , , const)\
+	UFUNCTION(BlueprintSetter)\
+	GROOMASSET_DEFINE_SET(Type, Name, , )
+
+#define DEFINE_GROOM_MEMBER_ACCESSOR_BP_ARRAY(Type, Name)\
+	GROOMASSET_DEFINE_MEMBER_NAME(Name)\
+	UFUNCTION(BlueprintGetter)\
+	GROOMASSET_DEFINE_GET(Type, Name, &, , )\
+	UFUNCTION(BlueprintSetter)\
+	GROOMASSET_DEFINE_SET(Type, Name, &, const)\
+	GROOMASSET_DEFINE_GET(Type, Name, &, const, const)
+
+#define DEFINE_GROOM_MEMBER_ACCESSOR(Type, Name)\
+	GROOMASSET_DEFINE_MEMBER_NAME(Name)\
+	GROOMASSET_DEFINE_GET(Type, Name, , , const)\
+	GROOMASSET_DEFINE_SET(Type, Name, , )
+
+#define DEFINE_GROOM_MEMBER_ACCESSOR_ARRAY(Type, Name)\
+	GROOMASSET_DEFINE_MEMBER_NAME(Name)\
+	GROOMASSET_DEFINE_GET(Type, Name, &, , )\
+	GROOMASSET_DEFINE_SET(Type, Name, &, const)\
+	GROOMASSET_DEFINE_GET(Type, Name, &, const, const)
+
+// Define most of the groom member accessor
+DEFINE_GROOM_MEMBER_ACCESSOR_BP_ARRAY(TArray<FHairGroupsRendering>, HairGroupsRendering);
+DEFINE_GROOM_MEMBER_ACCESSOR_BP_ARRAY(TArray<FHairGroupsPhysics>, HairGroupsPhysics);
+DEFINE_GROOM_MEMBER_ACCESSOR_BP_ARRAY(TArray<FHairGroupsInterpolation>, HairGroupsInterpolation);
+DEFINE_GROOM_MEMBER_ACCESSOR_BP_ARRAY(TArray<FHairGroupsLOD>, HairGroupsLOD);
+DEFINE_GROOM_MEMBER_ACCESSOR_BP_ARRAY(TArray<FHairGroupsCardsSourceDescription>, HairGroupsCards);
+DEFINE_GROOM_MEMBER_ACCESSOR_BP_ARRAY(TArray<FHairGroupsMeshesSourceDescription>, HairGroupsMeshes);
+DEFINE_GROOM_MEMBER_ACCESSOR_BP_ARRAY(TArray<FHairGroupsMaterial>, HairGroupsMaterials);
+DEFINE_GROOM_MEMBER_ACCESSOR_BP_ARRAY(TArray<int32>, DeformedGroupSections);
+
+DEFINE_GROOM_MEMBER_ACCESSOR_BP(bool, EnableGlobalInterpolation);
+DEFINE_GROOM_MEMBER_ACCESSOR_BP(bool, EnableSimulationCache);
+DEFINE_GROOM_MEMBER_ACCESSOR_BP(EGroomInterpolationType, HairInterpolationType);
+DEFINE_GROOM_MEMBER_ACCESSOR_BP(USkeletalMesh*, RiggedSkeletalMesh);
+
+DEFINE_GROOM_MEMBER_ACCESSOR(FPerPlatformInt, MinLOD);
+DEFINE_GROOM_MEMBER_ACCESSOR(FPerPlatformBool, DisableBelowMinLodStripping);
+DEFINE_GROOM_MEMBER_ACCESSOR_ARRAY(TArray<float>, EffectiveLODBias);
+DEFINE_GROOM_MEMBER_ACCESSOR_ARRAY(TArray<FHairGroupPlatformData>, HairGroupsPlatformData);
+DEFINE_GROOM_MEMBER_ACCESSOR_ARRAY(TArray<FHairGroupInfoWithVisibility>, HairGroupsInfo);
diff --git a/Engine/Plugins/Runtime/HairStrands/Source/HairStrandsCore/Private/GroomBindingAsset.cpp b/Engine/Plugins/Runtime/HairStrands/Source/HairStrandsCore/Private/GroomBindingAsset.cpp
--- a/Engine/Plugins/Runtime/HairStrands/Source/HairStrandsCore/Private/GroomBindingAsset.cpp
+++ b/Engine/Plugins/Runtime/HairStrands/Source/HairStrandsCore/Private/GroomBindingAsset.cpp
@@ -552,7 +552,7 @@ static bool DoesGroomNeedStrandsBinding(const UGroomAsset* InGroom, uint32 InGro
 {
 	if (IsHairStrandsEnabled(EHairStrandsShaderType::Strands))
 	{
-		for (const FHairLODSettings& LODSettings : InGroom->HairGroupsLOD[InGroupIndex].LODs)
+		for (const FHairLODSettings& LODSettings : InGroom->GetHairGroupsLOD()[InGroupIndex].LODs)
 		{
 			if (LODSettings.bVisible && LODSettings.GeometryType == EGroomGeometryType::Strands)
 			{
@@ -603,7 +603,7 @@ bool UGroomBindingAsset::IsCompatible(const UGroomAsset* InGroom, const UGroomBi
 		{
 			// Guides
 			{
-				const uint32 GroomCount = InGroom->HairGroupsPlatformData[GroupIt].Guides.BulkData.GetNumCurves();
+				const uint32 GroomCount = InGroom->GetHairGroupsPlatformData()[GroupIt].Guides.BulkData.GetNumCurves();
 				const uint32 BindingCount = InBinding->GetGroupInfos()[GroupIt].SimRootCount;
 
 				if (GroomCount != 0 && GroomCount != BindingCount)
@@ -622,7 +622,7 @@ bool UGroomBindingAsset::IsCompatible(const UGroomAsset* InGroom, const UGroomBi
 			const bool bNeedStrandsRoot = DoesGroomNeedStrandsBinding(InGroom, GroupIt);
 			if (bNeedStrandsRoot)
 			{
-				const uint32 GroomCount = InGroom->HairGroupsPlatformData[GroupIt].Strands.BulkData.GetNumCurves();
+				const uint32 GroomCount = InGroom->GetHairGroupsPlatformData()[GroupIt].Strands.BulkData.GetNumCurves();
 				const uint32 BindingCount = InBinding->GetGroupInfos()[GroupIt].RenRootCount;
 
 				// Groom may have stripped strands data so GroomCount would be 0
diff --git a/Engine/Plugins/Runtime/HairStrands/Source/HairStrandsCore/Private/GroomBindingBuilder.cpp b/Engine/Plugins/Runtime/HairStrands/Source/HairStrandsCore/Private/GroomBindingBuilder.cpp
--- a/Engine/Plugins/Runtime/HairStrands/Source/HairStrandsCore/Private/GroomBindingBuilder.cpp
+++ b/Engine/Plugins/Runtime/HairStrands/Source/HairStrandsCore/Private/GroomBindingBuilder.cpp
@@ -2090,11 +2090,11 @@ static bool InternalBuildBinding_CPU(UGroomBindingAsset* BindingAsset, uint32 In
 		const uint32 GroupCount = GroomAsset->GetNumHairGroups();
 		const uint32 MeshLODCount = TargetMeshData->GetNumLODs();
 
-		check(InGroupIndex < uint32(GroomAsset->HairGroupsPlatformData.Num()));
+		check(InGroupIndex < uint32(GroomAsset->GetHairGroupsPlatformData().Num()));
 
 		// Check if root data are needs for strands
 		bool bNeedStrandsRoot = false;
-		for (const FHairLODSettings& LODSettings : GroomAsset->HairGroupsLOD[InGroupIndex].LODs)
+		for (const FHairLODSettings& LODSettings : GroomAsset->GetHairGroupsLOD()[InGroupIndex].LODs)
 		{
 			if (LODSettings.GeometryType == EGroomGeometryType::Strands)
 			{
@@ -2109,7 +2109,7 @@ static bool InternalBuildBinding_CPU(UGroomBindingAsset* BindingAsset, uint32 In
 		GroomAsset->GetHairStrandsDatas(InGroupIndex, StrandsData, GuidesData);
 
 		// 1.2 Init root data for guides/strands/cards
-		const FHairGroupPlatformData& GroupData = GroomAsset->HairGroupsPlatformData[InGroupIndex];
+		const FHairGroupPlatformData& GroupData = GroomAsset->GetHairGroupsPlatformData()[InGroupIndex];
 		{
 			// Guides
 			InitHairStrandsRootData(OutData.SimRootData, &GuidesData, MeshLODCount, NumInterpolationPoints);
@@ -2205,7 +2205,7 @@ static bool InternalBuildBinding_CPU(UGroomBindingAsset* BindingAsset, uint32 In
 			const uint32 CardsLODCount = OutData.CardsRootData.Num();
 			for (uint32 CardsLODIt = 0; CardsLODIt < CardsLODCount; ++CardsLODIt)
 			{
-				if (BindingAsset->GetGroom()->HairGroupsPlatformData[InGroupIndex].Cards.IsValid(CardsLODIt))
+				if (BindingAsset->GetGroom()->GetHairGroupsPlatformData()[InGroupIndex].Cards.IsValid(CardsLODIt))
 				{
 					FHairStrandsDatas LODGuidesData;
 					const bool bIsValid = GroomAsset->GetHairCardsGuidesDatas(InGroupIndex, CardsLODIt, LODGuidesData);
diff --git a/Engine/Plugins/Runtime/HairStrands/Source/HairStrandsCore/Private/GroomComponent.cpp b/Engine/Plugins/Runtime/HairStrands/Source/HairStrandsCore/Private/GroomComponent.cpp
--- a/Engine/Plugins/Runtime/HairStrands/Source/HairStrandsCore/Private/GroomComponent.cpp
+++ b/Engine/Plugins/Runtime/HairStrands/Source/HairStrandsCore/Private/GroomComponent.cpp
@@ -138,9 +138,12 @@ static FHairGroupDesc GetGroomGroupsDesc(const UGroomAsset* Asset, const UGroomC
 		return FHairGroupDesc();
 	}
 
+	const FHairGroupsRendering& LocalRendering = Asset->GetHairGroupsRendering()[GroupIndex];
+	const FHairGroupPlatformData& LocalPlatformData = Asset->GetHairGroupsPlatformData()[GroupIndex];
+
 	FHairGroupDesc O = Component->GroomGroupsDesc[GroupIndex];
-	O.HairLength = Asset->HairGroupsPlatformData[GroupIndex].Strands.BulkData.Header.MaxLength;
-	O.LODBias 	 = Asset->EffectiveLODBias[GroupIndex] > 0 ? FMath::Max(O.LODBias, Asset->EffectiveLODBias[GroupIndex]) : O.LODBias;
+	O.HairLength = LocalPlatformData.Strands.BulkData.Header.MaxLength;
+	O.LODBias 	 = Asset->GetEffectiveLODBias()[GroupIndex] > 0 ? FMath::Max(O.LODBias, Asset->GetEffectiveLODBias()[GroupIndex]) : O.LODBias;
 
 	// Width can be overridden at 3 levels:
 	// * Instance width override
@@ -149,19 +152,19 @@ static FHairGroupDesc GetGroomGroupsDesc(const UGroomAsset* Asset, const UGroomC
 	if (!O.HairWidth_Override)
 	{ 
 		O.HairWidth = 
-			Asset->HairGroupsRendering[GroupIndex].GeometrySettings.HairWidth_Override ? 
-			Asset->HairGroupsRendering[GroupIndex].GeometrySettings.HairWidth : 
-			Asset->HairGroupsPlatformData[GroupIndex].Strands.BulkData.Header.MaxRadius * 2.0f;
+			LocalRendering.GeometrySettings.HairWidth_Override ? 
+			LocalRendering.GeometrySettings.HairWidth : 
+			LocalPlatformData.Strands.BulkData.Header.MaxRadius * 2.0f;
 	}
 
-	if (!O.HairRootScale_Override)				{ O.HairRootScale				= Asset->HairGroupsRendering[GroupIndex].GeometrySettings.HairRootScale;				}
-	if (!O.HairTipScale_Override)				{ O.HairTipScale				= Asset->HairGroupsRendering[GroupIndex].GeometrySettings.HairTipScale;					}
-	if (!O.bSupportVoxelization_Override)		{ O.bSupportVoxelization		= Asset->HairGroupsRendering[GroupIndex].ShadowSettings.bVoxelize;						}
-	if (!O.HairShadowDensity_Override)			{ O.HairShadowDensity			= Asset->HairGroupsRendering[GroupIndex].ShadowSettings.HairShadowDensity;				}
-	if (!O.HairRaytracingRadiusScale_Override)	{ O.HairRaytracingRadiusScale	= Asset->HairGroupsRendering[GroupIndex].ShadowSettings.HairRaytracingRadiusScale;		}
-	if (!O.bUseHairRaytracingGeometry_Override) { O.bUseHairRaytracingGeometry  = Asset->HairGroupsRendering[GroupIndex].ShadowSettings.bUseHairRaytracingGeometry;		}
-	if (!O.bUseStableRasterization_Override)	{ O.bUseStableRasterization		= Asset->HairGroupsRendering[GroupIndex].AdvancedSettings.bUseStableRasterization;		}
-	if (!O.bScatterSceneLighting_Override)		{ O.bScatterSceneLighting		= Asset->HairGroupsRendering[GroupIndex].AdvancedSettings.bScatterSceneLighting;		}
+	if (!O.HairRootScale_Override)				{ O.HairRootScale				= LocalRendering.GeometrySettings.HairRootScale;				}
+	if (!O.HairTipScale_Override)				{ O.HairTipScale				= LocalRendering.GeometrySettings.HairTipScale;					}
+	if (!O.bSupportVoxelization_Override)		{ O.bSupportVoxelization		= LocalRendering.ShadowSettings.bVoxelize;						}
+	if (!O.HairShadowDensity_Override)			{ O.HairShadowDensity			= LocalRendering.ShadowSettings.HairShadowDensity;				}
+	if (!O.HairRaytracingRadiusScale_Override)	{ O.HairRaytracingRadiusScale	= LocalRendering.ShadowSettings.HairRaytracingRadiusScale;		}
+	if (!O.bUseHairRaytracingGeometry_Override) { O.bUseHairRaytracingGeometry  = LocalRendering.ShadowSettings.bUseHairRaytracingGeometry;		}
+	if (!O.bUseStableRasterization_Override)	{ O.bUseStableRasterization		= LocalRendering.AdvancedSettings.bUseStableRasterization;		}
+	if (!O.bScatterSceneLighting_Override)		{ O.bScatterSceneLighting		= LocalRendering.AdvancedSettings.bScatterSceneLighting;		}
 
 	return O;
 }
@@ -353,12 +356,12 @@ public:
 		const EShaderPlatform ShaderPlatform = GetScene().GetShaderPlatform();
 
 		const int32 GroupCount = Component->GroomAsset->GetNumHairGroups();
-		check(Component->GroomAsset->HairGroupsPlatformData.Num() == Component->HairGroupInstances.Num());
+		check(Component->GroomAsset->GetHairGroupsPlatformData().Num() == Component->HairGroupInstances.Num());
 		for (int32 GroupIt=0; GroupIt<GroupCount; GroupIt++)
 		{
-			const bool bIsVisible = Component->GroomAsset->HairGroupsInfo[GroupIt].bIsVisible;
+			const bool bIsVisible = Component->GroomAsset->GetHairGroupsInfo()[GroupIt].bIsVisible;
 
-			const FHairGroupPlatformData& InGroupData = Component->GroomAsset->HairGroupsPlatformData[GroupIt];
+			const FHairGroupPlatformData& InGroupData = Component->GroomAsset->GetHairGroupsPlatformData()[GroupIt];
 			FHairGroupInstance* HairInstance = Component->HairGroupInstances[GroupIt];
 			check(HairInstance->HairGroupPublicData);
 			HairInstance->bForceCards = Component->bUseCards;
@@ -405,7 +408,7 @@ public:
 			// Material - Strands
 			if (IsHairStrandsEnabled(EHairStrandsShaderType::Strands, ShaderPlatform))
 			{
-				const int32 SlotIndex = Component->GroomAsset->GetMaterialIndex(Component->GroomAsset->HairGroupsRendering[GroupIt].MaterialSlotName);
+				const int32 SlotIndex = Component->GroomAsset->GetMaterialIndex(Component->GroomAsset->GetHairGroupsRendering()[GroupIt].MaterialSlotName);
 				const UMaterialInterface* Material = Component->GetMaterial(GetMaterialIndexWithFallback(SlotIndex), EHairGeometryType::Strands, true);
 				HairGroupMaterialProxies[GroupIt].Strands = Material ? Material->GetRenderProxy() : nullptr;
 			}
@@ -421,7 +424,7 @@ public:
 					{
 						// Material
 						int32 SlotIndex = INDEX_NONE;
-						for (const FHairGroupsCardsSourceDescription& Desc : Component->GroomAsset->HairGroupsCards)
+						for (const FHairGroupsCardsSourceDescription& Desc : Component->GroomAsset->GetHairGroupsCards())
 						{
 							if (Desc.GroupIndex == GroupIt && Desc.LODIndex == CardsLODIndex)
 							{
@@ -447,7 +450,7 @@ public:
 					{
 						// Material
 						int32 SlotIndex = INDEX_NONE;
-						for (const FHairGroupsMeshesSourceDescription& Desc : Component->GroomAsset->HairGroupsMeshes)
+						for (const FHairGroupsMeshesSourceDescription& Desc : Component->GroomAsset->GetHairGroupsMeshes())
 						{
 							if (Desc.GroupIndex == GroupIt && Desc.LODIndex == MeshesLODIndex)
 							{
@@ -1328,7 +1331,7 @@ void UGroomComponent::CreateHairSimulation(const int32 GroupIndex, const int32 L
 	if (GroupIndex < NiagaraComponents.Num())
 	{
 		TObjectPtr<UNiagaraComponent>& NiagaraComponent = NiagaraComponents[GroupIndex];
-		if (GroomAsset && (GroupIndex < GroomAsset->HairGroupsPhysics.Num()) && (LODIndex < GroomAsset->HairGroupsLOD[GroupIndex].LODs.Num()) && IsSimulationEnable(GroupIndex, LODIndex))
+		if (GroomAsset && (GroupIndex < GroomAsset->GetHairGroupsPhysics().Num()) && (LODIndex < GroomAsset->GetHairGroupsLOD()[GroupIndex].LODs.Num()) && IsSimulationEnable(GroupIndex, LODIndex))
 		{
 			if (!NiagaraComponent)
 			{
@@ -1347,10 +1350,10 @@ void UGroomComponent::CreateHairSimulation(const int32 GroupIndex, const int32 L
 				if (!CosseratRodsSystem) CosseratRodsSystem = LoadObject<UNiagaraSystem>(nullptr, TEXT("/HairStrands/Emitters/StableRodsSystem.StableRodsSystem"));
 
 				UNiagaraSystem* NiagaraAsset =
-					(GroomAsset->HairGroupsPhysics[GroupIndex].SolverSettings.NiagaraSolver == EGroomNiagaraSolvers::AngularSprings) ? ToRawPtr(AngularSpringsSystem) :
-					(GroomAsset->HairGroupsPhysics[GroupIndex].SolverSettings.NiagaraSolver == EGroomNiagaraSolvers::CosseratRods) ? ToRawPtr(CosseratRodsSystem) :
-					GroomAsset->HairGroupsPhysics[GroupIndex].SolverSettings.CustomSystem.LoadSynchronous();
-				NiagaraComponent->SetVisibleFlag(SimulationSettings.SimulationSetup.bDebugSimulation || GroomAsset->HairGroupsPhysics[GroupIndex].SolverSettings.bForceVisible);
+					(GroomAsset->GetHairGroupsPhysics()[GroupIndex].SolverSettings.NiagaraSolver == EGroomNiagaraSolvers::AngularSprings) ? ToRawPtr(AngularSpringsSystem) :
+					(GroomAsset->GetHairGroupsPhysics()[GroupIndex].SolverSettings.NiagaraSolver == EGroomNiagaraSolvers::CosseratRods) ? ToRawPtr(CosseratRodsSystem) :
+					GroomAsset->GetHairGroupsPhysics()[GroupIndex].SolverSettings.CustomSystem.LoadSynchronous();
+				NiagaraComponent->SetVisibleFlag(SimulationSettings.SimulationSetup.bDebugSimulation || GroomAsset->GetHairGroupsPhysics()[GroupIndex].SolverSettings.bForceVisible);
 				NiagaraComponent->SetAsset(NiagaraAsset);
 				NiagaraComponent->ReinitializeSystem();
 			}
@@ -1365,7 +1368,7 @@ void UGroomComponent::CreateHairSimulation(const int32 GroupIndex, const int32 L
 
 void UGroomComponent::UpdateHairSimulation()
 {
-	const int32 NumGroups = GroomAsset ? GroomAsset->HairGroupsPhysics.Num() : 0;
+	const int32 NumGroups = GroomAsset ? GroomAsset->GetHairGroupsPhysics().Num() : 0;
 	const int32 NumComponents = FMath::Max(NumGroups, NiagaraComponents.Num());
 
 	NiagaraComponents.SetNumZeroed(NumComponents);
@@ -1394,7 +1397,7 @@ void UGroomComponent::SwitchSimulationLOD(const int32 PreviousLOD, const int32 C
 	if (GroomAsset && PreviousLOD != CurrentLOD)
 	{
 		bool bRequiresSimulationUpdate = false;
-		for (int32 GroupIt = 0, GroupCount = GroomAsset->HairGroupsPlatformData.Num(); GroupIt < GroupCount; ++GroupIt)
+		for (int32 GroupIt = 0, GroupCount = GroomAsset->GetHairGroupsPlatformData().Num(); GroupIt < GroupCount; ++GroupIt)
 		{
 			if ((IsSimulationEnable(GroupIt, PreviousLOD) != IsSimulationEnable(GroupIt, CurrentLOD)))
 			{
@@ -1577,7 +1580,7 @@ void UGroomComponent::SetForcedLOD(int32 CurrLODIndex)
 		bool bNeedSimulationNeed = false;
 		if (GroomAsset)
 		{
-			for (int32 GroupIt = 0, GroupCount = GroomAsset->HairGroupsPlatformData.Num(); GroupIt < GroupCount; ++GroupIt)
+			for (int32 GroupIt = 0, GroupCount = GroomAsset->GetHairGroupsPlatformData().Num(); GroupIt < GroupCount; ++GroupIt)
 			{
 				for (uint32 LODIt = 0, LODCount = GroomAsset->GetLODCount(); LODIt < LODCount; ++LODIt)
 				{
@@ -1772,9 +1775,9 @@ FBoxSphereBounds UGroomComponent::CalcBounds(const FTransform& InLocalToWorld) c
 		FBox LocalHairBound(EForceInit::ForceInitToZero);
 		if (!GroomCacheBuffers.IsValid())
 		{
-			for (uint32 GroupIndex = 0, GroupCount = GroomAsset->HairGroupsPlatformData.Num(); GroupIndex<GroupCount; ++GroupIndex)
+			for (uint32 GroupIndex = 0, GroupCount = GroomAsset->GetHairGroupsPlatformData().Num(); GroupIndex<GroupCount; ++GroupIndex)
 			{
-				const FHairGroupPlatformData& GroupData = GroomAsset->HairGroupsPlatformData[GroupIndex];
+				const FHairGroupPlatformData& GroupData = GroomAsset->GetHairGroupsPlatformData()[GroupIndex];
 				if (IsHairStrandsEnabled(EHairStrandsShaderType::Strands) && GroupData.Strands.HasValidData())
 				{
 					// The hair width is not scaled by the transform scale. Since we increase the local bound by the hair width, 
@@ -1908,7 +1911,7 @@ int32 UGroomComponent::GetNumMaterials() const
 {
 	if (GroomAsset)
 	{
-		return FMath::Max(GroomAsset->HairGroupsMaterials.Num(), 1);
+		return FMath::Max(GroomAsset->GetHairGroupsMaterials().Num(), 1);
 	}
 	return 1;
 }
@@ -1919,9 +1922,9 @@ UMaterialInterface* UGroomComponent::GetMaterial(int32 ElementIndex, EHairGeomet
 
 	bool bUseHairDefaultMaterial = false;
 
-	if (!OverrideMaterial && GroomAsset && ElementIndex < GroomAsset->HairGroupsMaterials.Num())
+	if (!OverrideMaterial && GroomAsset && ElementIndex < GroomAsset->GetHairGroupsMaterials().Num())
 	{
-		if (UMaterialInterface* Material = GroomAsset->HairGroupsMaterials[ElementIndex].Material)
+		if (UMaterialInterface* Material = GroomAsset->GetHairGroupsMaterials()[ElementIndex].Material)
 		{
 			OverrideMaterial = Material;
 		}
@@ -1968,13 +1971,13 @@ EHairGeometryType UGroomComponent::GetMaterialGeometryType(int32 ElementIndex) c
 	}
 
 	const EShaderPlatform Platform = GetScene() ? GetScene()->GetShaderPlatform() : EShaderPlatform::SP_NumPlatforms;
-	for (uint32 GroupIt = 0, GroupCount = GroomAsset->HairGroupsRendering.Num(); GroupIt < GroupCount; ++GroupIt)
+	for (uint32 GroupIt = 0, GroupCount = GroomAsset->GetHairGroupsRendering().Num(); GroupIt < GroupCount; ++GroupIt)
 	{
 		// Material - Strands
-		const FHairGroupPlatformData& InGroupData = GroomAsset->HairGroupsPlatformData[GroupIt];
+		const FHairGroupPlatformData& InGroupData = GroomAsset->GetHairGroupsPlatformData()[GroupIt];
 		if (IsHairStrandsEnabled(EHairStrandsShaderType::Strands, Platform))
 		{
-			const int32 SlotIndex = GroomAsset->GetMaterialIndex(GroomAsset->HairGroupsRendering[GroupIt].MaterialSlotName);
+			const int32 SlotIndex = GroomAsset->GetMaterialIndex(GroomAsset->GetHairGroupsRendering()[GroupIt].MaterialSlotName);
 			if (SlotIndex == ElementIndex)
 			{
 				return EHairGeometryType::Strands;
@@ -1991,7 +1994,7 @@ EHairGeometryType UGroomComponent::GetMaterialGeometryType(int32 ElementIndex) c
 				{
 					// Material
 					int32 SlotIndex = INDEX_NONE;
-					for (const FHairGroupsCardsSourceDescription& Desc : GroomAsset->HairGroupsCards)
+					for (const FHairGroupsCardsSourceDescription& Desc : GroomAsset->GetHairGroupsCards())
 					{
 						if (Desc.GroupIndex == GroupIt && Desc.LODIndex == CardsLODIndex)
 						{
@@ -2018,7 +2021,7 @@ EHairGeometryType UGroomComponent::GetMaterialGeometryType(int32 ElementIndex) c
 				{
 					// Material
 					int32 SlotIndex = INDEX_NONE;
-					for (const FHairGroupsMeshesSourceDescription& Desc : GroomAsset->HairGroupsMeshes)
+					for (const FHairGroupsMeshesSourceDescription& Desc : GroomAsset->GetHairGroupsMeshes())
 					{
 						if (Desc.GroupIndex == GroupIt && Desc.LODIndex == MeshesLODIndex)
 						{
@@ -2052,11 +2055,11 @@ FHairStrandsRestResource* UGroomComponent::GetGuideStrandsRestResource(uint32 Gr
 		return nullptr;
 	}
 
-	if (!GroomAsset->HairGroupsPlatformData[GroupIndex].Guides.IsValid())
+	if (!GroomAsset->GetHairGroupsPlatformData()[GroupIndex].Guides.IsValid())
 	{
 		return nullptr;
 	}
-	return GroomAsset->HairGroupsPlatformData[GroupIndex].Guides.RestResource;
+	return GroomAsset->GetHairGroupsPlatformData()[GroupIndex].Guides.RestResource;
 }
 
 FHairStrandsDeformedResource* UGroomComponent::GetGuideStrandsDeformedResource(uint32 GroupIndex)
@@ -2124,7 +2127,7 @@ void UGroomComponent::UpdateSimulatedGroups()
 				Instance->Strands.HairInterpolationType = EHairInterpolationType::NoneSkinning;
 				if (bIsStrandsEnabled && LocalGroomAsset)
 				{
-					Instance->Strands.HairInterpolationType = ToHairInterpolationType(LocalGroomAsset->HairInterpolationType);
+					Instance->Strands.HairInterpolationType = ToHairInterpolationType(LocalGroomAsset->GetHairInterpolationType());
 				}
 				if (Instance->Guides.IsValid())
 				{
@@ -2239,9 +2242,9 @@ static USkeletalMeshComponent* ValidateBindingAsset(
 	}
 
 	// Validate against cards data
-	for (int32 GroupIt = 0, GroupCount = GroomAsset->HairGroupsPlatformData.Num(); GroupIt < GroupCount; ++GroupIt)
+	for (int32 GroupIt = 0, GroupCount = GroomAsset->GetHairGroupsPlatformData().Num(); GroupIt < GroupCount; ++GroupIt)
 	{
-		FHairGroupPlatformData& GroupData = GroomAsset->HairGroupsPlatformData[GroupIt];
+		FHairGroupPlatformData& GroupData = GroomAsset->GetHairGroupsPlatformData()[GroupIt];
 		uint32 CardsLODIndex = 0;
 		for (FHairGroupPlatformData::FCards::FLOD& LOD : GroupData.Cards.LODs)
 		{
@@ -2298,8 +2301,8 @@ static EGroomCacheType GetEffectiveGroomCacheType(const UGroomCache* InCache, co
 	{
 		if (InCache->GetType() == EGroomCacheType::Guides)
 		{
-			// Ensure simulation cache is enable only if the groom asset as interpolation data enabled (i.e., EnableSimulationCache=true)
-			return InGroom->EnableSimulationCache ? EGroomCacheType::Guides : EGroomCacheType::None;
+			// Ensure simulation cache is enable only if the groom asset as interpolation data enabled (i.e., GetEnableSimulationCache()=true)
+			return InGroom->GetEnableSimulationCache() ? EGroomCacheType::Guides : EGroomCacheType::None;
 		}
 		else
 		{
@@ -2341,14 +2344,14 @@ void UGroomComponent::InitResources(bool bIsBindingReloading)
 	TArray<bool> bHasNeedGlobalDeformation;
 	TArray<bool> bHasNeedSimulation;
 	TArray<bool> bHasNeedDeformation;
-	bHasNeedGlobalDeformation.Init(false, GroomAsset->HairGroupsPlatformData.Num());
-	bHasNeedSimulation.Init(false, GroomAsset->HairGroupsPlatformData.Num());
-	bHasNeedDeformation.Init(false, GroomAsset->HairGroupsPlatformData.Num());
+	bHasNeedGlobalDeformation.Init(false, GroomAsset->GetHairGroupsPlatformData().Num());
+	bHasNeedSimulation.Init(false, GroomAsset->GetHairGroupsPlatformData().Num());
+	bHasNeedDeformation.Init(false, GroomAsset->GetHairGroupsPlatformData().Num());
 	
 	bool bHasAnyNeedSimulation = false;
 	bool bHasAnyNeedDeformation = false;
 	bool bHasAnyNeedGlobalDeformation = false;
-	for (int32 GroupIt = 0, GroupCount = GroomAsset->HairGroupsPlatformData.Num(); GroupIt < GroupCount; ++GroupIt)
+	for (int32 GroupIt = 0, GroupCount = GroomAsset->GetHairGroupsPlatformData().Num(); GroupIt < GroupCount; ++GroupIt)
 	{
 		for (uint32 LODIt = 0, LODCount = GroomAsset->GetLODCount(); LODIt < LODCount; ++LODIt)
 		{
@@ -2409,9 +2412,9 @@ void UGroomComponent::InitResources(bool bIsBindingReloading)
 	}
 
 	// Grab deformed mesh if it exists 
-	if(GroomAsset && GroomAsset->RiggedSkeletalMesh)
+	if(GroomAsset && GroomAsset->GetRiggedSkeletalMesh())
 	{
-		const USkeleton* TargetSkeleton = GroomAsset->RiggedSkeletalMesh->GetSkeleton();
+		const USkeleton* TargetSkeleton = GroomAsset->GetRiggedSkeletalMesh()->GetSkeleton();
 		
 		DeformedMeshComponent = nullptr;
 		// Try to find the first component by walking the attachment hierarchy
@@ -2448,7 +2451,7 @@ void UGroomComponent::InitResources(bool bIsBindingReloading)
 					}
 				}
 
-				for (int32 GroupIt = 0, GroupCount = GroomAsset->HairGroupsPlatformData.Num(); GroupIt < GroupCount; ++GroupIt)
+				for (int32 GroupIt = 0, GroupCount = GroomAsset->GetHairGroupsPlatformData().Num(); GroupIt < GroupCount; ++GroupIt)
 				{
 					for (uint32 LODIt = 0, LODCount = GroomAsset->GetLODCount(); LODIt < LODCount; ++LODIt)
 					{
@@ -2481,7 +2484,7 @@ void UGroomComponent::InitResources(bool bIsBindingReloading)
 		if (bHasNeedBindingData)
 		{
 			LocalBindingAsset = BindingAsset;
-			if (LocalBindingAsset && LocalBindingAsset->GetHairGroupResources().Num() != GroomAsset->HairGroupsPlatformData.Num())
+			if (LocalBindingAsset && LocalBindingAsset->GetHairGroupResources().Num() != GroomAsset->GetHairGroupsPlatformData().Num())
 			{
 				LocalBindingAsset = nullptr;
 			}
@@ -2500,7 +2503,7 @@ void UGroomComponent::InitResources(bool bIsBindingReloading)
 	{
 		bHasNeedSkinningBinding = false;
 		bHasAnyNeedGlobalDeformation = false;
-		bHasNeedGlobalDeformation.Init(false, GroomAsset->HairGroupsPlatformData.Num());
+		bHasNeedGlobalDeformation.Init(false, GroomAsset->GetHairGroupsPlatformData().Num());
 	}
 
 	if (GroomCache)
@@ -2513,7 +2516,7 @@ void UGroomComponent::InitResources(bool bIsBindingReloading)
 		GroomCacheBuffers.Reset();
 	}
 
-	for (int32 GroupIt = 0, GroupCount = GroomAsset->HairGroupsPlatformData.Num(); GroupIt < GroupCount; ++GroupIt)
+	for (int32 GroupIt = 0, GroupCount = GroomAsset->GetHairGroupsPlatformData().Num(); GroupIt < GroupCount; ++GroupIt)
 	{
 		FHairGroupInstance* HairGroupInstance = new FHairGroupInstance();
 		HairGroupInstance->AddRef();
@@ -2530,7 +2533,7 @@ void UGroomComponent::InitResources(bool bIsBindingReloading)
 		HairGroupInstance->Debug.LODPredictedIndex = LODPredictedIndex;
 		HairGroupInstance->Debug.LODSelectionTypeForDebug = LODSelectionType;
 		HairGroupInstance->DeformedComponent = DeformedMeshComponent;
-		HairGroupInstance->DeformedSection = GroomAsset->DeformedGroupSections.IsValidIndex(GroupIt) ? GroomAsset->DeformedGroupSections[GroupIt] : INDEX_NONE; 
+		HairGroupInstance->DeformedSection = GroomAsset->GetDeformedGroupSections().IsValidIndex(GroupIt) ? GroomAsset->GetDeformedGroupSections()[GroupIt] : INDEX_NONE; 
 		HairGroupInstance->Debug.MeshComponentId = ValidatedMeshComponent ? ValidatedMeshComponent->ComponentId : FPrimitiveComponentId();
 		HairGroupInstance->Debug.CachedMeshPersistentPrimitiveIndex = FPersistentPrimitiveIndex();
 
@@ -2547,9 +2550,9 @@ void UGroomComponent::InitResources(bool bIsBindingReloading)
 		HairGroupInstance->Debug.RigidPreviousLocalToWorld		= HairGroupInstance->Debug.RigidCurrentLocalToWorld;
 		HairGroupInstance->LocalToWorld							= HairGroupInstance->GetCurrentLocalToWorld();
 
-		FHairGroupPlatformData& GroupData = GroomAsset->HairGroupsPlatformData[GroupIt];
+		FHairGroupPlatformData& GroupData = GroomAsset->GetHairGroupsPlatformData()[GroupIt];
 
-		const EHairInterpolationType HairInterpolationType = ToHairInterpolationType(GroomAsset->HairInterpolationType);
+		const EHairInterpolationType HairInterpolationType = ToHairInterpolationType(GroomAsset->GetHairInterpolationType());
 
 		const FName OwnerName = GroomAsset->GetAssetPathName();
 
@@ -2564,7 +2567,7 @@ void UGroomComponent::InitResources(bool bIsBindingReloading)
 			TArray<float> CPULODScreenSize;
 			TArray<bool> LODVisibility;
 			TArray<EHairGeometryType> LODGeometryTypes;
-			const FHairGroupsLOD& GroupLOD = GroomAsset->HairGroupsLOD[GroupIt];
+			const FHairGroupsLOD& GroupLOD = GroomAsset->GetHairGroupsLOD()[GroupIt];
 			for (uint32 LODIt = 0, LODCount = GroupLOD.LODs.Num(); LODIt < LODCount; ++LODIt)
 			{
 				const FHairLODSettings& LODSettings = GroupLOD.LODs[LODIt];
@@ -2667,7 +2670,7 @@ void UGroomComponent::InitResources(bool bIsBindingReloading)
 			bool bNeedRootResources = false;
 			bool bNeedDynamicResources = false;
 			{
-				const FHairGroupsLOD& GroupLOD = GroomAsset->HairGroupsLOD[GroupIt];
+				const FHairGroupsLOD& GroupLOD = GroomAsset->GetHairGroupsLOD()[GroupIt];
 				for (uint32 LODIt = 0, LODCount = GroupLOD.LODs.Num(); LODIt < LODCount; ++LODIt)
 				{
 					if (HairGroupInstance->HairGroupPublicData->GetGeometryType(LODIt) == EHairGeometryType::Strands)
@@ -2773,7 +2776,7 @@ void UGroomComponent::InitResources(bool bIsBindingReloading)
 				{
 					// ClusterCullingData seriazlizes only the screen size related to strands geometry.
 					// Other type of geometry are not serizalized, and so won't match
-					if (GroomAsset->HairGroupsLOD[GroupIt].LODs[LODIt].GeometryType == EGroomGeometryType::Strands)
+					if (GroomAsset->GetHairGroupsLOD()[GroupIt].LODs[LODIt].GeometryType == EGroomGeometryType::Strands)
 					{
 						check(GroupData.Strands.ClusterCullingResource->BulkData.Header.CPULODScreenSize[LODIt] == LODScreenSizes[LODIt]);
 						check(GroupData.Strands.ClusterCullingResource->BulkData.Header.LODVisibility[LODIt] == LODVisibilities[LODIt]);
@@ -3084,8 +3087,8 @@ void UGroomComponent::PostLoad()
 	// ValidateMaterials(false);
 #endif
 
-	// Warning in case the groom is using a guide cache, but the EnableSimulationCache is disabled.
-	if (GroomAsset && GroomCache && GroomCache->GetType() == EGroomCacheType::Guides && !GroomAsset->EnableSimulationCache)
+	// Warning in case the groom is using a guide cache, but the GetEnableSimulationCache() is disabled.
+	if (GroomAsset && GroomCache && GroomCache->GetType() == EGroomCacheType::Guides && !GroomAsset->GetEnableSimulationCache())
 	{
 		UE_LOG(LogHairStrands, Warning, TEXT("[Groom] %s - The groom instance tries to use a guide-cache, but 'Enable Guide-Cache Support' option is not enabled on the groom asset. Because of this, the Guide-Cache won't be used."), *GetPathName());
 	}
@@ -3754,10 +3757,10 @@ void UGroomComponent::PostEditChangeProperty(FPropertyChangedEvent& PropertyChan
 	{
 		for (int32 Index = 0; Index < GroomAsset->GetNumHairGroups(); ++Index)
 		{
-			if (!GroomAsset->EnableSimulationCache)
+			if (!GroomAsset->GetEnableSimulationCache())
 			{
 				GroomAsset->Modify();
-				GroomAsset->EnableSimulationCache = true;
+				GroomAsset->SetEnableSimulationCache(true);
 				GroomAsset->CacheDerivedDatas();
 				break;
 			}
diff --git a/Engine/Plugins/Runtime/HairStrands/Source/HairStrandsCore/Private/GroomDeformerBuilder.cpp b/Engine/Plugins/Runtime/HairStrands/Source/HairStrandsCore/Private/GroomDeformerBuilder.cpp
--- a/Engine/Plugins/Runtime/HairStrands/Source/HairStrandsCore/Private/GroomDeformerBuilder.cpp
+++ b/Engine/Plugins/Runtime/HairStrands/Source/HairStrandsCore/Private/GroomDeformerBuilder.cpp
@@ -89,7 +89,7 @@ static FVector AddRootBone(UGroomAsset* GroomAsset, USkeletalMesh* SkeletalMesh,
 	{
 		FHairStrandsDatas StrandsData;
 		FHairStrandsDatas& GuidesData = GuidesDatas[GroupIndex];
-		FGroomBuilder::BuildData(HairGroup, GroomAsset->HairGroupsInterpolation[GroupIndex], GroomAsset->HairGroupsInfo[GroupIndex], StrandsData, GuidesData);
+		FGroomBuilder::BuildData(HairGroup, GroomAsset->GetHairGroupsInterpolation()[GroupIndex], GroomAsset->GetHairGroupsInfo()[GroupIndex], StrandsData, GuidesData);
 		
 		if(GroomAsset->IsDeformationEnable(GroupIndex))
 		{
@@ -361,7 +361,7 @@ static USkeletalMesh* BuildSkeletalMesh(UGroomAsset* GroomAsset)
 	
 	// Create package for the skeletal mesh
 	UPackage* PackageSkeletalMesh = CreatePackage(*PackageNameSkeletalMesh);
-	USkeletalMesh* SkeletalMesh = GroomAsset->RiggedSkeletalMesh ? GroomAsset->RiggedSkeletalMesh.Get() :
+	USkeletalMesh* SkeletalMesh = GroomAsset->GetRiggedSkeletalMesh() ? GroomAsset->GetRiggedSkeletalMesh() :
 		NewObject<USkeletalMesh>(PackageSkeletalMesh, *NameSkeletalMesh, RF_Public | RF_Standalone | RF_Transactional);
 
 	// Reset the skeletal mesh
@@ -379,7 +379,7 @@ static USkeletalMesh* BuildSkeletalMesh(UGroomAsset* GroomAsset)
 	BuildSkeletonBones(GroomAsset, SkeletalMesh, LodRenderData, BonesPositions, BonesOffsets, GroupsOffsets);
 
 	// Build the mesh section
-	BuildMeshSection(GroomAsset, SkeletalMesh, LodRenderData, BonesPositions, BonesOffsets, GroupsOffsets, GroomAsset->DeformedGroupSections);
+	BuildMeshSection(GroomAsset, SkeletalMesh, LodRenderData, BonesPositions, BonesOffsets, GroupsOffsets, GroomAsset->GetDeformedGroupSections());
 
 	// Build the physics asset
 	BuildPhysicsAsset(GroomAsset, SkeletalMesh, LodRenderData);
diff --git a/Engine/Plugins/Runtime/HairStrands/Source/HairStrandsCore/Private/GroomEdit.cpp b/Engine/Plugins/Runtime/HairStrands/Source/HairStrandsCore/Private/GroomEdit.cpp
--- a/Engine/Plugins/Runtime/HairStrands/Source/HairStrandsCore/Private/GroomEdit.cpp
+++ b/Engine/Plugins/Runtime/HairStrands/Source/HairStrandsCore/Private/GroomEdit.cpp
@@ -166,7 +166,7 @@ void ConvertFromGroomAsset(UGroomAsset* InGroom, FEditableGroom* OutGroom)
 		FHairGroupInfo GroupInfo;
 		FHairStrandsDatas StrandsData;
 		FHairStrandsDatas GuidesData;
-		FGroomBuilder::BuildData(HairDescriptionGroups.HairGroups[GroupIndex], In.HairGroupsInterpolation[GroupIndex], GroupInfo, StrandsData, GuidesData);
+		FGroomBuilder::BuildData(HairDescriptionGroups.HairGroups[GroupIndex], In.GetHairGroupsInterpolation()[GroupIndex], GroupInfo, StrandsData, GuidesData);
 
 		Out.Groups[GroupIndex].GroupID = GroupInfo.GroupID;
 		Out.Groups[GroupIndex].GroupName = GroupInfo.GroupName;
diff --git a/Engine/Plugins/Runtime/HairStrands/Source/HairStrandsCore/Private/GroomRBFDeformer.cpp b/Engine/Plugins/Runtime/HairStrands/Source/HairStrandsCore/Private/GroomRBFDeformer.cpp
--- a/Engine/Plugins/Runtime/HairStrands/Source/HairStrandsCore/Private/GroomRBFDeformer.cpp
+++ b/Engine/Plugins/Runtime/HairStrands/Source/HairStrandsCore/Private/GroomRBFDeformer.cpp
@@ -424,7 +424,7 @@ static void ApplyDeformationToGroom(const TArray<FRBFDeformedPositions>& Deforme
 		FHairDescriptionGroups HairDescriptionGroups;
 		FGroomBuilder::BuildHairDescriptionGroups(HairDescription, HairDescriptionGroups);
 
-		const int32 GroupCount = GroomAsset->HairGroupsInterpolation.Num();
+		const int32 GroupCount = GroomAsset->GetHairGroupsInterpolation().Num();
 		check(HairDescriptionGroups.HairGroups.Num() == GroupCount);
 		for (int32 GroupIndex = 0; GroupIndex < GroupCount; ++GroupIndex)
 		{
@@ -434,15 +434,15 @@ static void ApplyDeformationToGroom(const TArray<FRBFDeformedPositions>& Deforme
 				const FHairDescriptionGroup& HairGroup = HairDescriptionGroups.HairGroups[GroupIndex];
 				check(GroupIndex <= HairDescriptionGroups.HairGroups.Num());
 				check(GroupIndex == HairGroup.Info.GroupID);
-				const FHairInterpolationSettings& InterpolationSettings = GroomAsset->HairGroupsInterpolation[GroupIndex].InterpolationSettings;
-				const FHairGroupsLOD& HairGroupLOD = GroomAsset->HairGroupsLOD[GroupIndex];
+				const FHairInterpolationSettings& InterpolationSettings = GroomAsset->GetHairGroupsInterpolation()[GroupIndex].InterpolationSettings;
+				const FHairGroupsLOD& HairGroupLOD = GroomAsset->GetHairGroupsLOD()[GroupIndex];
 
-				FHairGroupInfo& HairGroupsInfo = GroomAsset->HairGroupsInfo[GroupIndex];
-				FHairGroupPlatformData& HairGroupsData = GroomAsset->HairGroupsPlatformData[GroupIndex];
+				FHairGroupInfo& HairGroupsInfo = GroomAsset->GetHairGroupsInfo()[GroupIndex];
+				FHairGroupPlatformData& HairGroupsData = GroomAsset->GetHairGroupsPlatformData()[GroupIndex];
 
 				FHairStrandsDatas StrandsData;
 				FHairStrandsDatas GuidesData;
-				FGroomBuilder::BuildData(HairGroup, GroomAsset->HairGroupsInterpolation[GroupIndex], HairGroupsInfo, StrandsData, GuidesData);
+				FGroomBuilder::BuildData(HairGroup, GroomAsset->GetHairGroupsInterpolation()[GroupIndex], HairGroupsInfo, StrandsData, GuidesData);
 				//GroomAsset->GetHairStrandsDatas(GroupIndex, StrandsData, GuidesData);
 
 				FGroomBuilder::BuildBulkData(HairGroup.Info, GuidesData,  HairGroupsData.Guides.BulkData);
@@ -607,8 +607,8 @@ void FGroomRBFDeformer::GetRBFDeformedGroomAsset(const UGroomAsset* InGroomAsset
 		// Sanity check to insure that the groom has all the original vertices
 		for (int32 GroupIt = 0; GroupIt < NumGroups; ++GroupIt)
 		{
-			check(InGroomAsset->HairGroupsInterpolation[GroupIt].DecimationSettings.VertexDecimation == 1);
-			check(InGroomAsset->HairGroupsInterpolation[GroupIt].DecimationSettings.CurveDecimation == 1);
+			check(InGroomAsset->GetHairGroupsInterpolation()[GroupIt].DecimationSettings.VertexDecimation == 1);
+			check(InGroomAsset->GetHairGroupsInterpolation()[GroupIt].DecimationSettings.CurveDecimation == 1);
 		}
 
 		// Build Curve index for every vertices
@@ -635,14 +635,14 @@ void FGroomRBFDeformer::GetRBFDeformedGroomAsset(const UGroomAsset* InGroomAsset
 		{
 			const uint32 GroupIndex = Group.Info.GroupID;
 			const FHairStrandsDatas& OriginalGuides = Group.Guides;
-			const FHairGroupPlatformData& HairGroupData = InGroomAsset->HairGroupsPlatformData[GroupIndex];
+			const FHairGroupPlatformData& HairGroupData = InGroomAsset->GetHairGroupsPlatformData()[GroupIndex];
 
 			FHairStrandsDatas StrandsData;
 			FHairStrandsDatas GuidesData;
 			FHairGroupInfo DummyInfo;
 
 			// Disable explicitely curve reordering as we want to preserve vertex mapping to write deformed postions back to the hair description, after RBF deformation
-			FGroomBuilder::BuildData(Group, InGroomAsset->HairGroupsInterpolation[GroupIndex], DummyInfo, StrandsData, GuidesData, false /*bAllowCurveReordering*/);
+			FGroomBuilder::BuildData(Group, InGroomAsset->GetHairGroupsInterpolation()[GroupIndex], DummyInfo, StrandsData, GuidesData, false /*bAllowCurveReordering*/);
 
 			TArray<uint32> SimRootDataPointToCurveBuffer;
 			TArray<uint32> RenRootDataPointToCurveBuffer;
@@ -651,7 +651,7 @@ void FGroomRBFDeformer::GetRBFDeformedGroomAsset(const UGroomAsset* InGroomAsset
 
 			// Get deformed guides
 			// If the groom override the value, we output dummy value for the guides, since they won't be used
-			if (InGroomAsset->HairGroupsInterpolation[GroupIndex].InterpolationSettings.GuideType != EGroomGuideType::Imported)
+			if (InGroomAsset->GetHairGroupsInterpolation()[GroupIndex].InterpolationSettings.GuideType != EGroomGuideType::Imported)
 			{
 				const uint32 OriginalVertexCount = OriginalGuides.StrandsPoints.Num();
 				DeformedPositions[GroupIndex].GuideStrands.Init(FVector3f::ZeroVector, OriginalVertexCount);
@@ -771,7 +771,7 @@ void FGroomRBFDeformer::GetRBFDeformedGroomAsset(const UGroomAsset* InGroomAsset
 		}
 
 		// Apply changes onto cards and meshes (OutGroomASset already contain duplicated mesh asset
-		for (FHairGroupsCardsSourceDescription& Desc : OutGroomAsset->HairGroupsCards)
+		for (FHairGroupsCardsSourceDescription& Desc : OutGroomAsset->GetHairGroupsCards())
 		{
 			UStaticMesh* Mesh = nullptr;
 			if (Desc.SourceType == EHairCardsSourceType::Procedural)
@@ -799,7 +799,7 @@ void FGroomRBFDeformer::GetRBFDeformedGroomAsset(const UGroomAsset* InGroomAsset
 		} 
 
 		// Apply RBF deformation to mesh vertices
-		for (FHairGroupsMeshesSourceDescription& Desc : OutGroomAsset->HairGroupsMeshes)
+		for (FHairGroupsMeshesSourceDescription& Desc : OutGroomAsset->GetHairGroupsMeshes())
 		{
 			if (UStaticMesh* Mesh = Desc.ImportedMesh)
 			{
diff --git a/Engine/Plugins/Runtime/HairStrands/Source/HairStrandsCore/Private/GroomTextureBuilder.cpp b/Engine/Plugins/Runtime/HairStrands/Source/HairStrandsCore/Private/GroomTextureBuilder.cpp
--- a/Engine/Plugins/Runtime/HairStrands/Source/HairStrandsCore/Private/GroomTextureBuilder.cpp
+++ b/Engine/Plugins/Runtime/HairStrands/Source/HairStrandsCore/Private/GroomTextureBuilder.cpp
@@ -271,7 +271,7 @@ static void InternalBuildFollicleTexture_CPU(const TArray<FFollicleInfo>& InInfo
 		}
 
 		uint32 GroupIndex = 0;
-		for (const FHairGroupPlatformData& HairGroupData : Info.GroomAsset->HairGroupsPlatformData)
+		for (const FHairGroupPlatformData& HairGroupData : Info.GroomAsset->GetHairGroupsPlatformData())
 		{
 			FHairStrandsDatas StrandsData;
 			FHairStrandsDatas GuidesData;
@@ -982,8 +982,8 @@ static bool TraceTextures(
 			continue;
 		}
 
-		FHairGroupPlatformData& GroupData = InInfo.GroomAsset->HairGroupsPlatformData[GroupIndex];
-		FHairGroupsRendering& RenderingData = InInfo.GroomAsset->HairGroupsRendering[GroupIndex];
+		FHairGroupPlatformData& GroupData = InInfo.GroomAsset->GetHairGroupsPlatformData()[GroupIndex];
+		FHairGroupsRendering& RenderingData = InInfo.GroomAsset->GetHairGroupsRendering()[GroupIndex];
 
 		FRDGBufferRef VoxelOffsetAndCount = GraphBuilder.RegisterExternalBuffer(GroupData.Debug.Resource->VoxelOffsetAndCount);
 		FRDGBufferRef VoxelData = GraphBuilder.RegisterExternalBuffer(GroupData.Debug.Resource->VoxelData);
diff --git a/Engine/Plugins/Runtime/HairStrands/Source/HairStrandsCore/Private/HairStrandsCore.cpp b/Engine/Plugins/Runtime/HairStrands/Source/HairStrandsCore/Private/HairStrandsCore.cpp
--- a/Engine/Plugins/Runtime/HairStrands/Source/HairStrandsCore/Private/HairStrandsCore.cpp
+++ b/Engine/Plugins/Runtime/HairStrands/Source/HairStrandsCore/Private/HairStrandsCore.cpp
@@ -231,7 +231,7 @@ UGroomBindingAsset* FHairStrandsCore::CreateGroomBindingAsset(EGroomBindingMeshT
 			Out->SetSourceGeometryCache(Cast<UGeometryCache>(Source));
 			Out->SetTargetGeometryCache(Cast<UGeometryCache>(Target));
 		}
-		Out->GetHairGroupsPlatformData().Reserve(GroomAsset->HairGroupsPlatformData.Num());
+		Out->GetHairGroupsPlatformData().Reserve(GroomAsset->GetHairGroupsPlatformData().Num());
 		Out->SetNumInterpolationPoints(NumInterpolationPoints);
 		Out->SetMatchingSection(MatchingSection);
 		Out->MarkPackageDirty();
diff --git a/Engine/Plugins/Runtime/HairStrands/Source/HairStrandsCore/Private/Niagara/NiagaraDataInterfaceHairStrands.cpp b/Engine/Plugins/Runtime/HairStrands/Source/HairStrandsCore/Private/Niagara/NiagaraDataInterfaceHairStrands.cpp
--- a/Engine/Plugins/Runtime/HairStrands/Source/HairStrandsCore/Private/Niagara/NiagaraDataInterfaceHairStrands.cpp
+++ b/Engine/Plugins/Runtime/HairStrands/Source/HairStrandsCore/Private/Niagara/NiagaraDataInterfaceHairStrands.cpp
@@ -382,7 +382,7 @@ void FNDIHairStrandsData::Update(UNiagaraDataInterfaceHairStrands* Interface, FN
 
 		const bool bHasValidBindingAsset = (Interface->IsComponentValid() && Interface->SourceComponent->BindingAsset && Interface->SourceComponent->GroomAsset);
 
-		GlobalInterpolation = bHasValidBindingAsset ? Interface->SourceComponent->GroomAsset->EnableGlobalInterpolation : false;
+		GlobalInterpolation = bHasValidBindingAsset ? Interface->SourceComponent->GroomAsset->GetEnableGlobalInterpolation() : false;
 		bSkinningTransfer = bHasValidBindingAsset ?
 			(Interface->SourceComponent->BindingAsset->GetSourceSkeletalMesh() && Interface->SourceComponent->BindingAsset->GetTargetSkeletalMesh() &&
 			 Interface->SourceComponent->BindingAsset->GetSourceSkeletalMesh() != Interface->SourceComponent->BindingAsset->GetTargetSkeletalMesh()) : false;
@@ -392,9 +392,9 @@ void FNDIHairStrandsData::Update(UNiagaraDataInterfaceHairStrands* Interface, FN
 		const bool bIsSimulationEnable = Interface->IsComponentValid() ? Interface->SourceComponent->IsSimulationEnable(GroupIndex, LODIndex) : 
 																		    GroomAsset ? GroomAsset->IsSimulationEnable(GroupIndex, LODIndex) : false;
 
-		if (StrandsDatas != nullptr && GroomAsset != nullptr && GroupIndex >= 0 && GroupIndex < GroomAsset->HairGroupsPhysics.Num() && bIsSimulationEnable)
+		if (StrandsDatas != nullptr && GroomAsset != nullptr && GroupIndex >= 0 && GroupIndex < GroomAsset->GetHairGroupsPhysics().Num() && bIsSimulationEnable)
 		{
-			FHairGroupsPhysics& HairPhysics = GroomAsset->HairGroupsPhysics[GroupIndex];
+			FHairGroupsPhysics& HairPhysics = GroomAsset->GetHairGroupsPhysics()[GroupIndex];
 			StrandsSize = static_cast<uint8>(HairPhysics.StrandsParameters.StrandsSize);
 
 			HairGroupInstance = Interface->IsComponentValid() ? Interface->SourceComponent->GetGroupInstance(GroupIndex) : nullptr;
@@ -733,7 +733,7 @@ void UNiagaraDataInterfaceHairStrands::ExtractDatasAndResources(
 		OutLocalToWorld = SystemInstance ? SystemInstance->GetWorldTransform() : FTransform::Identity;
 		if (OutGroupIndex < DefaultSource->GetNumHairGroups())
 		{
-			OutStrandsRestResource = DefaultSource->HairGroupsPlatformData[OutGroupIndex].Guides.RestResource;
+			OutStrandsRestResource = DefaultSource->GetHairGroupsPlatformData()[OutGroupIndex].Guides.RestResource;
 			OutGroomAsset = DefaultSource;
 		}
 	}
diff --git a/Engine/Plugins/Runtime/HairStrands/Source/HairStrandsCore/Public/GroomAsset.h b/Engine/Plugins/Runtime/HairStrands/Source/HairStrandsCore/Public/GroomAsset.h
--- a/Engine/Plugins/Runtime/HairStrands/Source/HairStrandsCore/Public/GroomAsset.h
+++ b/Engine/Plugins/Runtime/HairStrands/Source/HairStrandsCore/Public/GroomAsset.h
@@ -429,70 +429,168 @@ class HAIRSTRANDSCORE_API UGroomAsset : public UObject, public IInterface_AssetU
 	DECLARE_MULTICAST_DELEGATE(FOnGroomAsyncLoadFinished);
 #endif
 
-public:
-	
+private:
+	UE_DEPRECATED(5.3, "Please do not access this member directly; use UGroomAsset accessor.")
 	UPROPERTY(EditAnywhere, Transient, EditFixedSize, Category = "HairInfo", meta = (DisplayName = "Group"))
 	TArray<FHairGroupInfoWithVisibility> HairGroupsInfo;
 
-	UPROPERTY(EditAnywhere, EditFixedSize, BlueprintReadWrite, Category = "HairRendering", meta = (DisplayName = "Group"))
+	UE_DEPRECATED(5.3, "Please do not access this member directly; use UGroomAsset accessor.")
+	UPROPERTY(EditAnywhere, EditFixedSize, BlueprintGetter = GetHairGroupsRendering, BlueprintSetter = SetHairGroupsRendering, Category = "HairRendering", meta = (DisplayName = "Group"))
 	TArray<FHairGroupsRendering> HairGroupsRendering;
 
-	UPROPERTY(EditAnywhere, EditFixedSize, BlueprintReadWrite, Category = "HairPhysics", meta = (DisplayName = "Group"))
+	UE_DEPRECATED(5.3, "Please do not access this member directly; use UGroomAsset accessor.")
+	UPROPERTY(EditAnywhere, EditFixedSize, BlueprintGetter = GetHairGroupsPhysics, BlueprintSetter = SetHairGroupsPhysics, Category = "HairPhysics", meta = (DisplayName = "Group"))
 	TArray<FHairGroupsPhysics> HairGroupsPhysics;
 
-	UPROPERTY(EditAnywhere, EditFixedSize, BlueprintReadWrite, Category = "HairInterpolation", meta = (DisplayName = "Group"))
+	UE_DEPRECATED(5.3, "Please do not access this member directly; use UGroomAsset accessor.")
+	UPROPERTY(EditAnywhere, EditFixedSize, BlueprintGetter = GetHairGroupsInterpolation, BlueprintSetter = SetHairGroupsInterpolation, Category = "HairInterpolation", meta = (DisplayName = "Group"))
 	TArray<FHairGroupsInterpolation> HairGroupsInterpolation;
 
-	UPROPERTY(EditAnywhere, EditFixedSize, BlueprintReadWrite, Category = "HairLOD", meta = (DisplayName = "Group"))
+	UE_DEPRECATED(5.3, "Please do not access this member directly; use UGroomAsset accessor.")
+	UPROPERTY(EditAnywhere, EditFixedSize, BlueprintGetter = GetHairGroupsLOD, BlueprintSetter = SetHairGroupsLOD, Category = "HairLOD", meta = (DisplayName = "Group"))
 	TArray<FHairGroupsLOD> HairGroupsLOD;
 
 	/** Cards - Source description data */
-	UPROPERTY(EditAnywhere, BlueprintReadWrite, Category = "HairCards", meta = (DisplayName = "Group"))
+	UE_DEPRECATED(5.3, "Please do not access this member directly; use UGroomAsset accessor.")
+	UPROPERTY(EditAnywhere, BlueprintGetter = GetHairGroupsCards, BlueprintSetter = SetHairGroupsCards, Category = "HairCards", meta = (DisplayName = "Group"))
 	TArray<FHairGroupsCardsSourceDescription> HairGroupsCards;
 
 	/** Meshes - Source description data */
-	UPROPERTY(EditAnywhere, BlueprintReadWrite, Category = "HairMeshes", meta = (DisplayName = "Group"))
+	UE_DEPRECATED(5.3, "Please do not access this member directly; use UGroomAsset accessor.")
+	UPROPERTY(EditAnywhere, BlueprintGetter = GetHairGroupsMeshes, BlueprintSetter = SetHairGroupsMeshes, Category = "HairMeshes", meta = (DisplayName = "Group"))
 	TArray<FHairGroupsMeshesSourceDescription> HairGroupsMeshes;
 
 	/** Meshes - Source description data */
-	UPROPERTY(EditAnywhere, BlueprintReadWrite, Category = "HairMaterials", meta = (DisplayName = "Group"))
+	UE_DEPRECATED(5.3, "Please do not access this member directly; use UGroomAsset accessor.")
+		UPROPERTY(EditAnywhere, BlueprintGetter = GetHairGroupsMaterials, BlueprintSetter = SetHairGroupsMaterials, Category = "HairMaterials", meta = (DisplayName = "Group"))
 	TArray<FHairGroupsMaterial> HairGroupsMaterials;
 
-	/** Store strands/cards/meshes data */
-	TArray<FHairGroupPlatformData> HairGroupsPlatformData;
-
 	/** Enable radial basis function interpolation to be used instead of the local skin rigid transform */
-	UPROPERTY(EditAnywhere, EditFixedSize, BlueprintReadWrite, Category = "HairInterpolation", meta = (ToolTip = "Enable radial basis function interpolation to be used instead of the local skin rigid transform (WIP)", DisplayName = "RBF Interpolation"))
+	UE_DEPRECATED(5.3, "Please do not access this member directly; use UGroomAsset accessor.")
+	UPROPERTY(EditAnywhere, EditFixedSize, BlueprintGetter = GetEnableGlobalInterpolation, BlueprintSetter = SetEnableGlobalInterpolation, Category = "HairInterpolation", meta = (ToolTip = "Enable radial basis function interpolation to be used instead of the local skin rigid transform (WIP)", DisplayName = "RBF Interpolation"))
 	bool EnableGlobalInterpolation = false;
 
 	/** Enable guide-cache support. This allows to attach a guide-cache dynamically at runtime */
-	UPROPERTY(EditAnywhere, EditFixedSize, BlueprintReadWrite, Category = "HairInterpolation", meta = (ToolTip = "Enable guide-cache support. This allows to attach a simulation-cache dynamically at runtime", DisplayName = "Enable Guide-Cache Support"))
+	UE_DEPRECATED(5.3, "Please do not access this member directly; use UGroomAsset accessor.")
+	UPROPERTY(EditAnywhere, EditFixedSize, BlueprintGetter = GetEnableSimulationCache, BlueprintSetter = SetEnableSimulationCache, Category = "HairInterpolation", meta = (ToolTip = "Enable guide-cache support. This allows to attach a simulation-cache dynamically at runtime", DisplayName = "Enable Guide-Cache Support"))
 	bool EnableSimulationCache = false;
 
 	/** Type of interpolation used */
-	UPROPERTY(EditAnywhere, EditFixedSize, BlueprintReadWrite, Category = "HairInterpolation", meta = (ToolTip = "Type of interpolation used (WIP)"))
+	UE_DEPRECATED(5.3, "Please do not access this member directly; use UGroomAsset accessor.")
+	UPROPERTY(EditAnywhere, EditFixedSize, BlueprintGetter = GetHairInterpolationType, BlueprintSetter = SetHairInterpolationType, Category = "HairInterpolation", meta = (ToolTip = "Type of interpolation used (WIP)"))
 	EGroomInterpolationType HairInterpolationType = EGroomInterpolationType::SmoothTransform;
 
 	/** Deformed skeletal mesh that will drive the groom deformation/simulation. For creating this skeletal mesh, enable EnableDeformation within the interpolation settings*/
-	UPROPERTY(VisibleAnywhere, BlueprintReadWrite, Category = "HairInterpolation")
+	UE_DEPRECATED(5.3, "Please do not access this member directly; use UGroomAsset accessor.")
+	UPROPERTY(VisibleAnywhere, BlueprintGetter = GetRiggedSkeletalMesh, BlueprintSetter = SetRiggedSkeletalMesh, Category = "HairInterpolation")
 	TObjectPtr<USkeletalMesh> RiggedSkeletalMesh;
 
 	/** Deformed skeletal mesh mapping from groups to sections */
+	UE_DEPRECATED(5.3, "Please do not access this member directly; use UGroomAsset accessor.")
 	UPROPERTY()
 	TArray<int32> DeformedGroupSections;
 
 	/** Minimum LOD to cook */
+	UE_DEPRECATED(5.3, "Please do not access this member directly; use UGroomAsset accessor.")
 	UPROPERTY(EditAnywhere, Category = "LOD", meta = (DisplayName = "Minimum LOD"))
 	FPerPlatformInt MinLOD;
 
 	/** When true all LODs below MinLod will still be cooked */
+	UE_DEPRECATED(5.3, "Please do not access this member directly; use UGroomAsset accessor.")
 	UPROPERTY(EditAnywhere, AdvancedDisplay, Category = "LOD")
 	FPerPlatformBool DisableBelowMinLodStripping;
 
 	/** The LOD bias to use after LOD stripping, regardless of MinLOD. Computed at cook time */
+	UE_DEPRECATED(5.3, "Please do not access this member directly; use UGroomAsset accessor.")
 	UPROPERTY()
 	TArray<float> EffectiveLODBias;
 
+	/** Store strands/cards/meshes data */
+	UE_DEPRECATED(5.3, "Please do not access this member directly; use UGroomAsset accessor.")
+	TArray<FHairGroupPlatformData> HairGroupsPlatformData;
+
+public:
+	static FName GetHairGroupsRenderingMemberName();
+	UFUNCTION(BlueprintGetter) TArray<FHairGroupsRendering>& GetHairGroupsRendering();
+	UFUNCTION(BlueprintSetter) void SetHairGroupsRendering(const TArray<FHairGroupsRendering>& In);
+	const TArray<FHairGroupsRendering>& GetHairGroupsRendering() const;
+
+	static FName GetHairGroupsPhysicsMemberName();
+	UFUNCTION(BlueprintGetter) TArray<FHairGroupsPhysics>& GetHairGroupsPhysics();
+	UFUNCTION(BlueprintSetter) void SetHairGroupsPhysics(const TArray<FHairGroupsPhysics>& In);
+	const TArray<FHairGroupsPhysics>& GetHairGroupsPhysics() const;
+
+	static FName GetHairGroupsInterpolationMemberName();
+	UFUNCTION(BlueprintGetter) TArray<FHairGroupsInterpolation>& GetHairGroupsInterpolation();
+	UFUNCTION(BlueprintSetter) void SetHairGroupsInterpolation(const TArray<FHairGroupsInterpolation>& In);
+	const TArray<FHairGroupsInterpolation>& GetHairGroupsInterpolation() const;
+
+	static FName GetHairGroupsLODMemberName();
+	UFUNCTION(BlueprintGetter) TArray<FHairGroupsLOD>& GetHairGroupsLOD();
+	UFUNCTION(BlueprintSetter) void SetHairGroupsLOD(const TArray<FHairGroupsLOD>& In);
+	const TArray<FHairGroupsLOD>& GetHairGroupsLOD() const;
+
+	static FName GetHairGroupsCardsMemberName();
+	UFUNCTION(BlueprintGetter) TArray<FHairGroupsCardsSourceDescription>& GetHairGroupsCards();
+	UFUNCTION(BlueprintSetter) void SetHairGroupsCards(const TArray<FHairGroupsCardsSourceDescription>& In);
+	const TArray<FHairGroupsCardsSourceDescription>& GetHairGroupsCards() const;
+
+	static FName GetHairGroupsMeshesMemberName();
+	UFUNCTION(BlueprintGetter) TArray<FHairGroupsMeshesSourceDescription>& GetHairGroupsMeshes();
+	UFUNCTION(BlueprintSetter) void SetHairGroupsMeshes(const TArray<FHairGroupsMeshesSourceDescription>& In);
+	const TArray<FHairGroupsMeshesSourceDescription>& GetHairGroupsMeshes() const;
+
+	static FName GetHairGroupsMaterialsMemberName();
+	UFUNCTION(BlueprintGetter) TArray<FHairGroupsMaterial>& GetHairGroupsMaterials();
+	UFUNCTION(BlueprintSetter) void SetHairGroupsMaterials(const TArray<FHairGroupsMaterial>& In);
+	const TArray<FHairGroupsMaterial>& GetHairGroupsMaterials() const;
+
+	static FName GetEnableGlobalInterpolationMemberName();
+	UFUNCTION(BlueprintGetter) bool GetEnableGlobalInterpolation() const;
+	UFUNCTION(BlueprintSetter) void SetEnableGlobalInterpolation(bool In);
+
+	static FName GetEnableSimulationCacheMemberName();
+	UFUNCTION(BlueprintGetter) bool GetEnableSimulationCache() const;
+	UFUNCTION(BlueprintSetter) void SetEnableSimulationCache(bool In);
+
+	static FName GetHairInterpolationTypeMemberName();
+	UFUNCTION(BlueprintGetter) EGroomInterpolationType GetHairInterpolationType() const;
+	UFUNCTION(BlueprintSetter) void SetHairInterpolationType(EGroomInterpolationType In);
+
+	static FName GetRiggedSkeletalMeshMemberName();
+	UFUNCTION(BlueprintGetter) USkeletalMesh* GetRiggedSkeletalMesh() const;
+	UFUNCTION(BlueprintSetter) void SetRiggedSkeletalMesh(USkeletalMesh* In);
+
+	static FName GetDeformedGroupSectionsMemberName();
+	UFUNCTION(BlueprintGetter) TArray<int32>& GetDeformedGroupSections();
+	UFUNCTION(BlueprintSetter) void SetDeformedGroupSections(const TArray<int32>& In);
+	const TArray<int32>& GetDeformedGroupSections() const;
+
+	static FName GetMinLODMemberName();
+	FPerPlatformInt GetMinLOD() const;
+	void SetMinLOD(FPerPlatformInt In);
+
+	static FName GetDisableBelowMinLodStrippingMemberName();
+	FPerPlatformBool GetDisableBelowMinLodStripping() const;
+	void SetDisableBelowMinLodStripping(FPerPlatformBool In);
+
+	static FName GetEffectiveLODBiasMemberName();
+	TArray<float>& GetEffectiveLODBias();
+	void SetEffectiveLODBias(const TArray<float>& In);
+	const TArray<float>& GetEffectiveLODBias() const;
+
+	static FName GetHairGroupsPlatformDataMemberName();
+	TArray<FHairGroupPlatformData>& GetHairGroupsPlatformData();
+	const TArray<FHairGroupPlatformData>& GetHairGroupsPlatformData() const;
+	void SetHairGroupsPlatformData(const TArray<FHairGroupPlatformData>& In);
+
+	static FName GetHairGroupsInfoMemberName();
+	TArray<FHairGroupInfoWithVisibility>& GetHairGroupsInfo();
+	const TArray<FHairGroupInfoWithVisibility>& GetHairGroupsInfo() const;
+	void SetHairGroupsInfo(const TArray<FHairGroupInfoWithVisibility>& In);
+
+public:
+
 	//~ Begin UObject Interface.
 	virtual void GetResourceSizeEx(FResourceSizeEx& CumulativeResourceSize) override;
 	virtual void PostLoad() override;
diff --git a/Engine/Plugins/Runtime/HairStrands/Source/HairStrandsEditor/Private/GroomActions.cpp b/Engine/Plugins/Runtime/HairStrands/Source/HairStrandsEditor/Private/GroomActions.cpp
--- a/Engine/Plugins/Runtime/HairStrands/Source/HairStrandsEditor/Private/GroomActions.cpp
+++ b/Engine/Plugins/Runtime/HairStrands/Source/HairStrandsEditor/Private/GroomActions.cpp
@@ -115,12 +115,12 @@ void ExecuteRebuild(const FToolMenuContext& InContext)
 					{
 						FGroomHairGroupPreview& OutGroup = GroupsPreview->Groups.AddDefaulted_GetRef();
 						OutGroup.GroupID    	= GroupIndex;
-						OutGroup.GroupName		= GroomAsset->HairGroupsInfo[GroupIndex].GroupName;
-						OutGroup.CurveCount 	= GroomAsset->HairGroupsPlatformData[GroupIndex].Strands.BulkData.GetNumCurves();
-						OutGroup.GuideCount 	= GroomAsset->HairGroupsPlatformData[GroupIndex].Guides.BulkData.GetNumCurves();
+						OutGroup.GroupName		= GroomAsset->GetHairGroupsInfo()[GroupIndex].GroupName;
+						OutGroup.CurveCount 	= GroomAsset->GetHairGroupsPlatformData()[GroupIndex].Strands.BulkData.GetNumCurves();
+						OutGroup.GuideCount 	= GroomAsset->GetHairGroupsPlatformData()[GroupIndex].Guides.BulkData.GetNumCurves();
 						OutGroup.Attributes 	= HairDescriptionGroups.HairGroups[GroupIndex].GetHairAttributes();
 						OutGroup.AttributeFlags = HairDescriptionGroups.HairGroups[GroupIndex].GetHairAttributeFlags();
-						OutGroup.InterpolationSettings = GroomAsset->HairGroupsInterpolation[GroupIndex];
+						OutGroup.InterpolationSettings = GroomAsset->GetHairGroupsInterpolation()[GroupIndex];
 					}
 				}
 				TSharedPtr<SGroomImportOptionsWindow> GroomOptionWindow = SGroomImportOptionsWindow::DisplayRebuildOptions(CurrentOptions, GroupsPreview, Filename);
@@ -134,8 +134,8 @@ void ExecuteRebuild(const FToolMenuContext& InContext)
 				bool bEnableRigging = false;
 				for (uint32 GroupIndex = 0; GroupIndex < GroupCount; ++GroupIndex)
 				{
-					GroomAsset->HairGroupsInterpolation[GroupIndex] = GroupsPreview->Groups[GroupIndex].InterpolationSettings;
-					bEnableRigging |= GroomAsset->HairGroupsInterpolation[GroupIndex].InterpolationSettings.GuideType == EGroomGuideType::Rigged;
+					GroomAsset->GetHairGroupsInterpolation()[GroupIndex] = GroupsPreview->Groups[GroupIndex].InterpolationSettings;
+					bEnableRigging |= GroomAsset->GetHairGroupsInterpolation()[GroupIndex].InterpolationSettings.GuideType == EGroomGuideType::Rigged;
 				}
 
 				bool bSucceeded = GroomAsset->CacheDerivedDatas();
@@ -143,7 +143,7 @@ void ExecuteRebuild(const FToolMenuContext& InContext)
 				{
 					if(bEnableRigging)
 					{
-						GroomAsset->RiggedSkeletalMesh = FGroomDeformerBuilder::CreateSkeletalMesh(GroomAsset);
+						GroomAsset->SetRiggedSkeletalMesh(FGroomDeformerBuilder::CreateSkeletalMesh(GroomAsset));
 					}
 					// Move the transient ImportOptions to the asset package and set it on the GroomAssetImportData for serialization
 					CurrentOptions->Rename(nullptr, GroomAssetImportData);
diff --git a/Engine/Plugins/Runtime/HairStrands/Source/HairStrandsEditor/Private/GroomAssetDetails.cpp b/Engine/Plugins/Runtime/HairStrands/Source/HairStrandsEditor/Private/GroomAssetDetails.cpp
--- a/Engine/Plugins/Runtime/HairStrands/Source/HairStrandsEditor/Private/GroomAssetDetails.cpp
+++ b/Engine/Plugins/Runtime/HairStrands/Source/HairStrandsEditor/Private/GroomAssetDetails.cpp
@@ -430,110 +430,110 @@ void FGroomRenderingDetails::CustomizeStrandsGroupProperties(IDetailLayoutBuilde
 	{
 	case EMaterialPanelType::Cards:
 	{
-		DetailLayout.HideProperty(DetailLayout.GetProperty(GET_MEMBER_NAME_CHECKED(UGroomAsset, HairGroupsInterpolation), UGroomAsset::StaticClass()));
-		DetailLayout.HideProperty(DetailLayout.GetProperty(GET_MEMBER_NAME_CHECKED(UGroomAsset, HairGroupsRendering), UGroomAsset::StaticClass()));
-		DetailLayout.HideProperty(DetailLayout.GetProperty(GET_MEMBER_NAME_CHECKED(UGroomAsset, HairGroupsPhysics), UGroomAsset::StaticClass()));
-//		DetailLayout.HideProperty(DetailLayout.GetProperty(GET_MEMBER_NAME_CHECKED(UGroomAsset, HairGroupsCards), UGroomAsset::StaticClass()));
-		DetailLayout.HideProperty(DetailLayout.GetProperty(GET_MEMBER_NAME_CHECKED(UGroomAsset, HairGroupsMeshes), UGroomAsset::StaticClass()));
-		DetailLayout.HideProperty(DetailLayout.GetProperty(GET_MEMBER_NAME_CHECKED(UGroomAsset, HairGroupsMaterials), UGroomAsset::StaticClass()));
-		DetailLayout.HideProperty(DetailLayout.GetProperty(GET_MEMBER_NAME_CHECKED(UGroomAsset, HairGroupsLOD), UGroomAsset::StaticClass()));
-		DetailLayout.HideProperty(DetailLayout.GetProperty(GET_MEMBER_NAME_CHECKED(UGroomAsset, HairGroupsInfo), UGroomAsset::StaticClass()));
-		DetailLayout.HideProperty(DetailLayout.GetProperty(GET_MEMBER_NAME_CHECKED(UGroomAsset, EnableGlobalInterpolation), UGroomAsset::StaticClass()));
-		DetailLayout.HideProperty(DetailLayout.GetProperty(GET_MEMBER_NAME_CHECKED(UGroomAsset, EnableSimulationCache), UGroomAsset::StaticClass()));
-		DetailLayout.HideProperty(DetailLayout.GetProperty(GET_MEMBER_NAME_CHECKED(UGroomAsset, HairInterpolationType), UGroomAsset::StaticClass()));
-		DetailLayout.HideProperty(DetailLayout.GetProperty(GET_MEMBER_NAME_CHECKED(UGroomAsset, MinLOD), UGroomAsset::StaticClass()));
-		DetailLayout.HideProperty(DetailLayout.GetProperty(GET_MEMBER_NAME_CHECKED(UGroomAsset, DisableBelowMinLodStripping), UGroomAsset::StaticClass()));
-		DetailLayout.HideProperty(DetailLayout.GetProperty(GET_MEMBER_NAME_CHECKED(UGroomAsset, RiggedSkeletalMesh), UGroomAsset::StaticClass()));
+		DetailLayout.HideProperty(DetailLayout.GetProperty(UGroomAsset::GetHairGroupsInterpolationMemberName(), UGroomAsset::StaticClass()));
+		DetailLayout.HideProperty(DetailLayout.GetProperty(UGroomAsset::GetHairGroupsRenderingMemberName(), UGroomAsset::StaticClass()));
+		DetailLayout.HideProperty(DetailLayout.GetProperty(UGroomAsset::GetHairGroupsPhysicsMemberName(), UGroomAsset::StaticClass()));
+//		DetailLayout.HideProperty(DetailLayout.GetProperty(UGroomAsset::GetHairGroupsCardsMemberName(), UGroomAsset::StaticClass()));
+		DetailLayout.HideProperty(DetailLayout.GetProperty(UGroomAsset::GetHairGroupsMeshesMemberName(), UGroomAsset::StaticClass()));
+		DetailLayout.HideProperty(DetailLayout.GetProperty(UGroomAsset::GetHairGroupsMaterialsMemberName(), UGroomAsset::StaticClass()));
+		DetailLayout.HideProperty(DetailLayout.GetProperty(UGroomAsset::GetHairGroupsLODMemberName(), UGroomAsset::StaticClass()));
+		DetailLayout.HideProperty(DetailLayout.GetProperty(UGroomAsset::GetHairGroupsInfoMemberName(), UGroomAsset::StaticClass()));
+		DetailLayout.HideProperty(DetailLayout.GetProperty(UGroomAsset::GetEnableGlobalInterpolationMemberName(), UGroomAsset::StaticClass()));
+		DetailLayout.HideProperty(DetailLayout.GetProperty(UGroomAsset::GetEnableSimulationCacheMemberName(), UGroomAsset::StaticClass()));
+		DetailLayout.HideProperty(DetailLayout.GetProperty(UGroomAsset::GetHairInterpolationTypeMemberName(), UGroomAsset::StaticClass()));
+		DetailLayout.HideProperty(DetailLayout.GetProperty(UGroomAsset::GetMinLODMemberName(), UGroomAsset::StaticClass()));
+		DetailLayout.HideProperty(DetailLayout.GetProperty(UGroomAsset::GetDisableBelowMinLodStrippingMemberName(), UGroomAsset::StaticClass()));
+		DetailLayout.HideProperty(DetailLayout.GetProperty(UGroomAsset::GetRiggedSkeletalMeshMemberName(), UGroomAsset::StaticClass()));
 	}
 	break;
 	case EMaterialPanelType::Meshes:
 	{
-		DetailLayout.HideProperty(DetailLayout.GetProperty(GET_MEMBER_NAME_CHECKED(UGroomAsset, HairGroupsInterpolation), UGroomAsset::StaticClass()));
-		DetailLayout.HideProperty(DetailLayout.GetProperty(GET_MEMBER_NAME_CHECKED(UGroomAsset, HairGroupsRendering), UGroomAsset::StaticClass()));
-		DetailLayout.HideProperty(DetailLayout.GetProperty(GET_MEMBER_NAME_CHECKED(UGroomAsset, HairGroupsPhysics), UGroomAsset::StaticClass()));
-		DetailLayout.HideProperty(DetailLayout.GetProperty(GET_MEMBER_NAME_CHECKED(UGroomAsset, HairGroupsCards), UGroomAsset::StaticClass()));
-//		DetailLayout.HideProperty(DetailLayout.GetProperty(GET_MEMBER_NAME_CHECKED(UGroomAsset, HairGroupsMeshes), UGroomAsset::StaticClass()));
-		DetailLayout.HideProperty(DetailLayout.GetProperty(GET_MEMBER_NAME_CHECKED(UGroomAsset, HairGroupsMaterials), UGroomAsset::StaticClass()));
-		DetailLayout.HideProperty(DetailLayout.GetProperty(GET_MEMBER_NAME_CHECKED(UGroomAsset, HairGroupsLOD), UGroomAsset::StaticClass()));
-		DetailLayout.HideProperty(DetailLayout.GetProperty(GET_MEMBER_NAME_CHECKED(UGroomAsset, HairGroupsInfo), UGroomAsset::StaticClass()));
-		DetailLayout.HideProperty(DetailLayout.GetProperty(GET_MEMBER_NAME_CHECKED(UGroomAsset, EnableGlobalInterpolation), UGroomAsset::StaticClass()));
-		DetailLayout.HideProperty(DetailLayout.GetProperty(GET_MEMBER_NAME_CHECKED(UGroomAsset, EnableSimulationCache), UGroomAsset::StaticClass()));
-		DetailLayout.HideProperty(DetailLayout.GetProperty(GET_MEMBER_NAME_CHECKED(UGroomAsset, HairInterpolationType), UGroomAsset::StaticClass()));
-		DetailLayout.HideProperty(DetailLayout.GetProperty(GET_MEMBER_NAME_CHECKED(UGroomAsset, MinLOD), UGroomAsset::StaticClass()));
-		DetailLayout.HideProperty(DetailLayout.GetProperty(GET_MEMBER_NAME_CHECKED(UGroomAsset, DisableBelowMinLodStripping), UGroomAsset::StaticClass()));
-		DetailLayout.HideProperty(DetailLayout.GetProperty(GET_MEMBER_NAME_CHECKED(UGroomAsset, RiggedSkeletalMesh), UGroomAsset::StaticClass()));
+		DetailLayout.HideProperty(DetailLayout.GetProperty(UGroomAsset::GetHairGroupsInterpolationMemberName(), UGroomAsset::StaticClass()));
+		DetailLayout.HideProperty(DetailLayout.GetProperty(UGroomAsset::GetHairGroupsRenderingMemberName(), UGroomAsset::StaticClass()));
+		DetailLayout.HideProperty(DetailLayout.GetProperty(UGroomAsset::GetHairGroupsPhysicsMemberName(), UGroomAsset::StaticClass()));
+		DetailLayout.HideProperty(DetailLayout.GetProperty(UGroomAsset::GetHairGroupsCardsMemberName(), UGroomAsset::StaticClass()));
+//		DetailLayout.HideProperty(DetailLayout.GetProperty(UGroomAsset::GetHairGroupsMeshesMemberName(), UGroomAsset::StaticClass()));
+		DetailLayout.HideProperty(DetailLayout.GetProperty(UGroomAsset::GetHairGroupsMaterialsMemberName(), UGroomAsset::StaticClass()));
+		DetailLayout.HideProperty(DetailLayout.GetProperty(UGroomAsset::GetHairGroupsLODMemberName(), UGroomAsset::StaticClass()));
+		DetailLayout.HideProperty(DetailLayout.GetProperty(UGroomAsset::GetHairGroupsInfoMemberName(), UGroomAsset::StaticClass()));
+		DetailLayout.HideProperty(DetailLayout.GetProperty(UGroomAsset::GetEnableGlobalInterpolationMemberName(), UGroomAsset::StaticClass()));
+		DetailLayout.HideProperty(DetailLayout.GetProperty(UGroomAsset::GetEnableSimulationCacheMemberName(), UGroomAsset::StaticClass()));
+		DetailLayout.HideProperty(DetailLayout.GetProperty(UGroomAsset::GetHairInterpolationTypeMemberName(), UGroomAsset::StaticClass()));
+		DetailLayout.HideProperty(DetailLayout.GetProperty(UGroomAsset::GetMinLODMemberName(), UGroomAsset::StaticClass()));
+		DetailLayout.HideProperty(DetailLayout.GetProperty(UGroomAsset::GetDisableBelowMinLodStrippingMemberName(), UGroomAsset::StaticClass()));
+		DetailLayout.HideProperty(DetailLayout.GetProperty(UGroomAsset::GetRiggedSkeletalMeshMemberName(), UGroomAsset::StaticClass()));
 	}
 	break;
 	case EMaterialPanelType::Strands:
 	{
-		DetailLayout.HideProperty(DetailLayout.GetProperty(GET_MEMBER_NAME_CHECKED(UGroomAsset, HairGroupsInterpolation), UGroomAsset::StaticClass()));
-//		DetailLayout.HideProperty(DetailLayout.GetProperty(GET_MEMBER_NAME_CHECKED(UGroomAsset, HairGroupsRendering), UGroomAsset::StaticClass()));
-		DetailLayout.HideProperty(DetailLayout.GetProperty(GET_MEMBER_NAME_CHECKED(UGroomAsset, HairGroupsPhysics), UGroomAsset::StaticClass()));
-		DetailLayout.HideProperty(DetailLayout.GetProperty(GET_MEMBER_NAME_CHECKED(UGroomAsset, HairGroupsCards), UGroomAsset::StaticClass()));
-		DetailLayout.HideProperty(DetailLayout.GetProperty(GET_MEMBER_NAME_CHECKED(UGroomAsset, HairGroupsMeshes), UGroomAsset::StaticClass()));
-		DetailLayout.HideProperty(DetailLayout.GetProperty(GET_MEMBER_NAME_CHECKED(UGroomAsset, HairGroupsMaterials), UGroomAsset::StaticClass()));
-		DetailLayout.HideProperty(DetailLayout.GetProperty(GET_MEMBER_NAME_CHECKED(UGroomAsset, HairGroupsLOD), UGroomAsset::StaticClass()));
-		DetailLayout.HideProperty(DetailLayout.GetProperty(GET_MEMBER_NAME_CHECKED(UGroomAsset, HairGroupsInfo), UGroomAsset::StaticClass()));
-		DetailLayout.HideProperty(DetailLayout.GetProperty(GET_MEMBER_NAME_CHECKED(UGroomAsset, EnableGlobalInterpolation), UGroomAsset::StaticClass()));
-		DetailLayout.HideProperty(DetailLayout.GetProperty(GET_MEMBER_NAME_CHECKED(UGroomAsset, EnableSimulationCache), UGroomAsset::StaticClass()));
-		DetailLayout.HideProperty(DetailLayout.GetProperty(GET_MEMBER_NAME_CHECKED(UGroomAsset, HairInterpolationType), UGroomAsset::StaticClass()));
-		DetailLayout.HideProperty(DetailLayout.GetProperty(GET_MEMBER_NAME_CHECKED(UGroomAsset, MinLOD), UGroomAsset::StaticClass()));
-		DetailLayout.HideProperty(DetailLayout.GetProperty(GET_MEMBER_NAME_CHECKED(UGroomAsset, DisableBelowMinLodStripping), UGroomAsset::StaticClass()));
-		DetailLayout.HideProperty(DetailLayout.GetProperty(GET_MEMBER_NAME_CHECKED(UGroomAsset, RiggedSkeletalMesh), UGroomAsset::StaticClass()));
+		DetailLayout.HideProperty(DetailLayout.GetProperty(UGroomAsset::GetHairGroupsInterpolationMemberName(), UGroomAsset::StaticClass()));
+//		DetailLayout.HideProperty(DetailLayout.GetProperty(UGroomAsset::GetHairGroupsRenderingMemberName(), UGroomAsset::StaticClass()));
+		DetailLayout.HideProperty(DetailLayout.GetProperty(UGroomAsset::GetHairGroupsPhysicsMemberName(), UGroomAsset::StaticClass()));
+		DetailLayout.HideProperty(DetailLayout.GetProperty(UGroomAsset::GetHairGroupsCardsMemberName(), UGroomAsset::StaticClass()));
+		DetailLayout.HideProperty(DetailLayout.GetProperty(UGroomAsset::GetHairGroupsMeshesMemberName(), UGroomAsset::StaticClass()));
+		DetailLayout.HideProperty(DetailLayout.GetProperty(UGroomAsset::GetHairGroupsMaterialsMemberName(), UGroomAsset::StaticClass()));
+		DetailLayout.HideProperty(DetailLayout.GetProperty(UGroomAsset::GetHairGroupsLODMemberName(), UGroomAsset::StaticClass()));
+		DetailLayout.HideProperty(DetailLayout.GetProperty(UGroomAsset::GetHairGroupsInfoMemberName(), UGroomAsset::StaticClass()));
+		DetailLayout.HideProperty(DetailLayout.GetProperty(UGroomAsset::GetEnableGlobalInterpolationMemberName(), UGroomAsset::StaticClass()));
+		DetailLayout.HideProperty(DetailLayout.GetProperty(UGroomAsset::GetEnableSimulationCacheMemberName(), UGroomAsset::StaticClass()));
+		DetailLayout.HideProperty(DetailLayout.GetProperty(UGroomAsset::GetHairInterpolationTypeMemberName(), UGroomAsset::StaticClass()));
+		DetailLayout.HideProperty(DetailLayout.GetProperty(UGroomAsset::GetMinLODMemberName(), UGroomAsset::StaticClass()));
+		DetailLayout.HideProperty(DetailLayout.GetProperty(UGroomAsset::GetDisableBelowMinLodStrippingMemberName(), UGroomAsset::StaticClass()));
+		DetailLayout.HideProperty(DetailLayout.GetProperty(UGroomAsset::GetRiggedSkeletalMeshMemberName(), UGroomAsset::StaticClass()));
 	}
 	break;
 	case EMaterialPanelType::Physics:
 	{
-		DetailLayout.HideProperty(DetailLayout.GetProperty(GET_MEMBER_NAME_CHECKED(UGroomAsset, HairGroupsInterpolation), UGroomAsset::StaticClass()));
-		DetailLayout.HideProperty(DetailLayout.GetProperty(GET_MEMBER_NAME_CHECKED(UGroomAsset, HairGroupsRendering), UGroomAsset::StaticClass()));
-//		DetailLayout.HideProperty(DetailLayout.GetProperty(GET_MEMBER_NAME_CHECKED(UGroomAsset, HairGroupsPhysics), UGroomAsset::StaticClass()));
-		DetailLayout.HideProperty(DetailLayout.GetProperty(GET_MEMBER_NAME_CHECKED(UGroomAsset, HairGroupsCards), UGroomAsset::StaticClass()));
-		DetailLayout.HideProperty(DetailLayout.GetProperty(GET_MEMBER_NAME_CHECKED(UGroomAsset, HairGroupsMeshes), UGroomAsset::StaticClass()));
-		DetailLayout.HideProperty(DetailLayout.GetProperty(GET_MEMBER_NAME_CHECKED(UGroomAsset, HairGroupsMaterials), UGroomAsset::StaticClass()));
-		DetailLayout.HideProperty(DetailLayout.GetProperty(GET_MEMBER_NAME_CHECKED(UGroomAsset, HairGroupsLOD), UGroomAsset::StaticClass()));
-		DetailLayout.HideProperty(DetailLayout.GetProperty(GET_MEMBER_NAME_CHECKED(UGroomAsset, HairGroupsInfo), UGroomAsset::StaticClass()));
-		DetailLayout.HideProperty(DetailLayout.GetProperty(GET_MEMBER_NAME_CHECKED(UGroomAsset, EnableGlobalInterpolation), UGroomAsset::StaticClass()));
-		DetailLayout.HideProperty(DetailLayout.GetProperty(GET_MEMBER_NAME_CHECKED(UGroomAsset, EnableSimulationCache), UGroomAsset::StaticClass()));
-		DetailLayout.HideProperty(DetailLayout.GetProperty(GET_MEMBER_NAME_CHECKED(UGroomAsset, HairInterpolationType), UGroomAsset::StaticClass()));
-		DetailLayout.HideProperty(DetailLayout.GetProperty(GET_MEMBER_NAME_CHECKED(UGroomAsset, MinLOD), UGroomAsset::StaticClass()));
-		DetailLayout.HideProperty(DetailLayout.GetProperty(GET_MEMBER_NAME_CHECKED(UGroomAsset, DisableBelowMinLodStripping), UGroomAsset::StaticClass()));
-		DetailLayout.HideProperty(DetailLayout.GetProperty(GET_MEMBER_NAME_CHECKED(UGroomAsset, RiggedSkeletalMesh), UGroomAsset::StaticClass()));
+		DetailLayout.HideProperty(DetailLayout.GetProperty(UGroomAsset::GetHairGroupsInterpolationMemberName(), UGroomAsset::StaticClass()));
+		DetailLayout.HideProperty(DetailLayout.GetProperty(UGroomAsset::GetHairGroupsRenderingMemberName(), UGroomAsset::StaticClass()));
+//		DetailLayout.HideProperty(DetailLayout.GetProperty(UGroomAsset::GetHairGroupsPhysicsMemberName(), UGroomAsset::StaticClass()));
+		DetailLayout.HideProperty(DetailLayout.GetProperty(UGroomAsset::GetHairGroupsCardsMemberName(), UGroomAsset::StaticClass()));
+		DetailLayout.HideProperty(DetailLayout.GetProperty(UGroomAsset::GetHairGroupsMeshesMemberName(), UGroomAsset::StaticClass()));
+		DetailLayout.HideProperty(DetailLayout.GetProperty(UGroomAsset::GetHairGroupsMaterialsMemberName(), UGroomAsset::StaticClass()));
+		DetailLayout.HideProperty(DetailLayout.GetProperty(UGroomAsset::GetHairGroupsLODMemberName(), UGroomAsset::StaticClass()));
+		DetailLayout.HideProperty(DetailLayout.GetProperty(UGroomAsset::GetHairGroupsInfoMemberName(), UGroomAsset::StaticClass()));
+		DetailLayout.HideProperty(DetailLayout.GetProperty(UGroomAsset::GetEnableGlobalInterpolationMemberName(), UGroomAsset::StaticClass()));
+		DetailLayout.HideProperty(DetailLayout.GetProperty(UGroomAsset::GetEnableSimulationCacheMemberName(), UGroomAsset::StaticClass()));
+		DetailLayout.HideProperty(DetailLayout.GetProperty(UGroomAsset::GetHairInterpolationTypeMemberName(), UGroomAsset::StaticClass()));
+		DetailLayout.HideProperty(DetailLayout.GetProperty(UGroomAsset::GetMinLODMemberName(), UGroomAsset::StaticClass()));
+		DetailLayout.HideProperty(DetailLayout.GetProperty(UGroomAsset::GetDisableBelowMinLodStrippingMemberName(), UGroomAsset::StaticClass()));
+		DetailLayout.HideProperty(DetailLayout.GetProperty(UGroomAsset::GetRiggedSkeletalMeshMemberName(), UGroomAsset::StaticClass()));
 	}
 	break;
 	case EMaterialPanelType::Interpolation:
 	{
-//		DetailLayout.HideProperty(DetailLayout.GetProperty(GET_MEMBER_NAME_CHECKED(UGroomAsset, HairGroupsInterpolation), UGroomAsset::StaticClass()));
-		DetailLayout.HideProperty(DetailLayout.GetProperty(GET_MEMBER_NAME_CHECKED(UGroomAsset, HairGroupsRendering), UGroomAsset::StaticClass()));
-		DetailLayout.HideProperty(DetailLayout.GetProperty(GET_MEMBER_NAME_CHECKED(UGroomAsset, HairGroupsPhysics), UGroomAsset::StaticClass()));
-		DetailLayout.HideProperty(DetailLayout.GetProperty(GET_MEMBER_NAME_CHECKED(UGroomAsset, HairGroupsCards), UGroomAsset::StaticClass()));
-		DetailLayout.HideProperty(DetailLayout.GetProperty(GET_MEMBER_NAME_CHECKED(UGroomAsset, HairGroupsMeshes), UGroomAsset::StaticClass()));
-		DetailLayout.HideProperty(DetailLayout.GetProperty(GET_MEMBER_NAME_CHECKED(UGroomAsset, HairGroupsMaterials), UGroomAsset::StaticClass()));
-		DetailLayout.HideProperty(DetailLayout.GetProperty(GET_MEMBER_NAME_CHECKED(UGroomAsset, HairGroupsLOD), UGroomAsset::StaticClass()));
-		DetailLayout.HideProperty(DetailLayout.GetProperty(GET_MEMBER_NAME_CHECKED(UGroomAsset, HairGroupsInfo), UGroomAsset::StaticClass()));
-//		DetailLayout.HideProperty(DetailLayout.GetProperty(GET_MEMBER_NAME_CHECKED(UGroomAsset, EnableGlobalInterpolation), UGroomAsset::StaticClass()));
-//		DetailLayout.HideProperty(DetailLayout.GetProperty(GET_MEMBER_NAME_CHECKED(UGroomAsset, EnableSimulationCache), UGroomAsset::StaticClass()));
-//		DetailLayout.HideProperty(DetailLayout.GetProperty(GET_MEMBER_NAME_CHECKED(UGroomAsset, HairInterpolationType), UGroomAsset::StaticClass()));
-		DetailLayout.HideProperty(DetailLayout.GetProperty(GET_MEMBER_NAME_CHECKED(UGroomAsset, MinLOD), UGroomAsset::StaticClass()));
-		DetailLayout.HideProperty(DetailLayout.GetProperty(GET_MEMBER_NAME_CHECKED(UGroomAsset, DisableBelowMinLodStripping), UGroomAsset::StaticClass()));
-		//DetailLayout.HideProperty(DetailLayout.GetProperty(GET_MEMBER_NAME_CHECKED(UGroomAsset, RiggedSkeletalMesh), UGroomAsset::StaticClass()));
+//		DetailLayout.HideProperty(DetailLayout.GetProperty(UGroomAsset::GetHairGroupsInterpolationMemberName(), UGroomAsset::StaticClass()));
+		DetailLayout.HideProperty(DetailLayout.GetProperty(UGroomAsset::GetHairGroupsRenderingMemberName(), UGroomAsset::StaticClass()));
+		DetailLayout.HideProperty(DetailLayout.GetProperty(UGroomAsset::GetHairGroupsPhysicsMemberName(), UGroomAsset::StaticClass()));
+		DetailLayout.HideProperty(DetailLayout.GetProperty(UGroomAsset::GetHairGroupsCardsMemberName(), UGroomAsset::StaticClass()));
+		DetailLayout.HideProperty(DetailLayout.GetProperty(UGroomAsset::GetHairGroupsMeshesMemberName(), UGroomAsset::StaticClass()));
+		DetailLayout.HideProperty(DetailLayout.GetProperty(UGroomAsset::GetHairGroupsMaterialsMemberName(), UGroomAsset::StaticClass()));
+		DetailLayout.HideProperty(DetailLayout.GetProperty(UGroomAsset::GetHairGroupsLODMemberName(), UGroomAsset::StaticClass()));
+		DetailLayout.HideProperty(DetailLayout.GetProperty(UGroomAsset::GetHairGroupsInfoMemberName(), UGroomAsset::StaticClass()));
+//		DetailLayout.HideProperty(DetailLayout.GetProperty(UGroomAsset::GetEnableGlobalInterpolationMemberName(), UGroomAsset::StaticClass()));
+//		DetailLayout.HideProperty(DetailLayout.GetProperty(UGroomAsset::GetEnableSimulationCacheMemberName(), UGroomAsset::StaticClass()));
+//		DetailLayout.HideProperty(DetailLayout.GetProperty(UGroomAsset::GetHairInterpolationTypeMemberName(), UGroomAsset::StaticClass()));
+		DetailLayout.HideProperty(DetailLayout.GetProperty(UGroomAsset::GetMinLODMemberName(), UGroomAsset::StaticClass()));
+		DetailLayout.HideProperty(DetailLayout.GetProperty(UGroomAsset::GetDisableBelowMinLodStrippingMemberName(), UGroomAsset::StaticClass()));
+//		DetailLayout.HideProperty(DetailLayout.GetProperty(UGroomAsset::GetRiggedSkeletalMeshMemberName(), UGroomAsset::StaticClass()));
 	}
 	break;
 	case EMaterialPanelType::LODs:
 	{
-		DetailLayout.HideProperty(DetailLayout.GetProperty(GET_MEMBER_NAME_CHECKED(UGroomAsset, HairGroupsInterpolation), UGroomAsset::StaticClass()));
-		DetailLayout.HideProperty(DetailLayout.GetProperty(GET_MEMBER_NAME_CHECKED(UGroomAsset, HairGroupsRendering), UGroomAsset::StaticClass()));
-		DetailLayout.HideProperty(DetailLayout.GetProperty(GET_MEMBER_NAME_CHECKED(UGroomAsset, HairGroupsPhysics), UGroomAsset::StaticClass()));
-		DetailLayout.HideProperty(DetailLayout.GetProperty(GET_MEMBER_NAME_CHECKED(UGroomAsset, HairGroupsCards), UGroomAsset::StaticClass()));
-		DetailLayout.HideProperty(DetailLayout.GetProperty(GET_MEMBER_NAME_CHECKED(UGroomAsset, HairGroupsMeshes), UGroomAsset::StaticClass()));
-		DetailLayout.HideProperty(DetailLayout.GetProperty(GET_MEMBER_NAME_CHECKED(UGroomAsset, HairGroupsMaterials), UGroomAsset::StaticClass()));
-//		DetailLayout.HideProperty(DetailLayout.GetProperty(GET_MEMBER_NAME_CHECKED(UGroomAsset, HairGroupsLOD), UGroomAsset::StaticClass()));
-		DetailLayout.HideProperty(DetailLayout.GetProperty(GET_MEMBER_NAME_CHECKED(UGroomAsset, HairGroupsInfo), UGroomAsset::StaticClass()));
-		DetailLayout.HideProperty(DetailLayout.GetProperty(GET_MEMBER_NAME_CHECKED(UGroomAsset, EnableGlobalInterpolation), UGroomAsset::StaticClass()));
-		DetailLayout.HideProperty(DetailLayout.GetProperty(GET_MEMBER_NAME_CHECKED(UGroomAsset, EnableSimulationCache), UGroomAsset::StaticClass()));
-		DetailLayout.HideProperty(DetailLayout.GetProperty(GET_MEMBER_NAME_CHECKED(UGroomAsset, HairInterpolationType), UGroomAsset::StaticClass()));
-//		DetailLayout.HideProperty(DetailLayout.GetProperty(GET_MEMBER_NAME_CHECKED(UGroomAsset, MinLOD), UGroomAsset::StaticClass()));
-//		DetailLayout.HideProperty(DetailLayout.GetProperty(GET_MEMBER_NAME_CHECKED(UGroomAsset, DisableBelowMinLodStripping), UGroomAsset::StaticClass()));
-		DetailLayout.HideProperty(DetailLayout.GetProperty(GET_MEMBER_NAME_CHECKED(UGroomAsset, RiggedSkeletalMesh), UGroomAsset::StaticClass()));
+		DetailLayout.HideProperty(DetailLayout.GetProperty(UGroomAsset::GetHairGroupsInterpolationMemberName(), UGroomAsset::StaticClass()));
+		DetailLayout.HideProperty(DetailLayout.GetProperty(UGroomAsset::GetHairGroupsRenderingMemberName(), UGroomAsset::StaticClass()));
+		DetailLayout.HideProperty(DetailLayout.GetProperty(UGroomAsset::GetHairGroupsPhysicsMemberName(), UGroomAsset::StaticClass()));
+		DetailLayout.HideProperty(DetailLayout.GetProperty(UGroomAsset::GetHairGroupsCardsMemberName(), UGroomAsset::StaticClass()));
+		DetailLayout.HideProperty(DetailLayout.GetProperty(UGroomAsset::GetHairGroupsMeshesMemberName(), UGroomAsset::StaticClass()));
+		DetailLayout.HideProperty(DetailLayout.GetProperty(UGroomAsset::GetHairGroupsMaterialsMemberName(), UGroomAsset::StaticClass()));
+//		DetailLayout.HideProperty(DetailLayout.GetProperty(UGroomAsset::GetHairGroupsLODMemberName(), UGroomAsset::StaticClass()));
+		DetailLayout.HideProperty(DetailLayout.GetProperty(UGroomAsset::GetHairGroupsInfoMemberName(), UGroomAsset::StaticClass()));
+		DetailLayout.HideProperty(DetailLayout.GetProperty(UGroomAsset::GetEnableGlobalInterpolationMemberName(), UGroomAsset::StaticClass()));
+		DetailLayout.HideProperty(DetailLayout.GetProperty(UGroomAsset::GetEnableSimulationCacheMemberName(), UGroomAsset::StaticClass()));
+		DetailLayout.HideProperty(DetailLayout.GetProperty(UGroomAsset::GetHairInterpolationTypeMemberName(), UGroomAsset::StaticClass()));
+//		DetailLayout.HideProperty(DetailLayout.GetProperty(UGroomAsset::GetMinLODMemberName(), UGroomAsset::StaticClass()));
+//		DetailLayout.HideProperty(DetailLayout.GetProperty(UGroomAsset::GetDisableBelowMinLodStrippingMemberName(), UGroomAsset::StaticClass()));
+		DetailLayout.HideProperty(DetailLayout.GetProperty(UGroomAsset::GetRiggedSkeletalMeshMemberName(), UGroomAsset::StaticClass()));
 	}
 	break;
 	}
@@ -542,7 +542,7 @@ void FGroomRenderingDetails::CustomizeStrandsGroupProperties(IDetailLayoutBuilde
 	{
 		case EMaterialPanelType::Cards:		
 		{
-			TSharedRef<IPropertyHandle> Property = DetailLayout.GetProperty(GET_MEMBER_NAME_CHECKED(UGroomAsset, HairGroupsCards), UGroomAsset::StaticClass());
+			TSharedRef<IPropertyHandle> Property = DetailLayout.GetProperty(UGroomAsset::GetHairGroupsCardsMemberName(), UGroomAsset::StaticClass());
 			AddNewGroupButton(FilesCategory, Property->GetProperty(), FText::FromString(TEXT("Add Card asset")));
 			if (Property->IsValidHandle())
 			{
@@ -555,7 +555,7 @@ void FGroomRenderingDetails::CustomizeStrandsGroupProperties(IDetailLayoutBuilde
 		break;
 		case EMaterialPanelType::Meshes:	
 		{
-			TSharedRef<IPropertyHandle> Property = DetailLayout.GetProperty(GET_MEMBER_NAME_CHECKED(UGroomAsset, HairGroupsMeshes), UGroomAsset::StaticClass());
+			TSharedRef<IPropertyHandle> Property = DetailLayout.GetProperty(UGroomAsset::GetHairGroupsMeshesMemberName(), UGroomAsset::StaticClass());
 			AddNewGroupButton(FilesCategory, Property->GetProperty(), FText::FromString(TEXT("Add Mesh asset")));
 			if (Property->IsValidHandle())
 			{
@@ -568,7 +568,7 @@ void FGroomRenderingDetails::CustomizeStrandsGroupProperties(IDetailLayoutBuilde
 		break;
 		case EMaterialPanelType::Strands:	
 		{
-			TSharedRef<IPropertyHandle> Property = DetailLayout.GetProperty(GET_MEMBER_NAME_CHECKED(UGroomAsset, HairGroupsRendering), UGroomAsset::StaticClass());
+			TSharedRef<IPropertyHandle> Property = DetailLayout.GetProperty(UGroomAsset::GetHairGroupsRenderingMemberName(), UGroomAsset::StaticClass());
 			if (Property->IsValidHandle())
 			{
 				TSharedRef<FDetailArrayBuilder> PropertyBuilder = MakeShareable(new FDetailArrayBuilder(Property, false, false, false));
@@ -580,7 +580,7 @@ void FGroomRenderingDetails::CustomizeStrandsGroupProperties(IDetailLayoutBuilde
 		break;
 		case EMaterialPanelType::Physics:
 		{
-			TSharedRef<IPropertyHandle> Property = DetailLayout.GetProperty(GET_MEMBER_NAME_CHECKED(UGroomAsset, HairGroupsPhysics), UGroomAsset::StaticClass());
+			TSharedRef<IPropertyHandle> Property = DetailLayout.GetProperty(UGroomAsset::GetHairGroupsPhysicsMemberName(), UGroomAsset::StaticClass());
 			if (Property->IsValidHandle())
 			{
 				TSharedRef<FDetailArrayBuilder> PropertyBuilder = MakeShareable(new FDetailArrayBuilder(Property, false, false, false));
@@ -592,7 +592,7 @@ void FGroomRenderingDetails::CustomizeStrandsGroupProperties(IDetailLayoutBuilde
 		break;
 		case EMaterialPanelType::Interpolation:
 		{
-			TSharedRef<IPropertyHandle> Property = DetailLayout.GetProperty(GET_MEMBER_NAME_CHECKED(UGroomAsset, HairGroupsInterpolation), UGroomAsset::StaticClass());
+			TSharedRef<IPropertyHandle> Property = DetailLayout.GetProperty(UGroomAsset::GetHairGroupsInterpolationMemberName(), UGroomAsset::StaticClass());
 			if (Property->IsValidHandle())
 			{
 				TSharedRef<FDetailArrayBuilder> PropertyBuilder = MakeShareable(new FDetailArrayBuilder(Property, false, false, false));
@@ -604,7 +604,7 @@ void FGroomRenderingDetails::CustomizeStrandsGroupProperties(IDetailLayoutBuilde
 		break;
 		case EMaterialPanelType::LODs:
 		{
-			TSharedRef<IPropertyHandle> Property = DetailLayout.GetProperty(GET_MEMBER_NAME_CHECKED(UGroomAsset, HairGroupsLOD), UGroomAsset::StaticClass());
+			TSharedRef<IPropertyHandle> Property = DetailLayout.GetProperty(UGroomAsset::GetHairGroupsLODMemberName(), UGroomAsset::StaticClass());
 			if (Property->IsValidHandle())
 			{
 				TSharedRef<FDetailArrayBuilder> PropertyBuilder = MakeShareable(new FDetailArrayBuilder(Property, false, false, false));
@@ -637,13 +637,13 @@ FReply FGroomRenderingDetails::OnAddGroup(FProperty* Property)
 	case EMaterialPanelType::Cards:
 	{
 		FScopedTransaction Transaction(FText::FromString(TEXT("AddCardsGroup")));
-		GroomAsset->HairGroupsCards.AddDefaulted();
+		GroomAsset->GetHairGroupsCards().AddDefaulted();
 
-		const int32 LODCount = GroomAsset->HairGroupsCards.Num();
+		const int32 LODCount = GroomAsset->GetHairGroupsCards().Num();
 		if (LODCount > 1)
 		{
-			const FHairGroupsCardsSourceDescription& Prev = GroomAsset->HairGroupsCards[LODCount - 2];
-			FHairGroupsCardsSourceDescription& Current = GroomAsset->HairGroupsCards[LODCount - 1];
+			const FHairGroupsCardsSourceDescription& Prev = GroomAsset->GetHairGroupsCards()[LODCount - 2];
+			FHairGroupsCardsSourceDescription& Current = GroomAsset->GetHairGroupsCards()[LODCount - 1];
 
 			Current.SourceType = Prev.SourceType;
 			Current.GroupIndex = Prev.GroupIndex;
@@ -655,7 +655,7 @@ FReply FGroomRenderingDetails::OnAddGroup(FProperty* Property)
 		}
 		else
 		{
-			FHairGroupsCardsSourceDescription& Current = GroomAsset->HairGroupsCards[LODCount - 1];
+			FHairGroupsCardsSourceDescription& Current = GroomAsset->GetHairGroupsCards()[LODCount - 1];
 			Current.ProceduralSettings.TextureSettings.AtlasMaxResolution  = GetHairCardsAtlasResolution(0,0);
 			Current.LODIndex = 0;
 		}
@@ -667,7 +667,7 @@ FReply FGroomRenderingDetails::OnAddGroup(FProperty* Property)
 	case EMaterialPanelType::Meshes:
 	{
 		FScopedTransaction Transaction(FText::FromString(TEXT("AddMeshesGroup")));
-		GroomAsset->HairGroupsMeshes.AddDefaulted();
+		GroomAsset->GetHairGroupsMeshes().AddDefaulted();
 
 		FPropertyChangedEvent PropertyChangedEvent(Property, EPropertyChangeType::ArrayAdd);
 		GroomAsset->PostEditChangeProperty(PropertyChangedEvent);
@@ -694,9 +694,9 @@ FName& FGroomRenderingDetails::GetMaterialSlotName(int32 GroupIndex)
 	check(GroomAsset);
 	switch (PanelType)
 	{
-		case EMaterialPanelType::Cards:		return GroomAsset->HairGroupsCards[GroupIndex].MaterialSlotName;
-		case EMaterialPanelType::Meshes:	return GroomAsset->HairGroupsMeshes[GroupIndex].MaterialSlotName;
-		case EMaterialPanelType::Strands:	return GroomAsset->HairGroupsRendering[GroupIndex].MaterialSlotName;
+		case EMaterialPanelType::Cards:		return GroomAsset->GetHairGroupsCards()[GroupIndex].MaterialSlotName;
+		case EMaterialPanelType::Meshes:	return GroomAsset->GetHairGroupsMeshes()[GroupIndex].MaterialSlotName;
+		case EMaterialPanelType::Strands:	return GroomAsset->GetHairGroupsRendering()[GroupIndex].MaterialSlotName;
 	}
 
 	static FName Default;
@@ -709,9 +709,9 @@ const FName& FGroomRenderingDetails::GetMaterialSlotName(int32 GroupIndex) const
 	check(GroomAsset);
 	switch (PanelType)
 	{
-	case EMaterialPanelType::Cards:		return GroupIndex < GroomAsset->HairGroupsCards.Num() ? GroomAsset->HairGroupsCards[GroupIndex].MaterialSlotName : Default;
-	case EMaterialPanelType::Meshes:	return GroupIndex < GroomAsset->HairGroupsMeshes.Num() ? GroomAsset->HairGroupsMeshes[GroupIndex].MaterialSlotName : Default;
-	case EMaterialPanelType::Strands:	return GroupIndex < GroomAsset->HairGroupsRendering.Num() ? GroomAsset->HairGroupsRendering[GroupIndex].MaterialSlotName : Default;
+		case EMaterialPanelType::Cards:		return GroupIndex < GroomAsset->GetHairGroupsCards().Num() ? GroomAsset->GetHairGroupsCards()[GroupIndex].MaterialSlotName : Default;
+		case EMaterialPanelType::Meshes:	return GroupIndex < GroomAsset->GetHairGroupsMeshes().Num() ? GroomAsset->GetHairGroupsMeshes()[GroupIndex].MaterialSlotName : Default;
+		case EMaterialPanelType::Strands:	return GroupIndex < GroomAsset->GetHairGroupsRendering().Num() ? GroomAsset->GetHairGroupsRendering()[GroupIndex].MaterialSlotName : Default;
 	}
 
 	return Default;
@@ -722,12 +722,12 @@ int32 FGroomRenderingDetails::GetGroupCount() const
 	check(GroomAsset);
 	switch (PanelType)
 	{
-	case EMaterialPanelType::Cards:			return GroomAsset->HairGroupsCards.Num();
-	case EMaterialPanelType::Meshes:		return GroomAsset->HairGroupsMeshes.Num();
-	case EMaterialPanelType::Strands:		return GroomAsset->HairGroupsRendering.Num();
-	case EMaterialPanelType::Physics:		return GroomAsset->HairGroupsPhysics.Num();
-	case EMaterialPanelType::Interpolation:	return GroomAsset->HairGroupsInterpolation.Num();
-	case EMaterialPanelType::LODs:			return GroomAsset->HairGroupsLOD.Num();
+		case EMaterialPanelType::Cards:			return GroomAsset->GetHairGroupsCards().Num();
+		case EMaterialPanelType::Meshes:		return GroomAsset->GetHairGroupsMeshes().Num();
+		case EMaterialPanelType::Strands:		return GroomAsset->GetHairGroupsRendering().Num();
+		case EMaterialPanelType::Physics:		return GroomAsset->GetHairGroupsPhysics().Num();
+		case EMaterialPanelType::Interpolation:	return GroomAsset->GetHairGroupsInterpolation().Num();
+		case EMaterialPanelType::LODs:			return GroomAsset->GetHairGroupsLOD().Num();
 	}
 
 	return 0;
@@ -740,10 +740,10 @@ FReply FGroomRenderingDetails::OnRemoveGroupClicked(int32 GroupIndex, FProperty*
 	{
 	case EMaterialPanelType::Cards:		
 	{
-		if (GroupIndex < GroomAsset->HairGroupsCards.Num())
+		if (GroupIndex < GroomAsset->GetHairGroupsCards().Num())
 		{
 			FScopedTransaction Transaction(FText::FromString(TEXT("RemoveCardsGroup")));
-			GroomAsset->HairGroupsCards.RemoveAt(GroupIndex);
+			GroomAsset->GetHairGroupsCards().RemoveAt(GroupIndex);
 
 			FPropertyChangedEvent PropertyChangedEvent(Property, EPropertyChangeType::ArrayRemove);
 			GroomAsset->PostEditChangeProperty(PropertyChangedEvent);
@@ -752,10 +752,10 @@ FReply FGroomRenderingDetails::OnRemoveGroupClicked(int32 GroupIndex, FProperty*
 	break;
 	case EMaterialPanelType::Meshes:
 	{
-		if (GroupIndex < GroomAsset->HairGroupsMeshes.Num())
+		if (GroupIndex < GroomAsset->GetHairGroupsMeshes().Num())
 		{
 			FScopedTransaction Transaction(FText::FromString(TEXT("RemoveMeshesGroup")));
-			GroomAsset->HairGroupsMeshes.RemoveAt(GroupIndex);
+			GroomAsset->GetHairGroupsMeshes().RemoveAt(GroupIndex);
 
 			FPropertyChangedEvent PropertyChangedEvent(Property, EPropertyChangeType::ArrayRemove);
 			GroomAsset->PostEditChangeProperty(PropertyChangedEvent);
@@ -779,9 +779,9 @@ void FGroomRenderingDetails::SetMaterialSlot(int32 GroupIndex, int32 MaterialInd
 	{
 		GetMaterialSlotName(GroupIndex) = NAME_None;
 	}
-	else if (GroupIndex < GroupCount && MaterialIndex >= 0 && MaterialIndex < GroomAsset->HairGroupsMaterials.Num())
+	else if (GroupIndex < GroupCount && MaterialIndex >= 0 && MaterialIndex < GroomAsset->GetHairGroupsMaterials().Num())
 	{
-		GetMaterialSlotName(GroupIndex) = GroomAsset->HairGroupsMaterials[MaterialIndex].SlotName;
+		GetMaterialSlotName(GroupIndex) = GroomAsset->GetHairGroupsMaterials()[MaterialIndex].SlotName;
 	}
 
 	GroomAsset->MarkMaterialsHasChanged();
@@ -794,7 +794,7 @@ TSharedRef<SWidget> FGroomRenderingDetails::OnGenerateStrandsMaterialMenuPicker(
 		return SNullWidget::NullWidget;
 	}
 
-	const int32 MaterialCount = GroomAsset->HairGroupsMaterials.Num();
+	const int32 MaterialCount = GroomAsset->GetHairGroupsMaterials().Num();
 	if(MaterialCount == 0)
 	{
 		return SNullWidget::NullWidget;
@@ -812,7 +812,7 @@ TSharedRef<SWidget> FGroomRenderingDetails::OnGenerateStrandsMaterialMenuPicker(
 	// Add a menu item for material
 	for (int32 MaterialIt = 0; MaterialIt < MaterialCount; ++MaterialIt)
 	{
-		FText MaterialString = FText::FromString(FString::FromInt(MaterialIt) + TEXT(" - ") + GroomAsset->HairGroupsMaterials[MaterialIt].SlotName.ToString());
+		FText MaterialString = FText::FromString(FString::FromInt(MaterialIt) + TEXT(" - ") + GroomAsset->GetHairGroupsMaterials()[MaterialIt].SlotName.ToString());
 		FUIAction Action(FExecuteAction::CreateSP(this, &FGroomRenderingDetails::SetMaterialSlot, GroupIndex, MaterialIt));
 		MenuBuilder.AddMenuEntry(MaterialString, FText::GetEmpty(), FSlateIcon(), Action);
 	}
@@ -827,7 +827,7 @@ FText FGroomRenderingDetails::GetStrandsMaterialName(int32 GroupIndex) const
 	FText MaterialString = FText::FromString(TEXT("Default"));
 	if (MaterialIndex != INDEX_NONE)
 	{
-		MaterialString = FText::FromString(FString::FromInt(MaterialIndex) + TEXT(" - ") + GroomAsset->HairGroupsMaterials[MaterialIndex].SlotName.ToString());
+		MaterialString = FText::FromString(FString::FromInt(MaterialIndex) + TEXT(" - ") + GroomAsset->GetHairGroupsMaterials()[MaterialIndex].SlotName.ToString());
 	}
 	return MaterialString;
 }
@@ -870,8 +870,8 @@ bool FGroomRenderingDetails::CommonResetToDefault(TSharedPtr<IPropertyHandle> Ch
 
 	// For cards & meshes the incoming index is actually the cards/mesh description index, not the group index
 	// For the rest, the group index refers to the actual group index.
-	const bool bIsCardDescIndexValid = GroupIndex < GroomAsset->HairGroupsCards.Num();
-	const bool bIsMeshDescIndexValid = GroupIndex < GroomAsset->HairGroupsMeshes.Num();
+	const bool bIsCardDescIndexValid = GroupIndex < GroomAsset->GetHairGroupsCards().Num();
+	const bool bIsMeshDescIndexValid = GroupIndex < GroomAsset->GetHairGroupsMeshes().Num();
 	const bool bIsGroupIndexValid = GroupIndex < GroomAsset->GetNumHairGroups();
 		
 	// Hair strands
@@ -879,23 +879,23 @@ bool FGroomRenderingDetails::CommonResetToDefault(TSharedPtr<IPropertyHandle> Ch
 	{
 		{
 			FHairGeometrySettings Default;
-			HAIR_RESET1(HairGroupsRendering, FHairGeometrySettings, GeometrySettings, HairWidth);
-			HAIR_RESET1(HairGroupsRendering, FHairGeometrySettings, GeometrySettings, HairRootScale);
-			HAIR_RESET1(HairGroupsRendering, FHairGeometrySettings, GeometrySettings, HairTipScale);
+			HAIR_RESET1(GetHairGroupsRendering(), FHairGeometrySettings, GeometrySettings, HairWidth);
+			HAIR_RESET1(GetHairGroupsRendering(), FHairGeometrySettings, GeometrySettings, HairRootScale);
+			HAIR_RESET1(GetHairGroupsRendering(), FHairGeometrySettings, GeometrySettings, HairTipScale);
 		}
 
 		{
 			FHairShadowSettings Default;
-			HAIR_RESET1(HairGroupsRendering, FHairShadowSettings, ShadowSettings, bVoxelize);
-			HAIR_RESET1(HairGroupsRendering, FHairShadowSettings, ShadowSettings, bUseHairRaytracingGeometry);
-			HAIR_RESET1(HairGroupsRendering, FHairShadowSettings, ShadowSettings, HairRaytracingRadiusScale);
-			HAIR_RESET1(HairGroupsRendering, FHairShadowSettings, ShadowSettings, HairRaytracingRadiusScale);
+			HAIR_RESET1(GetHairGroupsRendering(), FHairShadowSettings, ShadowSettings, bVoxelize);
+			HAIR_RESET1(GetHairGroupsRendering(), FHairShadowSettings, ShadowSettings, bUseHairRaytracingGeometry);
+			HAIR_RESET1(GetHairGroupsRendering(), FHairShadowSettings, ShadowSettings, HairRaytracingRadiusScale);
+			HAIR_RESET1(GetHairGroupsRendering(), FHairShadowSettings, ShadowSettings, HairRaytracingRadiusScale);
 		}
 
 		{
 			FHairAdvancedRenderingSettings Default;
-			HAIR_RESET1(HairGroupsRendering, FHairAdvancedRenderingSettings, AdvancedSettings, bScatterSceneLighting);
-			HAIR_RESET1(HairGroupsRendering, FHairAdvancedRenderingSettings, AdvancedSettings, bUseStableRasterization);
+			HAIR_RESET1(GetHairGroupsRendering(), FHairAdvancedRenderingSettings, AdvancedSettings, bScatterSceneLighting);
+			HAIR_RESET1(GetHairGroupsRendering(), FHairAdvancedRenderingSettings, AdvancedSettings, bUseStableRasterization);
 		}
 	}
 
@@ -904,18 +904,18 @@ bool FGroomRenderingDetails::CommonResetToDefault(TSharedPtr<IPropertyHandle> Ch
 	{
 		{
 			FHairDecimationSettings Default;
-			HAIR_RESET1(HairGroupsInterpolation, FHairDecimationSettings, DecimationSettings, CurveDecimation);
-			HAIR_RESET1(HairGroupsInterpolation, FHairDecimationSettings, DecimationSettings, VertexDecimation);
+			HAIR_RESET1(GetHairGroupsInterpolation(), FHairDecimationSettings, DecimationSettings, CurveDecimation);
+			HAIR_RESET1(GetHairGroupsInterpolation(), FHairDecimationSettings, DecimationSettings, VertexDecimation);
 		}
 
 		{
 			FHairInterpolationSettings Default;
-			HAIR_RESET1(HairGroupsInterpolation, FHairInterpolationSettings, InterpolationSettings, GuideType);
-			HAIR_RESET1(HairGroupsInterpolation, FHairInterpolationSettings, InterpolationSettings, HairToGuideDensity);
-			HAIR_RESET1(HairGroupsInterpolation, FHairInterpolationSettings, InterpolationSettings, bRandomizeGuide);
-			HAIR_RESET1(HairGroupsInterpolation, FHairInterpolationSettings, InterpolationSettings, bUseUniqueGuide);
-			HAIR_RESET1(HairGroupsInterpolation, FHairInterpolationSettings, InterpolationSettings, RiggedGuideNumCurves);
-			HAIR_RESET1(HairGroupsInterpolation, FHairInterpolationSettings, InterpolationSettings, RiggedGuideNumPoints);
+			HAIR_RESET1(GetHairGroupsInterpolation(), FHairInterpolationSettings, InterpolationSettings, GuideType);
+			HAIR_RESET1(GetHairGroupsInterpolation(), FHairInterpolationSettings, InterpolationSettings, HairToGuideDensity);
+			HAIR_RESET1(GetHairGroupsInterpolation(), FHairInterpolationSettings, InterpolationSettings, bRandomizeGuide);
+			HAIR_RESET1(GetHairGroupsInterpolation(), FHairInterpolationSettings, InterpolationSettings, bUseUniqueGuide);
+			HAIR_RESET1(GetHairGroupsInterpolation(), FHairInterpolationSettings, InterpolationSettings, RiggedGuideNumCurves);
+			HAIR_RESET1(GetHairGroupsInterpolation(), FHairInterpolationSettings, InterpolationSettings, RiggedGuideNumPoints);
 		}
 	}
 
@@ -924,19 +924,19 @@ bool FGroomRenderingDetails::CommonResetToDefault(TSharedPtr<IPropertyHandle> Ch
 	{
 		{
 			FHairGroupsLOD Default;
-			HAIR_RESET0(HairGroupsLOD, FHairGroupsLOD, ClusterWorldSize);
-			HAIR_RESET0(HairGroupsLOD, FHairGroupsLOD, ClusterScreenSizeScale);
+			HAIR_RESET0(GetHairGroupsLOD(), FHairGroupsLOD, ClusterWorldSize);
+			HAIR_RESET0(GetHairGroupsLOD(), FHairGroupsLOD, ClusterScreenSizeScale);
 		}
 
-		if (LODIndex>=0 && LODIndex < GroomAsset->HairGroupsLOD[GroupIndex].LODs.Num())
+		if (LODIndex>=0 && LODIndex < GroomAsset->GetHairGroupsLOD()[GroupIndex].LODs.Num())
 		{
 			FHairLODSettings Default;
-			HAIR_RESET1(HairGroupsLOD, FHairLODSettings, LODs[LODIndex], CurveDecimation);
-			HAIR_RESET1(HairGroupsLOD, FHairLODSettings, LODs[LODIndex], VertexDecimation);
-			HAIR_RESET1(HairGroupsLOD, FHairLODSettings, LODs[LODIndex], AngularThreshold);
-			HAIR_RESET1(HairGroupsLOD, FHairLODSettings, LODs[LODIndex], ScreenSize);
-			HAIR_RESET1(HairGroupsLOD, FHairLODSettings, LODs[LODIndex], ThicknessScale);
-			HAIR_RESET1(HairGroupsLOD, FHairLODSettings, LODs[LODIndex], GeometryType);
+			HAIR_RESET1(GetHairGroupsLOD(), FHairLODSettings, LODs[LODIndex], CurveDecimation);
+			HAIR_RESET1(GetHairGroupsLOD(), FHairLODSettings, LODs[LODIndex], VertexDecimation);
+			HAIR_RESET1(GetHairGroupsLOD(), FHairLODSettings, LODs[LODIndex], AngularThreshold);
+			HAIR_RESET1(GetHairGroupsLOD(), FHairLODSettings, LODs[LODIndex], ScreenSize);
+			HAIR_RESET1(GetHairGroupsLOD(), FHairLODSettings, LODs[LODIndex], ThicknessScale);
+			HAIR_RESET1(GetHairGroupsLOD(), FHairLODSettings, LODs[LODIndex], GeometryType);
 		}
 	}
 
@@ -945,21 +945,21 @@ bool FGroomRenderingDetails::CommonResetToDefault(TSharedPtr<IPropertyHandle> Ch
 	{
 		{
 			FHairGroupsCardsSourceDescription Default;
-			HAIR_RESET0(HairGroupsCards, FHairGroupsCardsSourceDescription, SourceType);
-			HAIR_RESET0(HairGroupsCards, FHairGroupsCardsSourceDescription, ProceduralMesh);
-			HAIR_RESET0(HairGroupsCards, FHairGroupsCardsSourceDescription, ImportedMesh);
-			HAIR_RESET0(HairGroupsCards, FHairGroupsCardsSourceDescription, GroupIndex);
-			HAIR_RESET0(HairGroupsCards, FHairGroupsCardsSourceDescription, LODIndex);
+			HAIR_RESET0(GetHairGroupsCards(), FHairGroupsCardsSourceDescription, SourceType);
+			HAIR_RESET0(GetHairGroupsCards(), FHairGroupsCardsSourceDescription, ProceduralMesh);
+			HAIR_RESET0(GetHairGroupsCards(), FHairGroupsCardsSourceDescription, ImportedMesh);
+			HAIR_RESET0(GetHairGroupsCards(), FHairGroupsCardsSourceDescription, GroupIndex);
+			HAIR_RESET0(GetHairGroupsCards(), FHairGroupsCardsSourceDescription, LODIndex);
 		}
 
 		{
 			FHairGroupCardsTextures Default;
-			HAIR_RESET1(HairGroupsCards, FHairGroupCardsTextures, Textures, DepthTexture);
-			HAIR_RESET1(HairGroupsCards, FHairGroupCardsTextures, Textures, CoverageTexture);
-			HAIR_RESET1(HairGroupsCards, FHairGroupCardsTextures, Textures, TangentTexture);
-			HAIR_RESET1(HairGroupsCards, FHairGroupCardsTextures, Textures, AttributeTexture);
-			HAIR_RESET1(HairGroupsCards, FHairGroupCardsTextures, Textures, MaterialTexture);
-			HAIR_RESET1(HairGroupsCards, FHairGroupCardsTextures, Textures, AuxilaryDataTexture);
+			HAIR_RESET1(GetHairGroupsCards(), FHairGroupCardsTextures, Textures, DepthTexture);
+			HAIR_RESET1(GetHairGroupsCards(), FHairGroupCardsTextures, Textures, CoverageTexture);
+			HAIR_RESET1(GetHairGroupsCards(), FHairGroupCardsTextures, Textures, TangentTexture);
+			HAIR_RESET1(GetHairGroupsCards(), FHairGroupCardsTextures, Textures, AttributeTexture);
+			HAIR_RESET1(GetHairGroupsCards(), FHairGroupCardsTextures, Textures, MaterialTexture);
+			HAIR_RESET1(GetHairGroupsCards(), FHairGroupCardsTextures, Textures, AuxilaryDataTexture);
 		}
 	}
 
@@ -968,19 +968,19 @@ bool FGroomRenderingDetails::CommonResetToDefault(TSharedPtr<IPropertyHandle> Ch
 	{
 		{
 			FHairGroupsMeshesSourceDescription Default;
-			HAIR_RESET0(HairGroupsMeshes, FHairGroupsMeshesSourceDescription, ImportedMesh);
-			HAIR_RESET0(HairGroupsMeshes, FHairGroupsMeshesSourceDescription, GroupIndex);
-			HAIR_RESET0(HairGroupsMeshes, FHairGroupsMeshesSourceDescription, LODIndex);
+			HAIR_RESET0(GetHairGroupsMeshes(), FHairGroupsMeshesSourceDescription, ImportedMesh);
+			HAIR_RESET0(GetHairGroupsMeshes(), FHairGroupsMeshesSourceDescription, GroupIndex);
+			HAIR_RESET0(GetHairGroupsMeshes(), FHairGroupsMeshesSourceDescription, LODIndex);
 		}
 
 		{
 			FHairGroupCardsTextures Default;
-			HAIR_RESET1(HairGroupsMeshes, FHairGroupCardsTextures, Textures, DepthTexture);
-			HAIR_RESET1(HairGroupsMeshes, FHairGroupCardsTextures, Textures, CoverageTexture);
-			HAIR_RESET1(HairGroupsMeshes, FHairGroupCardsTextures, Textures, TangentTexture);
-			HAIR_RESET1(HairGroupsMeshes, FHairGroupCardsTextures, Textures, AttributeTexture);
-			HAIR_RESET1(HairGroupsMeshes, FHairGroupCardsTextures, Textures, MaterialTexture);
-			HAIR_RESET1(HairGroupsMeshes, FHairGroupCardsTextures, Textures, AuxilaryDataTexture);
+			HAIR_RESET1(GetHairGroupsMeshes(), FHairGroupCardsTextures, Textures, DepthTexture);
+			HAIR_RESET1(GetHairGroupsMeshes(), FHairGroupCardsTextures, Textures, CoverageTexture);
+			HAIR_RESET1(GetHairGroupsMeshes(), FHairGroupCardsTextures, Textures, TangentTexture);
+			HAIR_RESET1(GetHairGroupsMeshes(), FHairGroupCardsTextures, Textures, AttributeTexture);
+			HAIR_RESET1(GetHairGroupsMeshes(), FHairGroupCardsTextures, Textures, MaterialTexture);
+			HAIR_RESET1(GetHairGroupsMeshes(), FHairGroupCardsTextures, Textures, AuxilaryDataTexture);
 		}
 	}
 
@@ -989,58 +989,58 @@ bool FGroomRenderingDetails::CommonResetToDefault(TSharedPtr<IPropertyHandle> Ch
 	{
 		{
 			FHairSolverSettings Default;
-			HAIR_RESET1(HairGroupsPhysics, FHairSolverSettings, SolverSettings, EnableSimulation);
-			HAIR_RESET1(HairGroupsPhysics, FHairSolverSettings, SolverSettings, NiagaraSolver);
-			HAIR_RESET1(HairGroupsPhysics, FHairSolverSettings, SolverSettings, CustomSystem);
-			HAIR_RESET1(HairGroupsPhysics, FHairSolverSettings, SolverSettings, SubSteps);
-			HAIR_RESET1(HairGroupsPhysics, FHairSolverSettings, SolverSettings, IterationCount);
+			HAIR_RESET1(GetHairGroupsPhysics(), FHairSolverSettings, SolverSettings, EnableSimulation);
+			HAIR_RESET1(GetHairGroupsPhysics(), FHairSolverSettings, SolverSettings, NiagaraSolver);
+			HAIR_RESET1(GetHairGroupsPhysics(), FHairSolverSettings, SolverSettings, CustomSystem);
+			HAIR_RESET1(GetHairGroupsPhysics(), FHairSolverSettings, SolverSettings, SubSteps);
+			HAIR_RESET1(GetHairGroupsPhysics(), FHairSolverSettings, SolverSettings, IterationCount);
 		}
 
 		{
 			FHairExternalForces Default;
-			HAIR_RESET1(HairGroupsPhysics, FHairExternalForces, ExternalForces, GravityVector);
-			HAIR_RESET1(HairGroupsPhysics, FHairExternalForces, ExternalForces, AirDrag);
-			HAIR_RESET1(HairGroupsPhysics, FHairExternalForces, ExternalForces, AirVelocity);
+			HAIR_RESET1(GetHairGroupsPhysics(), FHairExternalForces, ExternalForces, GravityVector);
+			HAIR_RESET1(GetHairGroupsPhysics(), FHairExternalForces, ExternalForces, AirDrag);
+			HAIR_RESET1(GetHairGroupsPhysics(), FHairExternalForces, ExternalForces, AirVelocity);
 		}
 
 		{
 			FHairBendConstraint Default;
-			HAIR_RESET2(HairGroupsPhysics, FHairBendConstraint, MaterialConstraints, BendConstraint, SolveBend);
-			HAIR_RESET2(HairGroupsPhysics, FHairBendConstraint, MaterialConstraints, BendConstraint, ProjectBend);
-			HAIR_RESET2(HairGroupsPhysics, FHairBendConstraint, MaterialConstraints, BendConstraint, BendDamping);
-			HAIR_RESET2(HairGroupsPhysics, FHairBendConstraint, MaterialConstraints, BendConstraint, BendStiffness);
-//			HAIR_RESET2(HairGroupsPhysics, FHairBendConstraint, MaterialConstraints, BendConstraint, BendScale);
+			HAIR_RESET2(GetHairGroupsPhysics(), FHairBendConstraint, MaterialConstraints, BendConstraint, SolveBend);
+			HAIR_RESET2(GetHairGroupsPhysics(), FHairBendConstraint, MaterialConstraints, BendConstraint, ProjectBend);
+			HAIR_RESET2(GetHairGroupsPhysics(), FHairBendConstraint, MaterialConstraints, BendConstraint, BendDamping);
+			HAIR_RESET2(GetHairGroupsPhysics(), FHairBendConstraint, MaterialConstraints, BendConstraint, BendStiffness);
+//			HAIR_RESET2(GetHairGroupsPhysics(), FHairBendConstraint, MaterialConstraints, BendConstraint, BendScale);
 		}
 
 		{
 			FHairStretchConstraint Default;
-			HAIR_RESET2(HairGroupsPhysics, FHairStretchConstraint, MaterialConstraints, StretchConstraint, SolveStretch);
-			HAIR_RESET2(HairGroupsPhysics, FHairStretchConstraint, MaterialConstraints, StretchConstraint, ProjectStretch);
-			HAIR_RESET2(HairGroupsPhysics, FHairStretchConstraint, MaterialConstraints, StretchConstraint, StretchDamping);
-			HAIR_RESET2(HairGroupsPhysics, FHairStretchConstraint, MaterialConstraints, StretchConstraint, StretchStiffness);
-//			HAIR_RESET2(HairGroupsPhysics, FHairStretchConstraint, MaterialConstraints, StretchConstraint, StretchScale);
+			HAIR_RESET2(GetHairGroupsPhysics(), FHairStretchConstraint, MaterialConstraints, StretchConstraint, SolveStretch);
+			HAIR_RESET2(GetHairGroupsPhysics(), FHairStretchConstraint, MaterialConstraints, StretchConstraint, ProjectStretch);
+			HAIR_RESET2(GetHairGroupsPhysics(), FHairStretchConstraint, MaterialConstraints, StretchConstraint, StretchDamping);
+			HAIR_RESET2(GetHairGroupsPhysics(), FHairStretchConstraint, MaterialConstraints, StretchConstraint, StretchStiffness);
+//			HAIR_RESET2(GetHairGroupsPhysics(), FHairStretchConstraint, MaterialConstraints, StretchConstraint, StretchScale);
 		}
 
 		{
 			FHairCollisionConstraint Default;
-			HAIR_RESET2(HairGroupsPhysics, FHairCollisionConstraint, MaterialConstraints, CollisionConstraint, SolveCollision);
-			HAIR_RESET2(HairGroupsPhysics, FHairCollisionConstraint, MaterialConstraints, CollisionConstraint, ProjectCollision);
-			HAIR_RESET2(HairGroupsPhysics, FHairCollisionConstraint, MaterialConstraints, CollisionConstraint, StaticFriction);
-			HAIR_RESET2(HairGroupsPhysics, FHairCollisionConstraint, MaterialConstraints, CollisionConstraint, KineticFriction);
-			HAIR_RESET2(HairGroupsPhysics, FHairCollisionConstraint, MaterialConstraints, CollisionConstraint, StrandsViscosity);
-			HAIR_RESET2(HairGroupsPhysics, FHairCollisionConstraint, MaterialConstraints, CollisionConstraint, GridDimension);
-			HAIR_RESET2(HairGroupsPhysics, FHairCollisionConstraint, MaterialConstraints, CollisionConstraint, CollisionRadius);
-//			HAIR_RESET2(HairGroupsPhysics, FHairCollisionConstraint, MaterialConstraints, CollisionConstraint, RadiusScale);
+			HAIR_RESET2(GetHairGroupsPhysics(), FHairCollisionConstraint, MaterialConstraints, CollisionConstraint, SolveCollision);
+			HAIR_RESET2(GetHairGroupsPhysics(), FHairCollisionConstraint, MaterialConstraints, CollisionConstraint, ProjectCollision);
+			HAIR_RESET2(GetHairGroupsPhysics(), FHairCollisionConstraint, MaterialConstraints, CollisionConstraint, StaticFriction);
+			HAIR_RESET2(GetHairGroupsPhysics(), FHairCollisionConstraint, MaterialConstraints, CollisionConstraint, KineticFriction);
+			HAIR_RESET2(GetHairGroupsPhysics(), FHairCollisionConstraint, MaterialConstraints, CollisionConstraint, StrandsViscosity);
+			HAIR_RESET2(GetHairGroupsPhysics(), FHairCollisionConstraint, MaterialConstraints, CollisionConstraint, GridDimension);
+			HAIR_RESET2(GetHairGroupsPhysics(), FHairCollisionConstraint, MaterialConstraints, CollisionConstraint, CollisionRadius);
+//			HAIR_RESET2(GetHairGroupsPhysics(), FHairCollisionConstraint, MaterialConstraints, CollisionConstraint, RadiusScale);
 		}
 		
 
 		{
 			FHairStrandsParameters Default;
-			HAIR_RESET1(HairGroupsPhysics, FHairStrandsParameters, StrandsParameters, StrandsSize);
-			HAIR_RESET1(HairGroupsPhysics, FHairStrandsParameters, StrandsParameters, StrandsDensity);
-			HAIR_RESET1(HairGroupsPhysics, FHairStrandsParameters, StrandsParameters, StrandsSmoothing);
-			HAIR_RESET1(HairGroupsPhysics, FHairStrandsParameters, StrandsParameters, StrandsThickness);
-//			HAIR_RESET1(HairGroupsPhysics, FHairStrandsParameters, StrandsParameters, ThicknessScale);
+			HAIR_RESET1(GetHairGroupsPhysics(), FHairStrandsParameters, StrandsParameters, StrandsSize);
+			HAIR_RESET1(GetHairGroupsPhysics(), FHairStrandsParameters, StrandsParameters, StrandsDensity);
+			HAIR_RESET1(GetHairGroupsPhysics(), FHairStrandsParameters, StrandsParameters, StrandsSmoothing);
+			HAIR_RESET1(GetHairGroupsPhysics(), FHairStrandsParameters, StrandsParameters, StrandsThickness);
+//			HAIR_RESET1(GetHairGroupsPhysics(), FHairStrandsParameters, StrandsParameters, ThicknessScale);
 		}
 	}
 
@@ -1186,11 +1186,11 @@ void FGroomRenderingDetails::AddPropertySeparator(FName PropertyName, IDetailChi
 FReply FGroomRenderingDetails::OnRemoveLODClicked(int32 GroupIndex, int32 LODIndex, FProperty* Property)
 {
 	check(GroomAsset);
-	if (GroupIndex < GroomAsset->HairGroupsLOD.Num() && LODIndex >= 0 && LODIndex < GroomAsset->HairGroupsLOD[GroupIndex].LODs.Num() && GroomAsset->HairGroupsLOD[GroupIndex].LODs.Num() > 1)
+	if (GroupIndex < GroomAsset->GetHairGroupsLOD().Num() && LODIndex >= 0 && LODIndex < GroomAsset->GetHairGroupsLOD()[GroupIndex].LODs.Num() && GroomAsset->GetHairGroupsLOD()[GroupIndex].LODs.Num() > 1)
 	{
 		FScopedTransaction Transaction(FText::FromString(TEXT("RemoveLOD")));
 
-		GroomAsset->HairGroupsLOD[GroupIndex].LODs.RemoveAt(LODIndex);
+		GroomAsset->GetHairGroupsLOD()[GroupIndex].LODs.RemoveAt(LODIndex);
 
 		FPropertyChangedEvent PropertyChangedEvent(Property, EPropertyChangeType::ArrayRemove);
 		GroomAsset->PostEditChangeProperty(PropertyChangedEvent);
@@ -1201,16 +1201,16 @@ FReply FGroomRenderingDetails::OnRemoveLODClicked(int32 GroupIndex, int32 LODInd
 FReply FGroomRenderingDetails::OnAddLODClicked(int32 GroupIndex, FProperty* Property)
 {
 	const int32 MaxLODCount = 8; 
-	if (GroupIndex < GroomAsset->HairGroupsLOD.Num() && GroomAsset->HairGroupsLOD[GroupIndex].LODs.Num() < MaxLODCount)
+	if (GroupIndex < GroomAsset->GetHairGroupsLOD().Num() && GroomAsset->GetHairGroupsLOD()[GroupIndex].LODs.Num() < MaxLODCount)
 	{
 		FScopedTransaction Transaction(FText::FromString(TEXT("AddLOD")));
 
-		GroomAsset->HairGroupsLOD[GroupIndex].LODs.AddDefaulted();
-		const int32 LODCount = GroomAsset->HairGroupsLOD[GroupIndex].LODs.Num();
+		GroomAsset->GetHairGroupsLOD()[GroupIndex].LODs.AddDefaulted();
+		const int32 LODCount = GroomAsset->GetHairGroupsLOD()[GroupIndex].LODs.Num();
 		if (LODCount > 1)
 		{
-			const FHairLODSettings& PrevLODSettings = GroomAsset->HairGroupsLOD[GroupIndex].LODs[LODCount-2];
-			FHairLODSettings& LODSettings = GroomAsset->HairGroupsLOD[GroupIndex].LODs[LODCount-1];
+			const FHairLODSettings& PrevLODSettings = GroomAsset->GetHairGroupsLOD()[GroupIndex].LODs[LODCount-2];
+			FHairLODSettings& LODSettings = GroomAsset->GetHairGroupsLOD()[GroupIndex].LODs[LODCount-1];
 
 			// Prefill the LOD setting with basic preset
 			LODSettings.VertexDecimation = PrevLODSettings.VertexDecimation * 0.5f;
@@ -1229,7 +1229,7 @@ FReply FGroomRenderingDetails::OnAddLODClicked(int32 GroupIndex, FProperty* Prop
 // Hair_TODO: rename into OnReloadCards
 FReply FGroomRenderingDetails::OnRefreshCards(int32 DescIndex, FProperty* Property)
 {
-	if (DescIndex < GroomAsset->HairGroupsCards.Num() && GroomAsset->HairGroupsCards[DescIndex].SourceType == EHairCardsSourceType::Procedural)
+	if (DescIndex < GroomAsset->GetHairGroupsCards().Num() && GroomAsset->GetHairGroupsCards()[DescIndex].SourceType == EHairCardsSourceType::Procedural)
 	{
 		FScopedTransaction Transaction(FText::FromString(TEXT("RefreshCards")));
 		
@@ -1242,7 +1242,7 @@ FReply FGroomRenderingDetails::OnRefreshCards(int32 DescIndex, FProperty* Proper
 // Hair_TODO: rename into OnGenerageCards
 FReply FGroomRenderingDetails::OnSaveCards(int32 DescIndex, FProperty* Property)
 {
-	if (DescIndex < GroomAsset->HairGroupsCards.Num())
+	if (DescIndex < GroomAsset->GetHairGroupsCards().Num())
 	{
 		GroomAsset->SaveProceduralCards(DescIndex);
 	}
@@ -1251,7 +1251,7 @@ FReply FGroomRenderingDetails::OnSaveCards(int32 DescIndex, FProperty* Property)
 
 FReply FGroomRenderingDetails::OnGenerateCardDataUsingPlugin(int32 GroupIndex)
 {
-	if (GroomAsset && GroomAsset->HairGroupsCards.IsValidIndex(GroupIndex))
+	if (GroomAsset && GroomAsset->GetHairGroupsCards().IsValidIndex(GroupIndex))
 	{
 		TArray<IHairCardGenerator*> HairCardGeneratorPlugins = IModularFeatures::Get().GetModularFeatureImplementations<IHairCardGenerator>(IHairCardGenerator::ModularFeatureName);
 		if (HairCardGeneratorPlugins.Num() > 0)
@@ -1261,7 +1261,7 @@ FReply FGroomRenderingDetails::OnGenerateCardDataUsingPlugin(int32 GroupIndex)
 			const FScopedTransaction Transaction(LOCTEXT("GenerateHairCardsTransaction", "Generate hair cards."));
 
 			// Use a copy so we can only apply changes on success
-			FHairGroupsCardsSourceDescription HairCardsCopy = GroomAsset->HairGroupsCards[GroupIndex];
+			FHairGroupsCardsSourceDescription HairCardsCopy = GroomAsset->GetHairGroupsCards()[GroupIndex];
 			// Clear fields that are supposed to be set by the generation (in case it leaves any unset, and we don't cary over old settings)
 			HairCardsCopy.Textures = FHairGroupCardsTextures();
 
@@ -1269,7 +1269,7 @@ FReply FGroomRenderingDetails::OnGenerateCardDataUsingPlugin(int32 GroupIndex)
 			if (bSuccess)
 			{
 				GroomAsset->Modify();
-				GroomAsset->HairGroupsCards[GroupIndex] = HairCardsCopy;
+				GroomAsset->GetHairGroupsCards()[GroupIndex] = HairCardsCopy;
 			}
 		}
 	}
@@ -1331,7 +1331,7 @@ void FGroomRenderingDetails::OnGenerateElementForLODs(TSharedRef<IPropertyHandle
 	];
 
 	// LOD Stats
-	const FHairStrandsClusterCullingBulkData& ClusterCullingBulkData = GroomAsset->HairGroupsPlatformData[GroupIndex].Strands.ClusterCullingBulkData;
+	const FHairStrandsClusterCullingBulkData& ClusterCullingBulkData = GroomAsset->GetHairGroupsPlatformData()[GroupIndex].Strands.ClusterCullingBulkData;
 	if (ClusterCullingBulkData.IsValid() && LODIndex < ClusterCullingBulkData.Header.LODInfos.Num())
 	{
 		const FHairLODInfo& LODInfo = ClusterCullingBulkData.Header.LODInfos[LODIndex];
@@ -1399,9 +1399,9 @@ TSharedRef<SWidget> FGroomRenderingDetails::MakeGroupNameButtonCustomization(int
 
 static FName GetGroupName(const UGroomAsset* GroomAsset, int32 GroupIndex)
 {
-	if (GroomAsset && GroupIndex >= 0 && GroupIndex < GroomAsset->HairGroupsInfo.Num())
+	if (GroomAsset && GroupIndex >= 0 && GroupIndex < GroomAsset->GetHairGroupsInfo().Num())
 	{
-		return GroomAsset->HairGroupsInfo[GroupIndex].GroupName;
+		return GroomAsset->GetHairGroupsInfo()[GroupIndex].GroupName;
 	}
 	return NAME_None;
 }
@@ -1505,29 +1505,29 @@ void FGroomRenderingDetails::OnGenerateElementForHairGroup(TSharedRef<IPropertyH
 		]
 	];
 
-	if (GroomAsset != nullptr && GroupIndex>=0 && GroupIndex < GroomAsset->HairGroupsInfo.Num() && (PanelType == EMaterialPanelType::Strands || PanelType == EMaterialPanelType::Interpolation))
+	if (GroomAsset != nullptr && GroupIndex>=0 && GroupIndex < GroomAsset->GetHairGroupsInfo().Num() && (PanelType == EMaterialPanelType::Strands || PanelType == EMaterialPanelType::Interpolation))
 	{
 		ChildrenBuilder.AddCustomRow(LOCTEXT("HairStrandsInfo_Array", "HairStrandsInfo"))
 		.ValueContent()
 		.HAlign(HAlign_Fill)
 		[
-			MakeHairStrandsInfoGrid(DetailFontInfo, GroomAsset->HairGroupsInfo[GroupIndex], GroomAsset->HairGroupsPlatformData[GroupIndex].Strands.BulkData.Header.MaxRadius)
+			MakeHairStrandsInfoGrid(DetailFontInfo, GroomAsset->GetHairGroupsInfo()[GroupIndex], GroomAsset->GetHairGroupsPlatformData()[GroupIndex].Strands.BulkData.Header.MaxRadius)
 		];
 		ChildrenBuilder.AddCustomRow(LOCTEXT("HairStrandsAttributeInfo_Array", "HairStrandsAttributeInfo"))
 		.ValueContent()
 		.HAlign(HAlign_Fill)
 		[
-			MakeHairStrandsAttributeInfoGrid(DetailFontInfo, GroomAsset->HairGroupsPlatformData[GroupIndex].Strands.BulkData.Header.ImportedAttributes, GroomAsset->HairGroupsPlatformData[GroupIndex].Strands.BulkData.Header.ImportedAttributeFlags)
+			MakeHairStrandsAttributeInfoGrid(DetailFontInfo, GroomAsset->GetHairGroupsPlatformData()[GroupIndex].Strands.BulkData.Header.ImportedAttributes, GroomAsset->GetHairGroupsPlatformData()[GroupIndex].Strands.BulkData.Header.ImportedAttributeFlags)
 		];
 	}
 
-	if (GroomAsset != nullptr && GroupIndex >= 0 && GroupIndex < GroomAsset->HairGroupsCards.Num() && (PanelType == EMaterialPanelType::Cards))
+	if (GroomAsset != nullptr && GroupIndex >= 0 && GroupIndex < GroomAsset->GetHairGroupsCards().Num() && (PanelType == EMaterialPanelType::Cards))
 	{
 		ChildrenBuilder.AddCustomRow(LOCTEXT("HairCardsInfo_Array", "HairCardsInfo"))
 		.ValueContent()
 		.HAlign(HAlign_Fill)
 		[
-			MakeHairCardsInfoGrid(DetailFontInfo, GroomAsset->HairGroupsCards[GroupIndex].CardsInfo)
+			MakeHairCardsInfoGrid(DetailFontInfo, GroomAsset->GetHairGroupsCards()[GroupIndex].CardsInfo)
 		];
 	}
 
@@ -1619,8 +1619,8 @@ void FGroomRenderingDetails::OnGenerateElementForHairGroup(TSharedRef<IPropertyH
 				TWeakObjectPtr<UGroomAsset> GroomAssetPtr = GroomAsset;
 				auto ShouldShowProceduralProperties = [GroomAssetPtr, GroupIndex]()
 				{
-					return GroomAssetPtr.IsValid() && GroomAssetPtr->HairGroupsCards.IsValidIndex(GroupIndex) &&
-						GroomAssetPtr->HairGroupsCards[GroupIndex].SourceType == EHairCardsSourceType::Procedural;
+					return GroomAssetPtr.IsValid() && GroomAssetPtr->GetHairGroupsCards().IsValidIndex(GroupIndex) &&
+						GroomAssetPtr->GetHairGroupsCards()[GroupIndex].SourceType == EHairCardsSourceType::Procedural;
 				};
 				TAttribute<EVisibility> ProceduralPropertyVisibility = TAttribute<EVisibility>::CreateLambda([ShouldShowProceduralProperties]()
 					{
@@ -1864,25 +1864,25 @@ FString FGroomRenderingDetails::OnGetObjectPath(int32 GroupIndex) const
 	{
 	case EMaterialPanelType::Cards:
 	{
-		if (GroupIndex < GroomAsset->HairGroupsCards.Num())
+		if (GroupIndex < GroomAsset->GetHairGroupsCards().Num())
 		{
-			MaterialIndex = GroomAsset->GetMaterialIndex(GroomAsset->HairGroupsCards[GroupIndex].MaterialSlotName);
+			MaterialIndex = GroomAsset->GetMaterialIndex(GroomAsset->GetHairGroupsCards()[GroupIndex].MaterialSlotName);
 		}
 	}
 	break;
 	case EMaterialPanelType::Meshes:
 	{
-		if (GroupIndex < GroomAsset->HairGroupsMeshes.Num())
+		if (GroupIndex < GroomAsset->GetHairGroupsMeshes().Num())
 		{
-			MaterialIndex = GroomAsset->GetMaterialIndex(GroomAsset->HairGroupsMeshes[GroupIndex].MaterialSlotName);
+			MaterialIndex = GroomAsset->GetMaterialIndex(GroomAsset->GetHairGroupsMeshes()[GroupIndex].MaterialSlotName);
 		}
 	}
 	break;
 	case EMaterialPanelType::Strands:
 	{
-		if (GroupIndex < GroomAsset->HairGroupsRendering.Num())
+		if (GroupIndex < GroomAsset->GetHairGroupsRendering().Num())
 		{
-			MaterialIndex = GroomAsset->GetMaterialIndex(GroomAsset->HairGroupsRendering[GroupIndex].MaterialSlotName);
+			MaterialIndex = GroomAsset->GetMaterialIndex(GroomAsset->GetHairGroupsRendering()[GroupIndex].MaterialSlotName);
 		}
 	}
 	break;
@@ -1891,7 +1891,7 @@ FString FGroomRenderingDetails::OnGetObjectPath(int32 GroupIndex) const
 	if (MaterialIndex == INDEX_NONE)
 		return FString();
 	
-	return GroomAsset->HairGroupsMaterials[MaterialIndex].Material->GetPathName();
+	return GroomAsset->GetHairGroupsMaterials()[MaterialIndex].Material->GetPathName();
 }
 
 /**
diff --git a/Engine/Plugins/Runtime/HairStrands/Source/HairStrandsEditor/Private/GroomCacheImporter.cpp b/Engine/Plugins/Runtime/HairStrands/Source/HairStrandsEditor/Private/GroomCacheImporter.cpp
--- a/Engine/Plugins/Runtime/HairStrands/Source/HairStrandsEditor/Private/GroomCacheImporter.cpp
+++ b/Engine/Plugins/Runtime/HairStrands/Source/HairStrandsEditor/Private/GroomCacheImporter.cpp
@@ -124,7 +124,7 @@ TArray<UGroomCache*> FGroomCacheImporter::ImportGroomCache(const FString& Source
 		FScopedSlowTask SlowTask(NumFrames, LOCTEXT("ImportGroomCache", "Importing GroomCache frames"));
 		SlowTask.MakeDialog();
 
-		const TArray<FHairGroupPlatformData>& GroupPlatformData = GroomAssetForCache->HairGroupsPlatformData;
+		const TArray<FHairGroupPlatformData>& GroupPlatformData = GroomAssetForCache->GetHairGroupsPlatformData();
 
 		// Each frame is translated into a HairDescription and processed into HairGroupData
 		for (int32 FrameIndex = AnimInfo.StartFrame; FrameIndex < AnimInfo.EndFrame + 1; ++FrameIndex)
@@ -157,14 +157,14 @@ TArray<UGroomCache*> FGroomCacheImporter::ImportGroomCache(const FString& Source
 					break;
 				}
 
-				TArray<FHairGroupInfoWithVisibility> HairGroupsInfo = GroomAssetForCache->HairGroupsInfo;
+				TArray<FHairGroupInfoWithVisibility> HairGroupsInfo = GroomAssetForCache->GetHairGroupsInfo();
 				TArray<FHairDescriptionGroup> HairGroupsData;
 				HairGroupsData.SetNum(GroupCount);
 				for (uint32 GroupIndex = 0; GroupIndex < GroupCount; ++GroupIndex)
 				{
 					const FHairDescriptionGroup& HairGroup = HairDescriptionGroups.HairGroups[GroupIndex];
 					FHairDescriptionGroup& HairGroupData = HairGroupsData[GroupIndex];
-					FGroomBuilder::BuildData(HairGroup, GroomAssetForCache->HairGroupsInterpolation[GroupIndex], HairGroupsInfo[GroupIndex], HairGroupData.Strands, HairGroupData.Guides);
+					FGroomBuilder::BuildData(HairGroup, GroomAssetForCache->GetHairGroupsInterpolation()[GroupIndex], HairGroupsInfo[GroupIndex], HairGroupData.Strands, HairGroupData.Guides);
 				}
 
 				// Validate that the GroomCache has the same topology as the static groom
diff --git a/Engine/Plugins/Runtime/HairStrands/Source/HairStrandsEditor/Private/GroomComponentDetailsCustomization.cpp b/Engine/Plugins/Runtime/HairStrands/Source/HairStrandsEditor/Private/GroomComponentDetailsCustomization.cpp
--- a/Engine/Plugins/Runtime/HairStrands/Source/HairStrandsEditor/Private/GroomComponentDetailsCustomization.cpp
+++ b/Engine/Plugins/Runtime/HairStrands/Source/HairStrandsEditor/Private/GroomComponentDetailsCustomization.cpp
@@ -256,9 +256,9 @@ void FGroomComponentDetailsCustomization::OnGenerateElementForHairGroup(TSharedR
 	if (GroomComponentPtr != nullptr && GroupIndex>=0 && GroupIndex < GroomComponentPtr->GroomGroupsDesc.Num())
 	{
 		FHairGroupInfo Infos;
-		if (GroomComponentPtr->GroomAsset && GroupIndex < GroomComponentPtr->GroomAsset->HairGroupsInfo.Num())
+		if (GroomComponentPtr->GroomAsset && GroupIndex < GroomComponentPtr->GroomAsset->GetHairGroupsInfo().Num())
 		{
-			Infos = GroomComponentPtr->GroomAsset->HairGroupsInfo[GroupIndex];
+			Infos = GroomComponentPtr->GroomAsset->GetHairGroupsInfo()[GroupIndex];
 		}
 		ChildrenBuilder.AddCustomRow(LOCTEXT("HairInfo_Separator", "Separator"))
 		.ValueContent()
diff --git a/Engine/Plugins/Runtime/HairStrands/Source/HairStrandsEditor/Private/GroomCustomAssetEditorToolkit.cpp b/Engine/Plugins/Runtime/HairStrands/Source/HairStrandsEditor/Private/GroomCustomAssetEditorToolkit.cpp
--- a/Engine/Plugins/Runtime/HairStrands/Source/HairStrandsEditor/Private/GroomCustomAssetEditorToolkit.cpp
+++ b/Engine/Plugins/Runtime/HairStrands/Source/HairStrandsEditor/Private/GroomCustomAssetEditorToolkit.cpp
@@ -537,14 +537,14 @@ void FGroomCustomAssetEditorToolkit::InitCustomAssetEditor(const EToolkitMode::T
 		}
 	}
 
-	FProperty* P0 = FindFProperty<FProperty>(GroomAsset->GetClass(), GET_MEMBER_NAME_CHECKED(UGroomAsset, HairGroupsInterpolation));
-	FProperty* P1 = FindFProperty<FProperty>(GroomAsset->GetClass(), GET_MEMBER_NAME_CHECKED(UGroomAsset, HairGroupsRendering));
-	FProperty* P2 = FindFProperty<FProperty>(GroomAsset->GetClass(), GET_MEMBER_NAME_CHECKED(UGroomAsset, HairGroupsPhysics));
-	FProperty* P3 = FindFProperty<FProperty>(GroomAsset->GetClass(), GET_MEMBER_NAME_CHECKED(UGroomAsset, HairGroupsCards));
-	FProperty* P4 = FindFProperty<FProperty>(GroomAsset->GetClass(), GET_MEMBER_NAME_CHECKED(UGroomAsset, HairGroupsLOD));
-	FProperty* P5 = FindFProperty<FProperty>(GroomAsset->GetClass(), GET_MEMBER_NAME_CHECKED(UGroomAsset, HairGroupsMeshes));
-	FProperty* P6 = FindFProperty<FProperty>(GroomAsset->GetClass(), GET_MEMBER_NAME_CHECKED(UGroomAsset, HairGroupsMaterials));
-	FProperty* P7 = FindFProperty<FProperty>(GroomAsset->GetClass(), GET_MEMBER_NAME_CHECKED(UGroomAsset, HairGroupsInfo));
+	FProperty* P0 = FindFProperty<FProperty>(GroomAsset->GetClass(), UGroomAsset::GetHairGroupsInterpolationMemberName());
+	FProperty* P1 = FindFProperty<FProperty>(GroomAsset->GetClass(), UGroomAsset::GetHairGroupsRenderingMemberName());
+	FProperty* P2 = FindFProperty<FProperty>(GroomAsset->GetClass(), UGroomAsset::GetHairGroupsPhysicsMemberName());
+	FProperty* P3 = FindFProperty<FProperty>(GroomAsset->GetClass(), UGroomAsset::GetHairGroupsCardsMemberName());
+	FProperty* P4 = FindFProperty<FProperty>(GroomAsset->GetClass(), UGroomAsset::GetHairGroupsLODMemberName());
+	FProperty* P5 = FindFProperty<FProperty>(GroomAsset->GetClass(), UGroomAsset::GetHairGroupsMeshesMemberName());
+	FProperty* P6 = FindFProperty<FProperty>(GroomAsset->GetClass(), UGroomAsset::GetHairGroupsMaterialsMemberName());
+	FProperty* P7 = FindFProperty<FProperty>(GroomAsset->GetClass(), UGroomAsset::GetHairGroupsInfoMemberName());
 	FProperty* P9 = FindFProperty<FProperty>(GroomAsset->GetClass(), GET_MEMBER_NAME_CHECKED(UGroomAsset, AssetUserData));
 
 	P7->SetMetaData(TEXT("Category"), TEXT("Hidden"));
diff --git a/Engine/Plugins/Runtime/HairStrands/Source/HairStrandsEditor/Private/GroomFactory.cpp b/Engine/Plugins/Runtime/HairStrands/Source/HairStrandsEditor/Private/GroomFactory.cpp
--- a/Engine/Plugins/Runtime/HairStrands/Source/HairStrandsEditor/Private/GroomFactory.cpp
+++ b/Engine/Plugins/Runtime/HairStrands/Source/HairStrandsEditor/Private/GroomFactory.cpp
@@ -21,8 +21,8 @@ UObject* UGroomFactory::FactoryCreateNew(UClass* InClass, UObject* InParent, FNa
 
 	// Groom has to have at least one group to edit with one cards LOD
 	Groom->SetNumGroup(1);
-	Groom->HairGroupsLOD[0].LODs[0].GeometryType = EGroomGeometryType::Cards;
-	Groom->HairGroupsInterpolation[0].InterpolationSettings.bUseUniqueGuide = true;
+	Groom->GetHairGroupsLOD()[0].LODs[0].GeometryType = EGroomGeometryType::Cards;
+	Groom->GetHairGroupsInterpolation()[0].InterpolationSettings.bUseUniqueGuide = true;
 
 	Groom->UpdateHairGroupsInfo();
 	Groom->UpdateCachedSettings();
diff --git a/Engine/Plugins/Runtime/HairStrands/Source/HairStrandsEditor/Private/GroomImportOptionsWindow.cpp b/Engine/Plugins/Runtime/HairStrands/Source/HairStrandsEditor/Private/GroomImportOptionsWindow.cpp
--- a/Engine/Plugins/Runtime/HairStrands/Source/HairStrandsEditor/Private/GroomImportOptionsWindow.cpp
+++ b/Engine/Plugins/Runtime/HairStrands/Source/HairStrandsEditor/Private/GroomImportOptionsWindow.cpp
@@ -114,7 +114,7 @@ void SGroomImportOptionsWindow::UpdateStatus(UGroomHairGroupsPreview* Descriptio
 			return;
 		}
 
-		const TArray<FHairGroupPlatformData>& GroomHairGroupsData = GroomAssetForCache->HairGroupsPlatformData;
+		const TArray<FHairGroupPlatformData>& GroomHairGroupsData = GroomAssetForCache->GetHairGroupsPlatformData();
 		if (GroomHairGroupsData.Num() != Description->Groups.Num())
 		{
 			CurrentStatus = bGuidesOnly ? EHairDescriptionStatus::GuidesOnlyIncompatible : EHairDescriptionStatus::GroomCacheIncompatible;
diff --git a/Engine/Plugins/Runtime/HairStrands/Source/HairStrandsEditor/Private/GroomMaterialDetails.cpp b/Engine/Plugins/Runtime/HairStrands/Source/HairStrandsEditor/Private/GroomMaterialDetails.cpp
--- a/Engine/Plugins/Runtime/HairStrands/Source/HairStrandsEditor/Private/GroomMaterialDetails.cpp
+++ b/Engine/Plugins/Runtime/HairStrands/Source/HairStrandsEditor/Private/GroomMaterialDetails.cpp
@@ -185,7 +185,7 @@ FText FGroomMaterialDetails::GetMaterialSlotNameText(int32 MaterialIndex) const
 {	
 	if (IsMaterialValid(MaterialIndex))
 	{
-		return FText::FromName(GroomAsset->HairGroupsMaterials[MaterialIndex].SlotName);
+		return FText::FromName(GroomAsset->GetHairGroupsMaterials()[MaterialIndex].SlotName);
 	}
 
 	return LOCTEXT("HairMaterial_InvalidIndex", "Invalid Material Index");
@@ -200,18 +200,18 @@ void FGroomMaterialDetails::CustomizeDetails( IDetailLayoutBuilder& DetailLayout
 
 	// Hide all properties
 	{
-		DetailLayout.HideProperty(DetailLayout.GetProperty(GET_MEMBER_NAME_CHECKED(UGroomAsset, HairGroupsInterpolation), UGroomAsset::StaticClass()));
-		DetailLayout.HideProperty(DetailLayout.GetProperty(GET_MEMBER_NAME_CHECKED(UGroomAsset, HairGroupsRendering), UGroomAsset::StaticClass()));
-		DetailLayout.HideProperty(DetailLayout.GetProperty(GET_MEMBER_NAME_CHECKED(UGroomAsset, HairGroupsPhysics), UGroomAsset::StaticClass()));
-		DetailLayout.HideProperty(DetailLayout.GetProperty(GET_MEMBER_NAME_CHECKED(UGroomAsset, HairGroupsCards), UGroomAsset::StaticClass()));
-		DetailLayout.HideProperty(DetailLayout.GetProperty(GET_MEMBER_NAME_CHECKED(UGroomAsset, HairGroupsMeshes), UGroomAsset::StaticClass()));
-		DetailLayout.HideProperty(DetailLayout.GetProperty(GET_MEMBER_NAME_CHECKED(UGroomAsset, HairGroupsMaterials), UGroomAsset::StaticClass()));
-		DetailLayout.HideProperty(DetailLayout.GetProperty(GET_MEMBER_NAME_CHECKED(UGroomAsset, HairGroupsLOD), UGroomAsset::StaticClass()));
-		DetailLayout.HideProperty(DetailLayout.GetProperty(GET_MEMBER_NAME_CHECKED(UGroomAsset, HairGroupsInfo), UGroomAsset::StaticClass()));
-		DetailLayout.HideProperty(DetailLayout.GetProperty(GET_MEMBER_NAME_CHECKED(UGroomAsset, EnableGlobalInterpolation), UGroomAsset::StaticClass()));
-		DetailLayout.HideProperty(DetailLayout.GetProperty(GET_MEMBER_NAME_CHECKED(UGroomAsset, HairInterpolationType), UGroomAsset::StaticClass()));
-		DetailLayout.HideProperty(DetailLayout.GetProperty(GET_MEMBER_NAME_CHECKED(UGroomAsset, MinLOD), UGroomAsset::StaticClass()));
-		DetailLayout.HideProperty(DetailLayout.GetProperty(GET_MEMBER_NAME_CHECKED(UGroomAsset, DisableBelowMinLodStripping), UGroomAsset::StaticClass()));
+		DetailLayout.HideProperty(DetailLayout.GetProperty(UGroomAsset::GetHairGroupsInterpolationMemberName(), UGroomAsset::StaticClass()));
+		DetailLayout.HideProperty(DetailLayout.GetProperty(UGroomAsset::GetHairGroupsRenderingMemberName(), UGroomAsset::StaticClass()));
+		DetailLayout.HideProperty(DetailLayout.GetProperty(UGroomAsset::GetHairGroupsPhysicsMemberName(), UGroomAsset::StaticClass()));
+		DetailLayout.HideProperty(DetailLayout.GetProperty(UGroomAsset::GetHairGroupsCardsMemberName(), UGroomAsset::StaticClass()));
+		DetailLayout.HideProperty(DetailLayout.GetProperty(UGroomAsset::GetHairGroupsMeshesMemberName(), UGroomAsset::StaticClass()));
+		DetailLayout.HideProperty(DetailLayout.GetProperty(UGroomAsset::GetHairGroupsMaterialsMemberName(), UGroomAsset::StaticClass()));
+		DetailLayout.HideProperty(DetailLayout.GetProperty(UGroomAsset::GetHairGroupsLODMemberName(), UGroomAsset::StaticClass()));
+		DetailLayout.HideProperty(DetailLayout.GetProperty(UGroomAsset::GetHairGroupsInfoMemberName(), UGroomAsset::StaticClass()));
+		DetailLayout.HideProperty(DetailLayout.GetProperty(UGroomAsset::GetEnableGlobalInterpolationMemberName(), UGroomAsset::StaticClass()));
+		DetailLayout.HideProperty(DetailLayout.GetProperty(UGroomAsset::GetHairInterpolationTypeMemberName(), UGroomAsset::StaticClass()));
+		DetailLayout.HideProperty(DetailLayout.GetProperty(UGroomAsset::GetMinLODMemberName(), UGroomAsset::StaticClass()));
+		DetailLayout.HideProperty(DetailLayout.GetProperty(UGroomAsset::GetDisableBelowMinLodStrippingMemberName(), UGroomAsset::StaticClass()));
 	}
 
 	GroomDetailLayout = &DetailLayout;
@@ -223,9 +223,9 @@ void FGroomMaterialDetails::OnGetMaterialsForArray(class IMaterialListBuilder& O
 	if (!GroomAsset)
 		return;
 
-	for (int32 MaterialIndex = 0; MaterialIndex < GroomAsset->HairGroupsMaterials.Num(); ++MaterialIndex)
+	for (int32 MaterialIndex = 0; MaterialIndex < GroomAsset->GetHairGroupsMaterials().Num(); ++MaterialIndex)
 	{
-		OutMaterials.AddMaterial(MaterialIndex, GroomAsset->HairGroupsMaterials[MaterialIndex].Material, true);
+		OutMaterials.AddMaterial(MaterialIndex, GroomAsset->GetHairGroupsMaterials()[MaterialIndex].Material, true);
 	}
 }
 
@@ -240,36 +240,36 @@ void FGroomMaterialDetails::OnMaterialArrayChanged(UMaterialInterface* NewMateri
 	FProperty* MaterialProperty = FindFProperty<FProperty>(UGroomAsset::StaticClass(), "HairGroupsMaterials");
 	check(MaterialProperty);
 	GroomAsset->PreEditChange(MaterialProperty);
-	check(GroomAsset->HairGroupsMaterials.Num() > SlotIndex)
+	check(GroomAsset->GetHairGroupsMaterials().Num() > SlotIndex)
 
 	if (NewMaterial != PrevMaterial)
 	{
 		GEditor->BeginTransaction(LOCTEXT("GroomEditorMaterialChanged", "Groom editor: material changed"));
 		bMadeTransaction = true;
 		GroomAsset->Modify();
-		GroomAsset->HairGroupsMaterials[SlotIndex].Material = NewMaterial;
+		GroomAsset->GetHairGroupsMaterials()[SlotIndex].Material = NewMaterial;
 
 		// Add a default name to the material slot if this slot was manually add and there is no name yet
-		if (NewMaterial != nullptr && GroomAsset->HairGroupsMaterials[SlotIndex].SlotName == NAME_None)
+		if (NewMaterial != nullptr && GroomAsset->GetHairGroupsMaterials()[SlotIndex].SlotName == NAME_None)
 		{
-			if (GroomAsset->HairGroupsMaterials[SlotIndex].SlotName == NAME_None)
+			if (GroomAsset->GetHairGroupsMaterials()[SlotIndex].SlotName == NAME_None)
 			{					
-				GroomAsset->HairGroupsMaterials[SlotIndex].SlotName = NewMaterial->GetFName();
+				GroomAsset->GetHairGroupsMaterials()[SlotIndex].SlotName = NewMaterial->GetFName();
 			}
 
 			//Ensure the imported material slot name is unique
-			if (GroomAsset->HairGroupsMaterials[SlotIndex].SlotName == NAME_None)
+			if (GroomAsset->GetHairGroupsMaterials()[SlotIndex].SlotName == NAME_None)
 			{
 				UGroomAsset* LocalGroomAsset = GroomAsset;
 				auto IsMaterialNameUnique = [&LocalGroomAsset, SlotIndex](const FName TestName)
 				{
-					for (int32 MaterialIndex = 0; MaterialIndex < LocalGroomAsset->HairGroupsMaterials.Num(); ++MaterialIndex)
+					for (int32 MaterialIndex = 0; MaterialIndex < LocalGroomAsset->GetHairGroupsMaterials().Num(); ++MaterialIndex)
 					{
 						if (MaterialIndex == SlotIndex)
 						{
 							continue;
 						}
-						if (LocalGroomAsset->HairGroupsMaterials[MaterialIndex].SlotName == TestName)
+						if (LocalGroomAsset->GetHairGroupsMaterials()[MaterialIndex].SlotName == TestName)
 						{
 							return false;
 						}
@@ -290,7 +290,7 @@ void FGroomMaterialDetails::OnMaterialArrayChanged(UMaterialInterface* NewMateri
 						MaterialSlotName = NewMaterial->GetName() + TEXT("_") + FString::FromInt(MatchNameCounter);
 					}
 				}
-				GroomAsset->HairGroupsMaterials[SlotIndex].SlotName = FName(*MaterialSlotName);
+				GroomAsset->GetHairGroupsMaterials()[SlotIndex].SlotName = FName(*MaterialSlotName);
 			}
 		}
 	}
@@ -326,7 +326,7 @@ FReply FGroomMaterialDetails::AddMaterialSlot()
 	do
 	{
 		bHasUniqueName = true;
-		for (const FHairGroupsMaterial& Group : GroomAsset->HairGroupsMaterials)
+		for (const FHairGroupsMaterial& Group : GroomAsset->GetHairGroupsMaterials())
 		{
 			if (Group.SlotName == SlotName)
 			{
@@ -340,7 +340,7 @@ FReply FGroomMaterialDetails::AddMaterialSlot()
 
 	// Add new material
 	NewMaterial.SlotName = SlotName;
-	GroomAsset->HairGroupsMaterials.Add(NewMaterial);
+	GroomAsset->GetHairGroupsMaterials().Add(NewMaterial);
 	GroomAsset->PostEditChange();
 
 	return FReply::Handled();
@@ -352,7 +352,7 @@ FText FGroomMaterialDetails::GetMaterialArrayText() const
 	int32 SlotNumber = 0;
 	if (GroomAsset)
 	{
-		SlotNumber = GroomAsset->HairGroupsMaterials.Num();
+		SlotNumber = GroomAsset->GetHairGroupsMaterials().Num();
 	}
 	MaterialArrayText = FString::FromInt(SlotNumber) + MaterialArrayText;
 	return FText::FromString(MaterialArrayText);
@@ -362,7 +362,7 @@ FText FGroomMaterialDetails::GetMaterialNameText(int32 MaterialIndex) const
 {
 	if (IsMaterialValid(MaterialIndex))
 	{
-		return FText::FromName(GroomAsset->HairGroupsMaterials[MaterialIndex].SlotName);
+		return FText::FromName(GroomAsset->GetHairGroupsMaterials()[MaterialIndex].SlotName);
 	}
 	return FText::FromName(NAME_None);
 }
@@ -370,7 +370,7 @@ FText FGroomMaterialDetails::GetMaterialNameText(int32 MaterialIndex) const
 void FGroomMaterialDetails::OnMaterialNameCommitted(const FText& InValue, ETextCommit::Type CommitType, int32 MaterialIndex)
 {
 	FName NewSlotName = FName(*(InValue.ToString()));
-	if (IsMaterialValid(MaterialIndex) && NewSlotName != GroomAsset->HairGroupsMaterials[MaterialIndex].SlotName)
+	if (IsMaterialValid(MaterialIndex) && NewSlotName != GroomAsset->GetHairGroupsMaterials()[MaterialIndex].SlotName)
 	{
 		FScopedTransaction ScopeTransaction(LOCTEXT("PersonaMaterialSlotNameChanged", "Persona editor: Material slot name change"));
 
@@ -379,22 +379,22 @@ void FGroomMaterialDetails::OnMaterialNameCommitted(const FText& InValue, ETextC
 		GroomAsset->PreEditChange(ChangedProperty);
 
 		// Rename group which were using the old slot name
-		FName PreviousSlotName = GroomAsset->HairGroupsMaterials[MaterialIndex].SlotName;
-		for (FHairGroupsRendering& Group : GroomAsset->HairGroupsRendering)
+		FName PreviousSlotName = GroomAsset->GetHairGroupsMaterials()[MaterialIndex].SlotName;
+		for (FHairGroupsRendering& Group : GroomAsset->GetHairGroupsRendering())
 		{
 			if (PreviousSlotName == Group.MaterialSlotName)
 			{
 				Group.MaterialSlotName = NewSlotName;
 			}
 		}
-		for (FHairGroupsCardsSourceDescription& Group : GroomAsset->HairGroupsCards)
+		for (FHairGroupsCardsSourceDescription& Group : GroomAsset->GetHairGroupsCards())
 		{
 			if (PreviousSlotName == Group.MaterialSlotName)
 			{
 				Group.MaterialSlotName = NewSlotName;
 			}
 		}
-		for (FHairGroupsMeshesSourceDescription& Group : GroomAsset->HairGroupsMeshes)
+		for (FHairGroupsMeshesSourceDescription& Group : GroomAsset->GetHairGroupsMeshes())
 		{
 			if (PreviousSlotName == Group.MaterialSlotName)
 			{
@@ -402,7 +402,7 @@ void FGroomMaterialDetails::OnMaterialNameCommitted(const FText& InValue, ETextC
 			}
 		}
 
-		GroomAsset->HairGroupsMaterials[MaterialIndex].SlotName = NewSlotName;
+		GroomAsset->GetHairGroupsMaterials()[MaterialIndex].SlotName = NewSlotName;
 
 		FPropertyChangedEvent PropertyUpdateStruct(ChangedProperty);
 		GroomAsset->PostEditChangeProperty(PropertyUpdateStruct);
@@ -428,7 +428,7 @@ TSharedRef<SWidget> FGroomMaterialDetails::OnGenerateCustomMaterialWidgetsForMat
 
 bool FGroomMaterialDetails::IsMaterialValid(int32 MaterialIndex) const
 {
-	return GroomAsset && MaterialIndex >= 0 && MaterialIndex < GroomAsset->HairGroupsMaterials.Num();
+	return GroomAsset && MaterialIndex >= 0 && MaterialIndex < GroomAsset->GetHairGroupsMaterials().Num();
 }
 
 bool FGroomMaterialDetails::CanDeleteMaterialSlot(int32 MaterialIndex) const
@@ -457,7 +457,7 @@ void FGroomMaterialDetails::OnDeleteMaterialSlot(int32 MaterialIndex)
 		}
 		bDeleteWarningConsumed = true;
 	}
-	GroomAsset->HairGroupsMaterials.RemoveAt(MaterialIndex);
+	GroomAsset->GetHairGroupsMaterials().RemoveAt(MaterialIndex);
 
 	FScopedTransaction Transaction(LOCTEXT("PersonaOnDeleteMaterialSlotTransaction", "Persona editor: Delete material slot"));
 	GroomAsset->Modify();
diff --git a/Engine/Plugins/Runtime/HairStrands/Source/HairStrandsEditor/Private/HairStrandsImporter.cpp b/Engine/Plugins/Runtime/HairStrands/Source/HairStrandsEditor/Private/HairStrandsImporter.cpp
--- a/Engine/Plugins/Runtime/HairStrands/Source/HairStrandsEditor/Private/HairStrandsImporter.cpp
+++ b/Engine/Plugins/Runtime/HairStrands/Source/HairStrandsEditor/Private/HairStrandsImporter.cpp
@@ -48,7 +48,7 @@ UGroomAsset* FHairStrandsImporter::ImportHair(const FHairImportContext& ImportCo
 	{
 		// Insure the interpolation settings matches between the importer and the actual asset
 		const FHairGroupsInterpolation& InterpolationSettings = ImportContext.ImportOptions->InterpolationSettings[GroupIndex];
-		OutHairAsset->HairGroupsInterpolation[GroupIndex] = InterpolationSettings;
+		OutHairAsset->GetHairGroupsInterpolation()[GroupIndex] = InterpolationSettings;
 	}
 
 	const bool bSucceeded = OutHairAsset->CacheDerivedDatas();
diff --git a/Engine/Plugins/Runtime/HairStrands/Source/HairStrandsEditor/Private/ReimportHairStrandsFactory.cpp b/Engine/Plugins/Runtime/HairStrands/Source/HairStrandsEditor/Private/ReimportHairStrandsFactory.cpp
--- a/Engine/Plugins/Runtime/HairStrands/Source/HairStrandsEditor/Private/ReimportHairStrandsFactory.cpp
+++ b/Engine/Plugins/Runtime/HairStrands/Source/HairStrandsEditor/Private/ReimportHairStrandsFactory.cpp
@@ -366,7 +366,7 @@ EReimportResult::Type UReimportHairStrandsFactory::Reimport(UObject* Obj)
 		GroomReimportOptions->InterpolationSettings.Init(FHairGroupsInterpolation(), GroupCount);
 		for (uint32 GroupIndex = 0; GroupIndex < GroupCount; ++GroupIndex)
 		{
-			GroomReimportOptions->InterpolationSettings[GroupIndex] = HairAsset->HairGroupsInterpolation[GroupIndex];
+			GroomReimportOptions->InterpolationSettings[GroupIndex] = HairAsset->GetHairGroupsInterpolation()[GroupIndex];
 		}
 	}
 
