Fix incorrect loading of groom asset with only cards or meshes.

diff --git a/Engine/Plugins/Runtime/HairStrands/Source/HairStrandsCore/Private/GroomAsset.cpp b/Engine/Plugins/Runtime/HairStrands/Source/HairStrandsCore/Private/GroomAsset.cpp
--- a/Engine/Plugins/Runtime/HairStrands/Source/HairStrandsCore/Private/GroomAsset.cpp
+++ b/Engine/Plugins/Runtime/HairStrands/Source/HairStrandsCore/Private/GroomAsset.cpp
@@ -1215,7 +1215,7 @@ void UGroomAsset::PostLoad()
 
 	// Legacy asset are missing rendering or interpolation settings
 #if WITH_EDITORONLY_DATA
-	const bool bIsLegacyAsset = GetHairGroupsInfo().Num() == 0;
+	const bool bIsLegacyAsset = GetHairGroupsInterpolation().Num() == 0;
 	if (bIsLegacyAsset)
 	{
 		if (HairDescriptionBulkData[HairDescriptionType])
diff --git a/Engine/Plugins/Runtime/HairStrands/Source/HairStrandsCore/Public/GroomAsset.h b/Engine/Plugins/Runtime/HairStrands/Source/HairStrandsCore/Public/GroomAsset.h
--- a/Engine/Plugins/Runtime/HairStrands/Source/HairStrandsCore/Public/GroomAsset.h
+++ b/Engine/Plugins/Runtime/HairStrands/Source/HairStrandsCore/Public/GroomAsset.h
@@ -429,7 +429,7 @@ class HAIRSTRANDSCORE_API UGroomAsset : public UObject, public IInterface_AssetU
 	DECLARE_MULTICAST_DELEGATE(FOnGroomAsyncLoadFinished);
 #endif
 
-private:
+public:
 	UE_DEPRECATED(5.3, "Please do not access this member directly; use UGroomAsset accessor.")
 	UPROPERTY(EditAnywhere, Transient, EditFixedSize, Category = "HairInfo", meta = (DisplayName = "Group"))
 	TArray<FHairGroupInfoWithVisibility> HairGroupsInfo;
