Change how cluster data are computed.
* Curve decimation is now done continuousely for groom with LODs
* Curve coverage compensation is now done globally for the groom rather than per cluster.

diff --git a/Engine/Plugins/Runtime/HairStrands/Source/HairStrandsCore/Private/GroomAsset.cpp b/Engine/Plugins/Runtime/HairStrands/Source/HairStrandsCore/Private/GroomAsset.cpp
--- a/Engine/Plugins/Runtime/HairStrands/Source/HairStrandsCore/Private/GroomAsset.cpp
+++ b/Engine/Plugins/Runtime/HairStrands/Source/HairStrandsCore/Private/GroomAsset.cpp
@@ -125,8 +125,6 @@ uint32 FHairGroupPlatformData::FStrands::GetDataSize() const
 	Total += FBaseWithInterpolation::GetDataSize();
 	uint32 CTotal = 0;
 	uint32 BTotal = 0;
-	CTotal += ClusterCullingBulkData.Header.LODVisibility.GetAllocatedSize();
-	CTotal += ClusterCullingBulkData.Header.CPULODScreenSize.GetAllocatedSize();
 	BTotal += ClusterCullingBulkData.Data.CurveToClusterIds.IsBulkDataLoaded()	? ClusterCullingBulkData.Data.CurveToClusterIds.GetBulkDataSize() : 0;
 	CTotal += ClusterCullingBulkData.Data.PackedClusterInfos.IsBulkDataLoaded()	? ClusterCullingBulkData.Data.PackedClusterInfos.GetBulkDataSize() : 0;
 	CTotal += ClusterCullingBulkData.Data.PointLODs.IsBulkDataLoaded()	? ClusterCullingBulkData.Data.PointLODs.GetBulkDataSize() : 0;
@@ -1569,8 +1567,6 @@ static bool IsStrandsLODAttributes(const FName PropertyName)
 		PropertyName == UGroomAsset::GetHairGroupsLODMemberName()
 
 		// LOD groups needs for cluster rebuilding
-		|| PropertyName == GET_MEMBER_NAME_CHECKED(FHairGroupsLOD, ClusterWorldSize)
-		|| PropertyName == GET_MEMBER_NAME_CHECKED(FHairGroupsLOD, ClusterScreenSizeScale)
 		|| PropertyName == GET_MEMBER_NAME_CHECKED(FHairGroupsLOD, LODs)
 
 		// LOD settings needs for cluster rebuilding
diff --git a/Engine/Plugins/Runtime/HairStrands/Source/HairStrandsCore/Private/GroomAssetInterpolation.cpp b/Engine/Plugins/Runtime/HairStrands/Source/HairStrandsCore/Private/GroomAssetInterpolation.cpp
--- a/Engine/Plugins/Runtime/HairStrands/Source/HairStrandsCore/Private/GroomAssetInterpolation.cpp
+++ b/Engine/Plugins/Runtime/HairStrands/Source/HairStrandsCore/Private/GroomAssetInterpolation.cpp
@@ -36,12 +36,6 @@ FHairGroupsInterpolation::FHairGroupsInterpolation()
 	InterpolationSettings = FHairInterpolationSettings();
 }
 
-FHairGroupsLOD::FHairGroupsLOD()
-{
-	ClusterWorldSize = 1; // 1cm diameter
-	ClusterScreenSizeScale = 1;
-}
-
 bool FHairLODSettings::operator==(const FHairLODSettings& A) const
 {
 	return
@@ -68,9 +62,7 @@ bool FHairGroupsLOD::operator==(const FHairGroupsLOD& A) const
 		}
 	}
 
-	return
-		ClusterWorldSize == A.ClusterWorldSize &&
-		ClusterScreenSizeScale == A.ClusterScreenSizeScale;
+	return true;
 }
 
 
@@ -117,9 +109,6 @@ void FHairGroupsInterpolation::BuildDDCKey(FArchive& Ar)
 
 void FHairGroupsLOD::BuildDDCKey(FArchive& Ar)
 {
-	Ar << ClusterWorldSize;
-	Ar << ClusterScreenSizeScale;
-
 	for (FHairLODSettings& LOD : LODs)
 	{
 		Ar << LOD.CurveDecimation;
diff --git a/Engine/Plugins/Runtime/HairStrands/Source/HairStrandsCore/Private/GroomBuilder.cpp b/Engine/Plugins/Runtime/HairStrands/Source/HairStrandsCore/Private/GroomBuilder.cpp
--- a/Engine/Plugins/Runtime/HairStrands/Source/HairStrandsCore/Private/GroomBuilder.cpp
+++ b/Engine/Plugins/Runtime/HairStrands/Source/HairStrandsCore/Private/GroomBuilder.cpp
@@ -40,7 +40,7 @@ static FAutoConsoleVariableRef CVarHairGroupIndexBuilder_MaxVoxelResolution(TEXT
 
 FString FGroomBuilder::GetVersion()
 {
-	return TEXT("v8r37_a");
+	return TEXT("v8r39");
 }
 
 namespace FHairStrandsDecimation
@@ -2782,7 +2782,7 @@ struct FClusterGrid
 	TArray<FCluster> Clusters;
 };
 
-static void DecimateCurve(
+static void DecimateCurveLOD(
 	const TArray<FVector3f>& InPoints,
 	const uint32 InOffset,
 	const uint32 InCount,
@@ -2790,6 +2790,8 @@ static void DecimateCurve(
 	uint32* OutCountPerLOD,
 	TArray<uint8>& OutVertexLODMask)
 {
+	OutVertexLODMask.SetNum(InCount);
+
 	// Insure that all settings have more and more agressive, and rectify it is not the case.
 	TArray<FHairLODSettings> Settings = InSettings;
 	{
@@ -2879,15 +2881,14 @@ static void DecimateCurve(
 		// For all remaining vertices, we mark them as 'used'/'valid' for the current LOD levl
 		for (uint32 LocalIndex : OutIndices)
 		{
-			const uint32 VertexIndex = InOffset + LocalIndex;
-			OutVertexLODMask[VertexIndex] |= 1 << LODIt;
+			OutVertexLODMask[LocalIndex] |= 1 << LODIt;
 		}
 	}
 
 	// Sanity check to insure that vertex LOD in a continuous fashion.
-	for (uint32 VertexIt = 0; VertexIt < InCount; ++VertexIt)
+	for (uint32 PointIt = 0; PointIt < InCount; ++PointIt)
 	{
-		const uint8 Mask = OutVertexLODMask[InOffset + VertexIt];
+		const uint8 Mask = OutVertexLODMask[PointIt];
 		check(Mask == 0 || Mask == 1 || Mask == 3 || Mask == 7 || Mask == 15 || Mask == 31 || Mask == 63 || Mask == 127 || Mask == 255);
 	}
 }
@@ -2915,308 +2916,152 @@ static void BuildClusterData(
 	Out.CurveCount = InRenStrandsData.GetNumCurves();
 	check(Out.PointCount > 0 && Out.CurveCount > 0);
 
-	// 1. Allocate cluster per voxel containing contains >=1 render curve root
-	const FVector GroupMinBound = InRenStrandsData.BoundingBox.Min;
-	FVector GroupMaxBound = InRenStrandsData.BoundingBox.Max;
-	const float GroupRadius = FVector::Distance(GroupMaxBound, GroupMinBound) * 0.5f;
-
-	// Compute the voxel volume resolution, and snap the max bound to the voxel grid
-	// Iterate until voxel size are below max resolution, so that computation is not too long
-	FIntVector VoxelResolution = FIntVector::ZeroValue;
+	// 1. Compute the number of curves per LOD
+	Out.LODInfos.SetNum(LODCount);
+	for (uint8 LODIt = 0; LODIt < LODCount; ++LODIt)
 	{
-		const int32 MaxResolution = FMath::Max(GHairClusterBuilder_MaxVoxelResolution, 2);
-
-		float ClusterWorldSize = FMath::Max(InSettings.ClusterWorldSize, 0.001f);
-		bool bIsValid = false;
-		while (!bIsValid)
-		{
-			FVector VoxelResolutionF = (GroupMaxBound - GroupMinBound) / ClusterWorldSize;
-			VoxelResolution = FIntVector(FMath::CeilToInt(VoxelResolutionF.X), FMath::CeilToInt(VoxelResolutionF.Y), FMath::CeilToInt(VoxelResolutionF.Z));
-			bIsValid = VoxelResolution.X <= MaxResolution && VoxelResolution.Y <= MaxResolution && VoxelResolution.Z <= MaxResolution;
-			if (!bIsValid)
-			{
-				ClusterWorldSize *= 2;
-			}
-		}
-		GroupMaxBound = GroupMinBound + FVector(VoxelResolution) * ClusterWorldSize;
+		const float CurveDecimation = FMath::Clamp(InSettings.LODs[LODIt].CurveDecimation, 0.0f, 1.0f);
+		Out.LODInfos[LODIt].CurveCount = FMath::Clamp(FMath::CeilToInt(Out.CurveCount * CurveDecimation), 1, InRenStrandsData.GetNumCurves());
+		Out.LODInfos[LODIt].PointCount = 0;
+		Out.LODInfos[LODIt].RadiusScale= 1;
+		Out.LODInfos[LODIt].ScreenSize = InSettings.LODs[LODIt].ScreenSize;
+		Out.LODInfos[LODIt].bIsVisible = InSettings.LODs[LODIt].bVisible;
 	}
 
-	// 2. Insert all rendering curves into the voxel structure
-	FClusterGrid ClusterGrid(VoxelResolution, GroupMinBound, GroupMaxBound);
-	for (uint32 RenCurveIndex = 0; RenCurveIndex < RenCurveCount; ++RenCurveIndex)
+	// 2. Compute control point LOD
+	for (uint32 CurveIt = 0; CurveIt < Out.CurveCount; ++CurveIt)
 	{
-		FClusterGrid::FCurve RCurve;
-		RCurve.CurveIndex = RenCurveIndex;
-		RCurve.Count = InRenStrandsData.StrandsCurves.CurvesCount[RenCurveIndex];
-		RCurve.Offset = InRenStrandsData.StrandsCurves.CurvesOffset[RenCurveIndex];
-		RCurve.Area = 0.0f;
-		RCurve.AvgRadius = 0;
-		RCurve.MaxRadius = 0;
+		const uint32 PointOffset = InRenStrandsData.StrandsCurves.CurvesOffset[CurveIt];
+		const uint32 PointCount  = InRenStrandsData.StrandsCurves.CurvesCount[CurveIt];
 
-		// Compute area of each curve to later compute area correction
-		for (uint32 RenPointIndex = 0; RenPointIndex < RCurve.Count; ++RenPointIndex)
+		uint32 CountPerLOD[FHairClusterInfo::MaxLOD];
+
+		// 2.1 Write out cluster information
+		// Don't need to DecimateCurve if there are only 2 control vertices
+		TArray<uint8> PointLODMasks;
+		PointLODMasks.SetNum(PointCount);
+		if (PointCount > 2)
 		{
-			uint32 PointGlobalIndex = RenPointIndex + RCurve.Offset;
-			const FVector& V0 = (FVector)InRenStrandsData.StrandsPoints.PointsPosition[PointGlobalIndex];
-			if (RenPointIndex > 0)
+			DecimateCurveLOD(
+				InRenStrandsData.StrandsPoints.PointsPosition,
+				PointOffset,
+				PointCount,
+				InSettings.LODs,
+				CountPerLOD,
+				PointLODMasks);
+
+			// For stats
+			for (uint32 LODIt = 0; LODIt < LODCount; ++LODIt)
 			{
-				const FVector& V1 = (FVector)InRenStrandsData.StrandsPoints.PointsPosition[PointGlobalIndex - 1];
-				FVector OutDir;
-				float OutLength;
-				(V1 - V0).ToDirectionAndLength(OutDir, OutLength);
-				RCurve.Area += InRenStrandsData.StrandsPoints.PointsRadius[PointGlobalIndex] * OutLength;
+				Out.LODInfos[LODIt].PointCount += CountPerLOD[LODIt];
 			}
-
-			const float PointRadius = InRenStrandsData.StrandsPoints.PointsRadius[PointGlobalIndex];
-			RCurve.AvgRadius += PointRadius;
-			RCurve.MaxRadius = FMath::Max(RCurve.MaxRadius, PointRadius);
 		}
-		RCurve.AvgRadius /= FMath::Max(1u, RCurve.Count);
 
-		const FVector Root = (FVector)InRenStrandsData.StrandsPoints.PointsPosition[RCurve.Offset];
-		ClusterGrid.InsertRenderingCurve(RCurve, Root);
-	}
-
-	// 3. Count non-empty clusters
-	TArray<uint32> ValidClusterIndices;
-	{
-		uint32 GridLinearIndex = 0;
-		ValidClusterIndices.Reserve(ClusterGrid.Clusters.Num() * 0.2);
-		for (FClusterGrid::FCluster& Cluster : ClusterGrid.Clusters)
+		// 2.2 Write LOD mask per control point
+		Out.PointLODs.SetNum(InRenStrandsData.GetNumPoints());
+		for (uint8 LODIt = 0; LODIt < LODCount; ++LODIt)
 		{
-			if (Cluster.ClusterCurves.Num() > 0)
+			for (uint32 PointIt = 0; PointIt < PointCount; ++PointIt)
 			{
-				ValidClusterIndices.Add(GridLinearIndex);
+				const uint8 LODMask = PointLODMasks[PointIt];
+				if (LODMask & (1 << LODIt))
+				{
+					// Store the minimum/coarser LOD at which the current point is visible
+					const uint32 GlobalPointIndex = PointOffset + PointIt;
+					Out.PointLODs[GlobalPointIndex] = LODIt;
+				}
 			}
-			++GridLinearIndex;
 		}
 	}
-	Out.ClusterCount = ValidClusterIndices.Num();
-	Out.ClusterInfos.Init(FHairClusterInfo(), Out.ClusterCount);
-	Out.CurveToClusterIds.SetNum(Out.CurveCount);
-	Out.LODInfos.SetNum(LODCount);
-	Out.PointLODs.Init(0, Out.PointCount);
-
-	// 4. Write out cluster information
-	TArray<uint8> VertexLODMasks;
-	VertexLODMasks.SetNum(InRenStrandsData.GetNumPoints());
-
-	// Local variable for being capture by the lambda
-	TArray<FHairClusterInfo>& LocalClusterInfos = Out.ClusterInfos;
-	TArray<uint32>& LocalCurveToClusterIds = Out.CurveToClusterIds;
-	TArray<uint32>& LocalPointLODs = Out.PointLODs;
-#define USE_PARALLE_FOR 0
-#if USE_PARALLE_FOR
-	ParallelFor(Out.ClusterCount,
-		[
-			LODCount,
-			InGroomAssetRadius,
-			InSettings,
-			&ValidClusterIndices,
-			&InRenStrandsData,
-			&ClusterGrid,
-			&LocalClusterInfos,
-			&LocalCurveToClusterIds,
-			&VertexLODMasks,
-		]
-	(uint32 ClusterIt)
-#else
-	for (uint32 ClusterIt = 0; ClusterIt < Out.ClusterCount; ++ClusterIt)
-#endif
-	{
-		const uint32 GridLinearIndex = ValidClusterIndices[ClusterIt];
-		FClusterGrid::FCluster& Cluster = ClusterGrid.Clusters[GridLinearIndex];
-		check(Cluster.ClusterCurves.Num() != 0);
-
-		// 4.1 Sort curves
-		// Sort curve to have largest area first, so that lower area curves with less influence are removed first.
-		// This also helps the radius scaling to not explode.
-		Cluster.ClusterCurves.Sort([](const FClusterGrid::FCurve& A, const FClusterGrid::FCurve& B) -> bool
-		{
-			return A.Area > B.Area;
-		});
 
-		// 4.2 Compute cluster's area & fill in the vertex to cluster ID mapping
-		float ClusterArea = 0;
-		FVector ClusterMinBound(FLT_MAX);
-		FVector ClusterMaxBound(-FLT_MAX);
-
-		FVector RootMinBound(FLT_MAX);
-		FVector RootMaxBound(-FLT_MAX);
-
-		Cluster.CurveMaxRadius = 0;
-		Cluster.CurveAvgRadius = 0;
-		Cluster.Area = 0;
-		for (FClusterGrid::FCurve& ClusterCurve : Cluster.ClusterCurves)
+	// 3. Compute radius scale per LOD
+	{
+		struct FCurveInfo
 		{
-			check(ClusterCurve.CurveIndex < uint32(LocalCurveToClusterIds.Num()));
-			LocalCurveToClusterIds[ClusterCurve.CurveIndex] = ClusterIt;
-
-			for (uint32 RenPointIndex = 0; RenPointIndex < ClusterCurve.Count; ++RenPointIndex)
-			{
-				const uint32 PointGlobalIndex = RenPointIndex + ClusterCurve.Offset;
-
-				const FVector& P = (FVector)InRenStrandsData.StrandsPoints.PointsPosition[PointGlobalIndex];
+			float Area = 0;
+			float MinRadius = FLT_MAX;
+			float MaxRadius = 0;
+			float AvgRadius = 0;
+		};
+		// Compute curves area
+		TArray<FCurveInfo> CurveInfos;
+		CurveInfos.SetNum(RenCurveCount);
+		float AvgCurveArea = 0;
+		for (uint32 CurveIt = 0; CurveIt < Out.CurveCount; ++CurveIt)
+		{
+			FCurveInfo& CurveInfo = CurveInfos[CurveIt];
+			const uint32 PointOffset = InRenStrandsData.StrandsCurves.CurvesOffset[CurveIt];
+			const uint32 PointCount  = InRenStrandsData.StrandsCurves.CurvesCount[CurveIt];
+
+			// Compute area of each curve to later compute area correction
+			for (uint32 PointIt = 0; PointIt < PointCount; ++PointIt)
+			{
+				uint32 PointGlobalIndex = PointIt + PointOffset;
+				const FVector& V0 = (FVector)InRenStrandsData.StrandsPoints.PointsPosition[PointGlobalIndex];
+				if (PointIt > 0)
 				{
-					ClusterMinBound.X = FMath::Min(ClusterMinBound.X, P.X);
-					ClusterMinBound.Y = FMath::Min(ClusterMinBound.Y, P.Y);
-					ClusterMinBound.Z = FMath::Min(ClusterMinBound.Z, P.Z);
-
-					ClusterMaxBound.X = FMath::Max(ClusterMaxBound.X, P.X);
-					ClusterMaxBound.Y = FMath::Max(ClusterMaxBound.Y, P.Y);
-					ClusterMaxBound.Z = FMath::Max(ClusterMaxBound.Z, P.Z);
+					const FVector& V1 = (FVector)InRenStrandsData.StrandsPoints.PointsPosition[PointGlobalIndex - 1];
+					FVector OutDir;
+					float OutLength;
+					(V1 - V0).ToDirectionAndLength(OutDir, OutLength);
+					CurveInfo.Area += InRenStrandsData.StrandsPoints.PointsRadius[PointGlobalIndex] * OutLength;
 				}
 
-				if (RenPointIndex == 0)
-				{
-					RootMinBound.X = FMath::Min(RootMinBound.X, P.X);
-					RootMinBound.Y = FMath::Min(RootMinBound.Y, P.Y);
-					RootMinBound.Z = FMath::Min(RootMinBound.Z, P.Z);
-
-					RootMaxBound.X = FMath::Max(RootMaxBound.X, P.X);
-					RootMaxBound.Y = FMath::Max(RootMaxBound.Y, P.Y);
-					RootMaxBound.Z = FMath::Max(RootMaxBound.Z, P.Z);
-				}
+				const float PointRadius = InRenStrandsData.StrandsPoints.PointsRadius[PointGlobalIndex];
+				CurveInfo.AvgRadius += PointRadius;
+				CurveInfo.MaxRadius = FMath::Max(CurveInfo.MaxRadius, PointRadius);
+				CurveInfo.MinRadius = FMath::Max(CurveInfo.MinRadius, PointRadius);
 			}
-			Cluster.CurveMaxRadius = FMath::Max(Cluster.CurveMaxRadius, ClusterCurve.MaxRadius);
-			Cluster.CurveAvgRadius += ClusterCurve.AvgRadius;
-			Cluster.Area += ClusterCurve.Area;
-		}
-		Cluster.CurveAvgRadius /= FMath::Max(1, Cluster.ClusterCurves.Num());
-		Cluster.RootBoundRadius = (RootMaxBound - RootMinBound).GetMax() * 0.5f + Cluster.CurveAvgRadius;
+			CurveInfo.AvgRadius /= FMath::Max(1u, PointCount);
 
-		// Compute the max radius that a cluster can have. This is done by computing an estimate of the cluster coverage (using pre-computed LUT) 
-		// and computing how much is visible
-		// This supposes the radius is proportional to the radius of the roots bounding volume
-		const float NormalizedAvgRadius = Cluster.CurveAvgRadius / Cluster.RootBoundRadius;
-		const float ClusterCoverage = GetHairCoverage(Cluster.ClusterCurves.Num(), NormalizedAvgRadius);
-		const float ClusterVisibleRadius = Cluster.RootBoundRadius * ClusterCoverage;
+			AvgCurveArea += CurveInfo.Area;
+		}
+		AvgCurveArea /= FMath::Max(1u, Out.CurveCount);
 
-		const float ClusterRadius = FVector::Distance(ClusterMaxBound, ClusterMinBound) * 0.5f;
+		const FVector3f GroupMinBound(InRenStrandsData.BoundingBox.Min);
+		const FVector3f GroupMaxBound(InRenStrandsData.BoundingBox.Max);
+		const FVector3f GroupSize = GroupMaxBound - GroupMinBound;
+		const FVector3f Area3(GroupSize.X * GroupSize.Y, GroupSize.X * GroupSize.Z, GroupSize.Y * GroupSize.Z);
+		const float MaxArea = FMath::Max3(Area3.X, Area3.Y, Area3.Z);
+		const float AvgArea = (Area3.X + Area3.Y + Area3.Z) / 3.f;
+		//		const float MaxArea = FMath::Max3(GroupSize.X, GroupSize.Y, GroupSize.Z);
 
+		// Reference: "Stochastic Simplification of Aggregate Detail"
+		const float DepthComplexity = Out.CurveCount * AvgCurveArea / MaxArea;
+		const float r = DepthComplexity / Out.CurveCount;
 
-		// 4.3 Compute the number of curve per LOD
-		// Compute LOD infos (vertx count, vertex offset, radius scale ...)
-		// Compute the ratio of the cluster related the actual groom and scale the screen size accordingly
-		TArray<uint32> LODCurveCount;
-		LODCurveCount.SetNum(LODCount);
 		for (uint8 LODIt = 0; LODIt < LODCount; ++LODIt)
 		{
-			const float CurveDecimation = FMath::Clamp(InSettings.LODs[LODIt].CurveDecimation, 0.0f, 1.0f);
-			LODCurveCount[LODIt] = FMath::Clamp(FMath::CeilToInt(Cluster.ClusterCurves.Num() * CurveDecimation), 1, Cluster.ClusterCurves.Num());
-
-			// For stats
-			Out.LODInfos[LODIt].CurveCount += LODCurveCount[LODIt];
+			const float Lambda = InSettings.LODs[LODIt].CurveDecimation;
+			Out.LODInfos[LODIt].RadiusScale  = (1.f - FMath::Pow(1.f - r, 1.f / Lambda)) / r;
+			Out.LODInfos[LODIt].RadiusScale *= InSettings.LODs[LODIt].ThicknessScale;
 		}
+	}
 
-		// 4.4 Decimate each curve for all LODs
-		// This fill in a bitfield per vertex which indiates on which LODs a vertex can be used
-		for (uint32 CurveIt = 0, CurveCount = uint32(Cluster.ClusterCurves.Num()); CurveIt < CurveCount; ++CurveIt)
-		{
-			FClusterGrid::FCurve& ClusterCurve = Cluster.ClusterCurves[CurveIt];
-
-			// Don't need to DecimateCurve if there are only 2 control vertices
-			if (ClusterCurve.Count > 2)
-			{
-				DecimateCurve(
-					InRenStrandsData.StrandsPoints.PointsPosition,
-					ClusterCurve.Offset,
-					ClusterCurve.Count,
-					InSettings.LODs,
-					ClusterCurve.CountPerLOD,
-					VertexLODMasks);
-
-				// For stats
-				for (uint32 LODIt = 0; LODIt < LODCount; ++LODIt)
-				{
-					Out.LODInfos[LODIt].PointCount += ClusterCurve.CountPerLOD[LODIt];
-				}
-			}
-		}
-
-		// 4.5 Record/Insert vertex min LOD for each curve
-
-		FHairClusterInfo& ClusterInfo = LocalClusterInfos[ClusterIt];
-		ClusterInfo.LODCount = LODCount;
-		for (uint8 LODIt = 0; LODIt < LODCount; ++LODIt)
+	// 4. Fill in cluster Info for legacy compatibility
+	Out.ClusterCount = FMath::Min(FMath::DivideAndRoundUp(InRenStrandsData.GetNumCurves(), 100u), 1024u);
+	Out.CurveToClusterIds.SetNum(InRenStrandsData.GetNumCurves());
+	Out.ClusterInfos.SetNum(Out.ClusterCount);
+	const uint32 CurvePerCluster = FMath::DivideAndRoundUp(InRenStrandsData.GetNumCurves(), Out.ClusterCount);
+	for (uint32 ClusterIt = 0; ClusterIt < Out.ClusterCount; ++ClusterIt)
+	{
+		// Every 'CurvePerCluster' is assigned to the current cluster
+		for (uint32 CurveIt = 0; CurveIt < CurvePerCluster; ++CurveIt)
 		{
-			ClusterInfo.ScreenSize[LODIt] = 1.f;
-			ClusterInfo.RadiusScale[LODIt]= 1.f;
-			ClusterInfo.bIsVisible[LODIt] = false;
-
-			const uint32 CurveCount = LODCurveCount[LODIt];
-			const uint32 NextCurveCount = LODIt < LODCount - 1 ? LODCurveCount[LODIt + 1] : CurveCount;
-			for (uint32 CurveIt = 0; CurveIt < CurveCount; ++CurveIt)
+			const uint32 CurveIndex = CurveIt * CurvePerCluster;
+			if (CurveIndex < InRenStrandsData.GetNumCurves())
 			{
-				FClusterGrid::FCurve& ClusterCurve = Cluster.ClusterCurves[CurveIt];
-
-				for (uint32 PointIt = 0; PointIt < ClusterCurve.Count; ++PointIt)
-				{
-					const uint32 GlobalPointIndex = PointIt + ClusterCurve.Offset;
-					const uint8 LODMask = VertexLODMasks[GlobalPointIndex];
-					if (LODMask & (1 << LODIt))
-					{
-						// Store the minimum/coarser LOD at which the current point is visible
-						LocalPointLODs[GlobalPointIndex] = LODIt;
-					}
-				}
+				Out.CurveToClusterIds[CurveIndex] = ClusterIt;
 			}
 		}
-
-		// 4.6 Compute the radius scaling to preserve the cluster apperance as we decimate 
-		// the number of strands
+		Out.ClusterInfos[ClusterIt].LODCount = LODCount;
 		for (uint8 LODIt = 0; LODIt < LODCount; ++LODIt)
 		{
-			// Compute the visible area for various orientation? 
-			// Reference: Stochastic Simplification of Aggregate Detail
-			float  LODAvgRadiusRef = 0;
-			float  LODMaxRadiusRef = 0;
-
-			const uint32 ClusterCurveCount = LODCurveCount[LODIt];
-			for (uint32 CurveIt = 0; CurveIt < ClusterCurveCount; ++CurveIt)
-			{
-				const FClusterGrid::FCurve& ClusterCurve = Cluster.ClusterCurves[CurveIt];
-				LODAvgRadiusRef += ClusterCurve.AvgRadius;
-				LODMaxRadiusRef = FMath::Max(LODMaxRadiusRef, ClusterCurve.MaxRadius);
-			}
-			LODAvgRadiusRef /= ClusterCurveCount;
-
-			// Compute what should be the average (normalized) radius of the strands, and scale it 
-			// with the radius of the clusters/roots to get an actual world radius.
-			const float LODAvgRadiusTarget = Cluster.RootBoundRadius * GetHairAvgRadius(ClusterCurveCount, ClusterCoverage);
-
-			// Compute the ratio between the size of the cluster and the size of the groom (at rest position)
-			// On the GPU, we compute the screen size of the cluster, and use the LOD screensize to know which 
-			// LOD needs to be pick up. Since the screen area are setup by artists based the entire groom (not 
-			// based on the cluster size), we precompute the correcting ratio here, and pre-scale the LOD screensize
-			const float ScreenSizeScale = InSettings.ClusterScreenSizeScale * ClusterRadius / InGroomAssetRadius;
-
-			float LODScale = LODAvgRadiusTarget / LODAvgRadiusRef;
-			if (LODMaxRadiusRef * LODScale > ClusterVisibleRadius)
-			{
-				LODScale = FMath::Max(LODMaxRadiusRef, ClusterVisibleRadius) / LODMaxRadiusRef;
-			}
-			LODScale *= FMath::Max(InSettings.LODs[LODIt].ThicknessScale, 0.f);
-			//if (LODMaxRadiusRef * LODScale > Cluster.RootBoundRadius)
-			//{
-			//	LODScale = Cluster.RootBoundRadius / LODMaxRadiusRef;
-			//}
-
-			ClusterInfo.ScreenSize[LODIt] = InSettings.LODs[LODIt].ScreenSize * ScreenSizeScale;
-			ClusterInfo.bIsVisible[LODIt] = InSettings.LODs[LODIt].bVisible;
-			ClusterInfo.RadiusScale[LODIt]= LODScale;
+			Out.ClusterInfos[ClusterIt].ScreenSize[LODIt]  = Out.LODInfos[LODIt].ScreenSize;
+			Out.ClusterInfos[ClusterIt].RadiusScale[LODIt] = Out.LODInfos[LODIt].RadiusScale;
+			Out.ClusterInfos[ClusterIt].bIsVisible[LODIt]  = Out.LODInfos[LODIt].bIsVisible;
 		}
 	}
-#if USE_PARALLE_FOR
-	);
-#endif
-
-	// Compute the screen size of the entire group at which the groom need to change LOD
-	for (uint32 LODIt = 0; LODIt < LODCount; ++LODIt)
-	{
-		Out.CPULODScreenSize.Add(InSettings.LODs[LODIt].ScreenSize);
-		Out.LODVisibility.Add(InSettings.LODs[LODIt].bVisible);
-	}
 }
 
 static void BuildClusterBulkData(
@@ -3224,13 +3069,9 @@ static void BuildClusterBulkData(
 	FHairStrandsClusterCullingBulkData& Out)
 {
 	Out.Reset();
-
 	Out.Header.ClusterCount		= In.ClusterCount;
 	Out.Header.PointCount		= In.PointCount;
 	Out.Header.CurveCount		= In.CurveCount;
-
-	Out.Header.CPULODScreenSize	= In.CPULODScreenSize;
-	Out.Header.LODVisibility	= In.LODVisibility;
 	Out.Header.LODInfos			= In.LODInfos;
 
 	// Sanity check
diff --git a/Engine/Plugins/Runtime/HairStrands/Source/HairStrandsCore/Private/GroomComponent.cpp b/Engine/Plugins/Runtime/HairStrands/Source/HairStrandsCore/Private/GroomComponent.cpp
--- a/Engine/Plugins/Runtime/HairStrands/Source/HairStrandsCore/Private/GroomComponent.cpp
+++ b/Engine/Plugins/Runtime/HairStrands/Source/HairStrandsCore/Private/GroomComponent.cpp
@@ -2767,21 +2767,11 @@ void UGroomComponent::InitResources(bool bIsBindingReloading)
 			{
 				// This codes assumes strands LOD are contigus and the highest (i.e., 0...x). Change this code to something more robust
 				check(HairGroupInstance->HairGroupPublicData);
-				const int32 StrandsLODCount = GroupData.Strands.ClusterCullingResource->BulkData.Header.CPULODScreenSize.Num();
+				const int32 StrandsLODCount = GroupData.Strands.ClusterCullingResource->BulkData.Header.LODInfos.Num();
 				const TArray<float>& LODScreenSizes = HairGroupInstance->HairGroupPublicData->GetLODScreenSizes();
 				const TArray<bool>& LODVisibilities = HairGroupInstance->HairGroupPublicData->GetLODVisibilities();
 				check(StrandsLODCount <= LODScreenSizes.Num());
 				check(StrandsLODCount <= LODVisibilities.Num());
-				for (int32 LODIt = 0; LODIt < StrandsLODCount; ++LODIt)
-				{
-					// ClusterCullingData seriazlizes only the screen size related to strands geometry.
-					// Other type of geometry are not serizalized, and so won't match
-					if (GroomAsset->GetHairGroupsLOD()[GroupIt].LODs[LODIt].GeometryType == EGroomGeometryType::Strands)
-					{
-						check(GroupData.Strands.ClusterCullingResource->BulkData.Header.CPULODScreenSize[LODIt] == LODScreenSizes[LODIt]);
-						check(GroupData.Strands.ClusterCullingResource->BulkData.Header.LODVisibility[LODIt] == LODVisibilities[LODIt]);
-					}
-				}
 				HairGroupInstance->HairGroupPublicData->ClusterCount = HairGroupInstance->Strands.ClusterCullingResource->BulkData.Header.ClusterCount;
 				BeginInitResource(HairGroupInstance->HairGroupPublicData);
 			}
diff --git a/Engine/Plugins/Runtime/HairStrands/Source/HairStrandsCore/Private/GroomDebug.cpp b/Engine/Plugins/Runtime/HairStrands/Source/HairStrandsCore/Private/GroomDebug.cpp
--- a/Engine/Plugins/Runtime/HairStrands/Source/HairStrandsCore/Private/GroomDebug.cpp
+++ b/Engine/Plugins/Runtime/HairStrands/Source/HairStrandsCore/Private/GroomDebug.cpp
@@ -981,6 +981,7 @@ static void AddHairDebugPrintInstancePass(
 
 				D.Data3.Z |= FFloat16(Instance->HairGroupPublicData->VFInput.Strands.Common.Length).Encoded;
 				D.Data3.Z |= FFloat16(Instance->HairGroupPublicData->VFInput.Strands.Common.LengthScale).Encoded << 16u;
+				D.Data3.W  = Instance->HairGroupPublicData->GetActiveStrandsCoverageScale();
 			}
 			break;
 		case EHairGeometryType::Cards:
diff --git a/Engine/Plugins/Runtime/HairStrands/Source/HairStrandsCore/Private/GroomManager.cpp b/Engine/Plugins/Runtime/HairStrands/Source/HairStrandsCore/Private/GroomManager.cpp
--- a/Engine/Plugins/Runtime/HairStrands/Source/HairStrandsCore/Private/GroomManager.cpp
+++ b/Engine/Plugins/Runtime/HairStrands/Source/HairStrandsCore/Private/GroomManager.cpp
@@ -888,6 +888,7 @@ static void RunHairLODSelection(
 			Instance->HairGroupPublicData->ContinuousLODScreenSize = MaxScreenSize_RestBound;
 			Instance->HairGroupPublicData->ContinuousLODScreenPos = MaxContinuousLODScreenPos;
 			Instance->HairGroupPublicData->ContinuousLODBounds = SphereBound;
+			Instance->HairGroupPublicData->ContinuousLODCoverageScale = 1.f;
 
 			if (IsHairStrandsContinousLODEnabled())
 			{
@@ -900,6 +901,25 @@ static void RunHairLODSelection(
 				Instance->HairGroupPublicData->ContinuousLODPointCount = GetHairVisibilityComputeRasterPointCount(MaxScreenSize_Bound, Instance->HairGroupPublicData->RestPointCount);
 				Instance->HairGroupPublicData->ContinuousLODCurveCount = Instance->HairGroupPublicData->RestCurveCount;
 			}
+			else if (Instance->Strands.ClusterCullingResource)
+			{
+				const TArray<FHairLODInfo>& LODInfos = Instance->Strands.ClusterCullingResource->BulkData.Header.LODInfos;
+
+				const int32 iLODIndex = LODIndex;
+				const float S = LODIndex - iLODIndex;
+
+				const uint32 LODIndex0 = FMath::Min(iLODIndex,   LODInfos.Num()-1);
+				const uint32 LODIndex1 = FMath::Min(iLODIndex+1, LODInfos.Num()-1);
+				const uint32 EffectiveCurveCount = FMath::Lerp(LODInfos[LODIndex1].CurveCount, LODInfos[LODIndex0].CurveCount, 1.f-S);
+
+				// Compensate lost in curve by a coverage scale increase
+				const float CurveRatio = EffectiveCurveCount / float(Instance->HairGroupPublicData->RestCurveCount);
+
+				check(EffectiveCurveCount <= uint32(Instance->Strands.Data->Header.CurveToPointCount.Num()));
+				Instance->HairGroupPublicData->ContinuousLODCurveCount = EffectiveCurveCount;
+				Instance->HairGroupPublicData->ContinuousLODPointCount = EffectiveCurveCount > 0 ? Instance->Strands.Data->Header.CurveToPointCount[EffectiveCurveCount - 1] : 0;
+				Instance->HairGroupPublicData->ContinuousLODCoverageScale = 1.f/FMath::Max(CurveRatio, 0.01f);
+			}
 		}
 
 		// Function for selecting, loading, & initializing LOD resources
diff --git a/Engine/Plugins/Runtime/HairStrands/Source/HairStrandsCore/Private/HairStrandsClusterCulling.cpp b/Engine/Plugins/Runtime/HairStrands/Source/HairStrandsCore/Private/HairStrandsClusterCulling.cpp
--- a/Engine/Plugins/Runtime/HairStrands/Source/HairStrandsCore/Private/HairStrandsClusterCulling.cpp
+++ b/Engine/Plugins/Runtime/HairStrands/Source/HairStrandsCore/Private/HairStrandsClusterCulling.cpp
@@ -65,6 +65,7 @@ class FHairClusterCullCS: public FGlobalShader
 		SHADER_PARAMETER_RDG_BUFFER_UAV(RWBuffer<float>, OutRadiusScaleBuffer)
 		SHADER_PARAMETER_RDG_BUFFER_UAV(RWBuffer<uint>, OutCulledCurveBuffer)
 
+		SHADER_PARAMETER_STRUCT_REF(FViewUniformShaderParameters, ViewUniformBuffer)
 		SHADER_PARAMETER_STRUCT_INCLUDE(ShaderPrint::FShaderParameters, ShaderPrintUniformBuffer)
 		END_SHADER_PARAMETER_STRUCT()
 
@@ -113,6 +114,7 @@ static void AddClusterCullingPass(
 	const FShaderPrintData* ShaderPrintData,
 	FHairStrandClusterData::FHairGroup& ClusterData)
 {
+	check(View);
 
 	ShaderPrint::SetEnabled(true);
 	ShaderPrint::RequestSpaceForCharacters(2048);
@@ -134,7 +136,7 @@ static void AddClusterCullingPass(
 		Parameters->PointLODBuffer 		 = RegisterAsSRV(GraphBuilder, *ClusterData.PointLODBuffer);
 		Parameters->ClusterInfoBuffer 	 = RegisterAsSRV(GraphBuilder, *ClusterData.ClusterInfoBuffer);
 		Parameters->ClusterInfoParameters= ClusterData.ClusterInfoParameters;
-
+		Parameters->ViewUniformBuffer	 = View->ViewUniformBuffer;
 		Parameters->OutPointCounter 	 = PointCounterUAV;
 		Parameters->OutIndexBuffer 		 = RegisterAsUAV(GraphBuilder, *ClusterData.GetCulledVertexIdBuffer());
 		Parameters->OutRadiusScaleBuffer = RegisterAsUAV(GraphBuilder, *ClusterData.GetCulledVertexRadiusScaleBuffer());
diff --git a/Engine/Plugins/Runtime/HairStrands/Source/HairStrandsCore/Private/HairStrandsDatas.cpp b/Engine/Plugins/Runtime/HairStrands/Source/HairStrandsCore/Private/HairStrandsDatas.cpp
--- a/Engine/Plugins/Runtime/HairStrands/Source/HairStrandsCore/Private/HairStrandsDatas.cpp
+++ b/Engine/Plugins/Runtime/HairStrands/Source/HairStrandsCore/Private/HairStrandsDatas.cpp
@@ -657,8 +657,6 @@ void FHairStrandsClusterCullingBulkData::Reset()
 	Header.PointCount = 0;
 	Header.CurveCount = 0;
 	Header.ClusterInfoParameters = FVector4f::Zero();
-	Header.LODVisibility.Empty();
-	Header.CPULODScreenSize.Empty();
 	Header.LODInfos.Empty();
 
 	Data.CurveToClusterIds.RemoveBulkData();
@@ -683,8 +681,6 @@ void FHairStrandsClusterCullingBulkData::SerializeHeader(FArchive& Ar, UObject*
 	Ar << Header.ClusterCount;
 	Ar << Header.PointCount;
 	Ar << Header.CurveCount;
-	Ar << Header.LODVisibility;
-	Ar << Header.CPULODScreenSize;
 	Ar << Header.ClusterInfoParameters;
 	uint32 LODInfosCount = Header.LODInfos.Num();
 	Ar << LODInfosCount;
@@ -696,6 +692,9 @@ void FHairStrandsClusterCullingBulkData::SerializeHeader(FArchive& Ar, UObject*
 	{
 		Ar << Header.LODInfos[It].CurveCount;
 		Ar << Header.LODInfos[It].PointCount;
+		Ar << Header.LODInfos[It].RadiusScale;
+		Ar << Header.LODInfos[It].ScreenSize;
+		Ar << Header.LODInfos[It].bIsVisible;
 	}
 }
 
diff --git a/Engine/Plugins/Runtime/HairStrands/Source/HairStrandsCore/Public/GroomAssetInterpolation.h b/Engine/Plugins/Runtime/HairStrands/Source/HairStrandsCore/Public/GroomAssetInterpolation.h
--- a/Engine/Plugins/Runtime/HairStrands/Source/HairStrandsCore/Public/GroomAssetInterpolation.h
+++ b/Engine/Plugins/Runtime/HairStrands/Source/HairStrandsCore/Public/GroomAssetInterpolation.h
@@ -214,20 +214,10 @@ struct HAIRSTRANDSCORE_API FHairGroupsLOD
 {
 	GENERATED_USTRUCT_BODY()
 
-	FHairGroupsLOD();
-
 	/** LODs  */
 	UPROPERTY(EditAnywhere, Category = "DecimationSettings", meta = (ToolTip = "Reduce the number of hair strands in a uniform manner", ClampMin = "0", UIMin = "0", UIMax = "1.0"))
 	TArray<FHairLODSettings> LODs;
 
-	/** Size of the hair clusters used for LOD decimation */
-	UPROPERTY(EditAnywhere, Category = "DecimationSettings", meta = (ToolTip = "If enable this LOD version will use cards representation"))
-	float ClusterWorldSize;
-
-	/** Size in pixel of a hair cluster when fully visible (screensize is equal to 1) */
-	UPROPERTY(EditAnywhere, Category = "DecimationSettings", meta = (ToolTip = "If enable this LOD version will use cards representation"))
-	float ClusterScreenSizeScale;
-
 	bool operator==(const FHairGroupsLOD& A) const;
 
 	void BuildDDCKey(FArchive& Ar);
diff --git a/Engine/Plugins/Runtime/HairStrands/Source/HairStrandsCore/Public/HairStrandsDatas.h b/Engine/Plugins/Runtime/HairStrands/Source/HairStrandsCore/Public/HairStrandsDatas.h
--- a/Engine/Plugins/Runtime/HairStrands/Source/HairStrandsCore/Public/HairStrandsDatas.h
+++ b/Engine/Plugins/Runtime/HairStrands/Source/HairStrandsCore/Public/HairStrandsDatas.h
@@ -284,6 +284,9 @@ struct FHairLODInfo
 {
 	uint32 CurveCount = 0;
 	uint32 PointCount = 0;
+	float RadiusScale = 1;
+	float ScreenSize = 1;
+	bool bIsVisible = true;
 };
 
 struct FHairClusterInfoFormat
@@ -694,12 +697,6 @@ struct HAIRSTRANDSCORE_API FHairStrandsClusterCullingData
 	void Reset();
 	bool IsValid() const { return ClusterCount > 0 && PointCount > 0; }
 
-	/* Set LOD visibility, allowing to remove the simulation/rendering of certain LOD */
-	TArray<bool>				LODVisibility;
-
-	/* Screen size at which LOD should switches on CPU */
-	TArray<float>				CPULODScreenSize;
-
 	/* LOD info for the various clusters for LOD management on GPU */
 	TArray<FHairClusterInfo>	ClusterInfos;
 	TArray<uint32>				CurveToClusterIds;
@@ -724,12 +721,6 @@ struct HAIRSTRANDSCORE_API FHairStrandsClusterCullingBulkData : FHairStrandsBulk
 
 	struct FHeader
 	{
-		/* Set LOD visibility, allowing to remove the simulation/rendering of certain LOD */
-		TArray<bool> LODVisibility;
-	
-		/* Screen size at which LOD should switches on CPU */
-		TArray<float> CPULODScreenSize;
-	
 		/* Curve count and Point count per LOD */
 		TArray<FHairLODInfo> LODInfos;
 	
diff --git a/Engine/Plugins/Runtime/HairStrands/Source/HairStrandsEditor/Private/GroomAssetDetails.cpp b/Engine/Plugins/Runtime/HairStrands/Source/HairStrandsEditor/Private/GroomAssetDetails.cpp
--- a/Engine/Plugins/Runtime/HairStrands/Source/HairStrandsEditor/Private/GroomAssetDetails.cpp
+++ b/Engine/Plugins/Runtime/HairStrands/Source/HairStrandsEditor/Private/GroomAssetDetails.cpp
@@ -922,12 +922,6 @@ bool FGroomRenderingDetails::CommonResetToDefault(TSharedPtr<IPropertyHandle> Ch
 	// LODs
 	if (bIsGroupIndexValid)
 	{
-		{
-			FHairGroupsLOD Default;
-			HAIR_RESET0(GetHairGroupsLOD(), FHairGroupsLOD, ClusterWorldSize);
-			HAIR_RESET0(GetHairGroupsLOD(), FHairGroupsLOD, ClusterScreenSizeScale);
-		}
-
 		if (LODIndex>=0 && LODIndex < GroomAsset->GetHairGroupsLOD()[GroupIndex].LODs.Num())
 		{
 			FHairLODSettings Default;
diff --git a/Engine/Shaders/Private/HairStrands/HairStrandsClusterCulling.usf b/Engine/Shaders/Private/HairStrands/HairStrandsClusterCulling.usf
--- a/Engine/Shaders/Private/HairStrands/HairStrandsClusterCulling.usf
+++ b/Engine/Shaders/Private/HairStrands/HairStrandsClusterCulling.usf
@@ -370,7 +370,7 @@ void Main(uint3 DispatchThreadId : SV_DispatchThreadID, uint ThreadIndex : SV_Gr
 		// 1.3 Get cluster info (i.e. map all info per LOD for a given cluster ID) and derive effective LOD data (radius scale)
 		const FHairClusterLOD ClusterLODInfo = GetHairClusterLOD(ClusterInfoBuffer[ClusterId], ClusterInfoParameters, LODIndex);
 		ClusterLOD[LocalCurveIndex] 		 = ClusterLODInfo.LOD;
-		ClusterRadiusScale[LocalCurveIndex]  = ClusterLODInfo.RadiusScale;
+		ClusterRadiusScale[LocalCurveIndex]  = 1.f; // Curve are no longer scaled individually by ClusterLODInfo.RadiusScale, but scale globally based on CPU value
 		LODCount[LocalCurveIndex]			 = ClusterLODInfo.LODCount;
 
 		// 1.4 Curve's point count/offset
@@ -447,7 +447,7 @@ void Main(uint3 DispatchThreadId : SV_DispatchThreadID, uint ThreadIndex : SV_Gr
 	if (ThreadIndex < AllocationCount)
 	{
 		OutIndexBuffer[AllocationOffset + ThreadIndex]       = CulledPointIndices[ThreadIndex];
-		OutRadiusScaleBuffer[AllocationOffset + ThreadIndex] = CulledPointRadiusScales[ThreadIndex] * CulledPointFadeScales[ThreadIndex];
+		OutRadiusScaleBuffer[AllocationOffset + ThreadIndex] = 1.f;  // For fading out the curve which are about to be culled, one can use CulledPointFadeScales[ThreadIndex]
 	}
 
 	// 6. Write out curve data with visible points for subsequent passes
@@ -501,7 +501,7 @@ void Main(uint3 DispatchThreadId : SV_DispatchThreadID, uint ThreadIndex : SV_Gr
 		// 1.2 Get cluster info (i.e. map all info per LOD for a given cluster ID) and derive effective LOD data (radius scale)
 		const FHairClusterLOD ClusterLODInfo = GetHairClusterLOD(ClusterInfoBuffer[ClusterId], ClusterInfoParameters, LODIndex);
 		ClusterLOD = ClusterLODInfo.LOD;
-		ClusterRadiusScale = ClusterLODInfo.RadiusScale;
+		ClusterRadiusScale = 1.f; // Curve are no longer scaled individually by ClusterLODInfo.RadiusScale, but scale globally based on CPU value
 		LODCount = ClusterLODInfo.LODCount;
 
 		// 1.4 Curve's point count/offset
@@ -583,7 +583,7 @@ void Main(uint3 DispatchThreadId : SV_DispatchThreadID, uint ThreadIndex : SV_Gr
 				const float FadeScale = CulledPointFadeScales[CurrentPointIndex];
 
 				OutIndexBuffer[AllocationOffset + CurrentPointIndex]       = CurvePointIndex + CullPointIndex;
-				OutRadiusScaleBuffer[AllocationOffset + CurrentPointIndex] = ClusterRadiusScale * FadeScale;
+				OutRadiusScaleBuffer[AllocationOffset + CurrentPointIndex] = 1.f;  // For fading out the curve which are about to be culled, one can use FadeScale
 			}
 		}
 	}
diff --git a/Engine/Shaders/Private/HairStrands/HairStrandsDebug.usf b/Engine/Shaders/Private/HairStrands/HairStrandsDebug.usf
--- a/Engine/Shaders/Private/HairStrands/HairStrandsDebug.usf
+++ b/Engine/Shaders/Private/HairStrands/HairStrandsDebug.usf
@@ -1193,6 +1193,7 @@ struct FInstanceInfo
 	float RootScale;
 	float Length;
 	float LengthScale;
+	float CoverageScale;
 	bool Voxel;
 	bool Scatter;
 	bool RTGeometry;
@@ -1285,6 +1286,7 @@ FInstanceInfo ReadInstanceInfo(uint Index)
 	Out.TipScale				= f16tof32(Data3.y >> 16u);
 	Out.Length					= f16tof32(Data3.z & 0xFFFF);
 	Out.LengthScale				= f16tof32(Data3.z >> 16u);
+	Out.CoverageScale			= Data3.w;
 
 	Out.bHasLODSwitched			= Data4.x > 0;
 
@@ -1422,7 +1424,7 @@ void MainCS(uint GroupIndex : SV_GroupIndex, uint3 DispatchThreadId : SV_Dispatc
 	Newline(Context);
 
 	// Draw instance details
-	Print(Context, TEXT("   Group  LOD     Geometry Binding Cache Sim RBF Clip Screen Strands            Points             Bind.Tri/Sect Vis. Shad. SwLOD Name"));
+	Print(Context, TEXT("   Group  LOD     Geometry Binding Cache Sim RBF Clip Screen CovSc. Strands            Points             Bind.Tri/Sect Vis. Shad. SwLOD Name"));
 	Newline(Context);
 
 	for (uint InstanceIndex = 0; InstanceIndex < InstanceCount; ++InstanceIndex)
@@ -1477,6 +1479,9 @@ void MainCS(uint GroupIndex : SV_GroupIndex, uint3 DispatchThreadId : SV_Dispatc
 		Print(Context, Info.ScreenSize, FontOrange, 4, 2);
 		Print(Context, TEXT("   "));
 
+		Print(Context, Info.CoverageScale, FontOrange, 4, 2);
+		Print(Context, TEXT("   "));
+
 		// Curve sCount
 		const float CurveRatio = Info.ActiveCurveCount / float(Info.StrandsCount);
 		PrintPercentage(Context, CurveRatio, InitFontColor(ColorMapTurbo(CurveRatio)));
@@ -1743,6 +1748,9 @@ void MainCS(int2 DispatchThreadId : SV_DispatchThreadID)
 int2 MaxResolution;
 uint GroupIndex;
 uint GeometryType;
+uint CurveCount;
+uint PointCount;
+float CoverageScale;
 float LOD;
 float ScreenSize;
 float3 GroupColor;
@@ -1753,43 +1761,41 @@ void MainCS(uint3 DispatchThreadId : SV_DispatchThreadID)
 	// Legend
 	if (GroupIndex == 0)
 	{
-		FShaderPrintContext Context = InitShaderPrintContext(true, uint2(10, 50));
-		Print(Context, TEXT("  Group Geom.   LOD  Screen Size"), FontSilver); Newline(Context);
+		FShaderPrintContext Ctx = InitShaderPrintContext(true, uint2(10, 50));
+		Print(Ctx, TEXT("  Group Geom.   LOD  Screen Size #Curve      #Point      Scale"), FontSilver); Newline(Ctx);
 	}
 
 	// Values
 	{
-		FShaderPrintContext Context = InitShaderPrintContext(true, uint2(10, 65 + GroupIndex * 15));
+		FShaderPrintContext Ctx = InitShaderPrintContext(true, uint2(10, 65 + GroupIndex * 15));
 
 		// Group Color
-		Print(Context, TEXT("* "), InitFontColor(GroupColor));
+		Print(Ctx, TEXT("* "), InitFontColor(GroupColor));
 
 		// Group index
-		Print(Context, GroupIndex, FontEmerald, 3, 0);
+		Print(Ctx, GroupIndex, FontEmerald, 3, 0);
 
 		// Geometry Type
-		Print(Context, TEXT("   "));
+		Print(Ctx, TEXT("   "));
 		if (GeometryType == 0)
 		{
-			Print(Context, TEXT("Strands "), FontOrange);
+			Print(Ctx, TEXT("Strands "), FontOrange);
 		}
 		else if (GeometryType == 1)
 		{
-			Print(Context, TEXT("Cards   "), FontOrange);
+			Print(Ctx, TEXT("Cards   "), FontOrange);
 		}
 		else if (GeometryType == 2)
 		{
-			Print(Context, TEXT("Meshes  "), FontOrange);
+			Print(Ctx, TEXT("Meshes  "), FontOrange);
 		}
 
-		// LOD
-		Print(Context, LOD, FontWhite, 4, 2);
-		Print(Context, TEXT(" "));
-
-		// Screen size
-		Print(Context, saturate(ScreenSize), FontWhite);
-
-		Newline(Context);
+		Print(Ctx, LOD, FontWhite, 4, 2); Print(Ctx, TEXT(" "));
+		Print(Ctx, saturate(ScreenSize), FontWhite);
+		Print(Ctx, CurveCount, FontWhite);
+		Print(Ctx, PointCount, FontWhite);
+		Print(Ctx, CoverageScale, FontWhite);
+		Newline(Ctx);
 	}
 }
 
diff --git a/Engine/Source/Runtime/Renderer/Private/HairStrands/HairStrandsDebug.cpp b/Engine/Source/Runtime/Renderer/Private/HairStrands/HairStrandsDebug.cpp
--- a/Engine/Source/Runtime/Renderer/Private/HairStrands/HairStrandsDebug.cpp
+++ b/Engine/Source/Runtime/Renderer/Private/HairStrands/HairStrandsDebug.cpp
@@ -108,6 +108,9 @@ class FHairPrintLODInfoCS : public FGlobalShader
 		SHADER_PARAMETER(FVector3f, GroupColor)
 		SHADER_PARAMETER(uint32, GroupIndex)
 		SHADER_PARAMETER(uint32, GeometryType)
+		SHADER_PARAMETER(uint32, CurveCount)
+		SHADER_PARAMETER(uint32, PointCount)
+		SHADER_PARAMETER(float, CoverageScale)
 		SHADER_PARAMETER(float, ScreenSize)
 		SHADER_PARAMETER(float, LOD)
 		SHADER_PARAMETER_STRUCT_INCLUDE(ShaderPrint::FShaderParameters, ShaderPrintUniformBuffer)
@@ -153,6 +156,9 @@ static void AddPrintLODInfoPass(
 	Parameters->LOD = Data->LODIndex;
 	Parameters->GroupColor = FVector3f(GroupColor.R, GroupColor.G, GroupColor.B);
 	Parameters->ScreenSize = Data->DebugScreenSize;
+	Parameters->CurveCount = Data->GetActiveStrandsCurveCount();
+	Parameters->PointCount = Data->GetActiveStrandsPointCount();
+	Parameters->CoverageScale = Data->GetActiveStrandsCoverageScale();
 	switch (Data->VFInput.GeometryType)
 	{
 	case EHairGeometryType::Strands: Parameters->GeometryType = 0; break;
diff --git a/Engine/Source/Runtime/Renderer/Private/HairStrands/HairStrandsInterface.cpp b/Engine/Source/Runtime/Renderer/Private/HairStrands/HairStrandsInterface.cpp
--- a/Engine/Source/Runtime/Renderer/Private/HairStrands/HairStrandsInterface.cpp
+++ b/Engine/Source/Runtime/Renderer/Private/HairStrands/HairStrandsInterface.cpp
@@ -505,8 +505,7 @@ uint32 FHairGroupPublicData::GetActiveStrandsCurveCount() const
 
 float FHairGroupPublicData::GetActiveStrandsCoverageScale() const
 {
-	const float CurveRatio = ContinuousLODCurveCount / float(RestCurveCount);
-	return 1.f/FMath::Max(CurveRatio, 0.01f);
+	return ContinuousLODCoverageScale; 
 }
 
 ///////////////////////////////////////////////////////////////////////////////////////////////////
diff --git a/Engine/Source/Runtime/Renderer/Public/HairStrandsInterface.h b/Engine/Source/Runtime/Renderer/Public/HairStrandsInterface.h
--- a/Engine/Source/Runtime/Renderer/Public/HairStrandsInterface.h
+++ b/Engine/Source/Runtime/Renderer/Public/HairStrandsInterface.h
@@ -415,6 +415,7 @@ public:
 	uint32 ContinuousLODPointCount = 0;
 	uint32 ContinuousLODCurveCount = 0;
 	float ContinuousLODScreenSize = 1.f;
+	float ContinuousLODCoverageScale = 1.f;
 	FVector2f ContinuousLODScreenPos = FVector2f(0,0);
 	FBoxSphereBounds ContinuousLODBounds; 	//used by Continuous LOD
 
