Make cluster data loading continuously.

diff --git a/Engine/Plugins/Runtime/HairStrands/Source/HairStrandsCore/Private/GroomBuilder.cpp b/Engine/Plugins/Runtime/HairStrands/Source/HairStrandsCore/Private/GroomBuilder.cpp
--- a/Engine/Plugins/Runtime/HairStrands/Source/HairStrandsCore/Private/GroomBuilder.cpp
+++ b/Engine/Plugins/Runtime/HairStrands/Source/HairStrandsCore/Private/GroomBuilder.cpp
@@ -40,7 +40,7 @@ static FAutoConsoleVariableRef CVarHairGroupIndexBuilder_MaxVoxelResolution(TEXT
 
 FString FGroomBuilder::GetVersion()
 {
-	return TEXT("v8r39");
+	return TEXT("v8r40");
 }
 
 namespace FHairStrandsDecimation
@@ -3074,6 +3074,10 @@ static void BuildClusterBulkData(
 	Out.Header.CurveCount		= In.CurveCount;
 	Out.Header.LODInfos			= In.LODInfos;
 
+	Out.Header.Strides.PackedClusterInfoStride = FHairClusterInfoFormat::SizeInByte;
+	Out.Header.Strides.CurveToClusterIdStride = FHairClusterIndexFormat::SizeInByte;
+	Out.Header.Strides.PointLODStride = sizeof(uint32);
+
 	// Sanity check
 	check(Out.Header.ClusterCount	== uint32(In.ClusterInfos.Num()));
 	check(Out.Header.CurveCount		== uint32(In.CurveToClusterIds.Num()));
diff --git a/Engine/Plugins/Runtime/HairStrands/Source/HairStrandsCore/Private/HairStrandsDatas.cpp b/Engine/Plugins/Runtime/HairStrands/Source/HairStrandsCore/Private/HairStrandsDatas.cpp
--- a/Engine/Plugins/Runtime/HairStrands/Source/HairStrandsCore/Private/HairStrandsDatas.cpp
+++ b/Engine/Plugins/Runtime/HairStrands/Source/HairStrandsCore/Private/HairStrandsDatas.cpp
@@ -696,6 +696,10 @@ void FHairStrandsClusterCullingBulkData::SerializeHeader(FArchive& Ar, UObject*
 		Ar << Header.LODInfos[It].ScreenSize;
 		Ar << Header.LODInfos[It].bIsVisible;
 	}
+
+	Ar << Header.Strides.PackedClusterInfoStride;
+	Ar << Header.Strides.CurveToClusterIdStride;
+	Ar << Header.Strides.PointLODStride;
 }
 
 uint32 FHairStrandsClusterCullingBulkData::GetResourceCount() const
@@ -710,9 +714,20 @@ void FHairStrandsClusterCullingBulkData::GetResources(FHairStrandsBulkCommon::FQ
 	const bool bHasClusterData = Header.ClusterCount > 0;
 	if (bHasClusterData)
 	{
-		Out.Add(Data.CurveToClusterIds, TEXT("_CurveToClusterIds"), 0, 0); 	// Load all data. TODO load data based on active curve/point count
-		Out.Add(Data.PackedClusterInfos, TEXT("_PackedClusterInfos"), 0, 0);// Load all data
-		Out.Add(Data.PointLODs, TEXT("_PointLODs"), 0, 0); 					// Load all data. TODO load data based on active curve/point count
+		// Translate requested curve count into chunk/offset/size to be read
+		uint32 PointCount = 0;
+		uint32 CurveCount = 0;
+		if (Out.Type == FHairStrandsBulkCommon::FQuery::ReadIO || Out.Type == FHairStrandsBulkCommon::FQuery::ReadDDC)
+		{
+			CurveCount = FMath::Min(Header.CurveCount, Out.GetCurveCount());
+			PointCount = FMath::Min(Header.PointCount, Out.GetPointCount());
+		}
+
+		const uint32 PointLODCount = FMath::DivideAndRoundUp(PointCount, HAIR_POINT_LOD_COUNT_PER_UINT);
+
+		Out.Add(Data.PackedClusterInfos, TEXT("_PackedClusterInfos"), Data.PackedClusterInfos.LoadedSize, Header.ClusterCount * Header.Strides.PackedClusterInfoStride); // Load all data
+		Out.Add(Data.CurveToClusterIds, TEXT("_CurveToClusterIds"), Data.CurveToClusterIds.LoadedSize, CurveCount * Header.Strides.CurveToClusterIdStride);
+		Out.Add(Data.PointLODs, TEXT("_PointLODs"), Data.PointLODs.LoadedSize, PointLODCount * Header.Strides.PointLODStride);
 	}
 }
 
diff --git a/Engine/Plugins/Runtime/HairStrands/Source/HairStrandsCore/Public/HairStrandsDatas.h b/Engine/Plugins/Runtime/HairStrands/Source/HairStrandsCore/Public/HairStrandsDatas.h
--- a/Engine/Plugins/Runtime/HairStrands/Source/HairStrandsCore/Public/HairStrandsDatas.h
+++ b/Engine/Plugins/Runtime/HairStrands/Source/HairStrandsCore/Public/HairStrandsDatas.h
@@ -728,6 +728,13 @@ struct HAIRSTRANDSCORE_API FHairStrandsClusterCullingBulkData : FHairStrandsBulk
 		uint32 PointCount = 0;
 		uint32 CurveCount = 0;
 		FVector4f ClusterInfoParameters = FVector4f::Zero(); // xy:Scale/Offset for ScreenSize, zw:Scale/Offset for Radius
+
+		struct FStrides
+		{
+			uint32 PackedClusterInfoStride = 0;
+			uint32 CurveToClusterIdStride = 0;
+			uint32 PointLODStride = 0;
+		} Strides;
 	} Header;
 
 	struct FData
