Rename VertexCount -> PointCount for semantic clarify.

diff --git a/Engine/Plugins/Runtime/HairStrands/Source/HairStrandsCore/Private/HairStrandsInterpolation.cpp b/Engine/Plugins/Runtime/HairStrands/Source/HairStrandsCore/Private/HairStrandsInterpolation.cpp
--- a/Engine/Plugins/Runtime/HairStrands/Source/HairStrandsCore/Private/HairStrandsInterpolation.cpp
+++ b/Engine/Plugins/Runtime/HairStrands/Source/HairStrandsCore/Private/HairStrandsInterpolation.cpp
@@ -1169,7 +1169,7 @@ class FHairRaytracingGeometryCS : public FGlobalShader
 	using FPermutationDomain = TShaderPermutationDomain<FCulling, FProceduralPrimitive>;
 
 	BEGIN_SHADER_PARAMETER_STRUCT(FParameters, )
-		SHADER_PARAMETER(uint32, VertexCount)
+		SHADER_PARAMETER(uint32, PointCount)
 		SHADER_PARAMETER(float, HairStrandsVF_HairRadius)
 		SHADER_PARAMETER(float, HairStrandsVF_HairRootScale)
 		SHADER_PARAMETER(float, HairStrandsVF_HairTipScale)
@@ -1215,7 +1215,7 @@ static void AddGenerateRaytracingGeometryPass(
 	FRDGBuilder& GraphBuilder,
 	FGlobalShaderMap* ShaderMap,
 	const FShaderPrintData* ShaderPrintData,
-	uint32 VertexCount,
+	uint32 PointCount,
 	bool bProceduralPrimitive,
 	int ProceduralSplits,
 	float HairRadius,
@@ -1229,7 +1229,7 @@ static void AddGenerateRaytracingGeometryPass(
 	const FRDGBufferUAVRef& OutIndexBuffer)
 {
 	FHairRaytracingGeometryCS::FParameters* Parameters = GraphBuilder.AllocParameters<FHairRaytracingGeometryCS::FParameters>();
-	Parameters->VertexCount = VertexCount;
+	Parameters->PointCount = PointCount;
 	Parameters->PositionOffsetBuffer = HairWorldOffsetBuffer;
 	Parameters->HairStrandsVF_HairRadius = HairRadius;
 	Parameters->HairStrandsVF_HairRootScale = RootScale;
@@ -1260,7 +1260,7 @@ static void AddGenerateRaytracingGeometryPass(
 	PermutationVector.Set<FHairRaytracingGeometryCS::FProceduralPrimitive>(bProceduralPrimitive);
 	TShaderMapRef<FHairRaytracingGeometryCS> ComputeShader(ShaderMap, PermutationVector);
 
-	const FIntVector DispatchCount(FMath::DivideAndRoundUp(VertexCount, FHairRaytracingGeometryCS::GetGroupSize()), 1, 1);
+	const FIntVector DispatchCount(FMath::DivideAndRoundUp(PointCount, FHairRaytracingGeometryCS::GetGroupSize()), 1, 1);
 	FComputeShaderUtils::AddPass(
 		GraphBuilder,
 		RDG_EVENT_NAME("HairStrands::GenerateRaytracingGeometry(bCulling:%d, Procedural:%d)", bCullingEnable, bProceduralPrimitive),
diff --git a/Engine/Shaders/Private/HairStrands/HairStrandsRaytracingGeometry.usf b/Engine/Shaders/Private/HairStrands/HairStrandsRaytracingGeometry.usf
--- a/Engine/Shaders/Private/HairStrands/HairStrandsRaytracingGeometry.usf
+++ b/Engine/Shaders/Private/HairStrands/HairStrandsRaytracingGeometry.usf
@@ -8,7 +8,7 @@
 
 #if SHADER_RT_GEOMETRY
 
-uint VertexCount;
+uint PointCount;
 
 float HairStrandsVF_HairRadius;
 float HairStrandsVF_HairRootScale;
@@ -36,7 +36,7 @@ FHairControlPoint GetControlPoint(uint ControlPointId)
 {
 	FHairControlPoint Out = (FHairControlPoint)0;
 
-	if (ControlPointId < VertexCount)
+	if (ControlPointId < PointCount)
 	{
 		Out = UnpackHairControlPoint(
 			PositionBuffer[ControlPointId],
@@ -159,7 +159,7 @@ void MainCS(uint DispatchThreadId : SV_DispatchThreadID)
 	uint VertexIndex0 = DispatchThreadId;
 	uint VertexIndex1 = VertexIndex0 + 1;
 
-	const bool bIsValid = VertexIndex0 < VertexCount && VertexIndex1 < VertexCount;
+	const bool bIsValid = VertexIndex0 < PointCount && VertexIndex1 < PointCount;
 	if (!bIsValid)
 	{
 		SetInvalidSegment(VertexIndex0);
@@ -297,7 +297,7 @@ void MainCS(uint DispatchThreadId : SV_DispatchThreadID)
 #error GROUP_SIZE has changed, please update the reduction code.
 #endif
 
-uint VertexCount;
+uint PointCount;
 uint bDrawInvalidElement;
 float PositionThreshold2;
 
@@ -323,7 +323,7 @@ void MainCS(uint2 DispatchThreadId : SV_DispatchThreadID, uint LinearIndex : SV_
 	{
 		// Fetch vertex index if culling is enabled
 		uint VertexIndex0 = DispatchThreadId.x;
-		bool bIsValid = VertexIndex0 < VertexCount;
+		bool bIsValid = VertexIndex0 < PointCount;
 		if (bIsValid && HairStrandsVF_bCullingEnable)
 		{
 			const uint VertexCountAfterCulling = HairStrandsVF_CullingIndirectBuffer[3];
diff --git a/Engine/Source/Runtime/Renderer/Private/HairStrands/HairStrandsVisibility.cpp b/Engine/Source/Runtime/Renderer/Private/HairStrands/HairStrandsVisibility.cpp
--- a/Engine/Source/Runtime/Renderer/Private/HairStrands/HairStrandsVisibility.cpp
+++ b/Engine/Source/Runtime/Renderer/Private/HairStrands/HairStrandsVisibility.cpp
@@ -3994,7 +3994,7 @@ class FHairStrandsPositionChangedCS : public FGlobalShader
 	SHADER_USE_PARAMETER_STRUCT(FHairStrandsPositionChangedCS, FGlobalShader);
 
 	BEGIN_SHADER_PARAMETER_STRUCT(FParameters, )
-		SHADER_PARAMETER(uint32, VertexCount)
+		SHADER_PARAMETER(uint32, PointCount)
 		SHADER_PARAMETER(float, PositionThreshold2)
 		SHADER_PARAMETER(uint32, HairStrandsVF_bCullingEnable)
 		SHADER_PARAMETER(uint32, bDrawInvalidElement)
@@ -4032,7 +4032,7 @@ static void AddHairStrandsHasPositionChangedPass(
 	AddClearUAVPass(GraphBuilder, InvalidationPrintCounterUAV, 0u);
 
 	FHairStrandsPositionChangedCS::FParameters* Parameters = GraphBuilder.AllocParameters<FHairStrandsPositionChangedCS::FParameters>();
-	Parameters->VertexCount = PointCount;
+	Parameters->PointCount = PointCount;
 	Parameters->PositionThreshold2 = FMath::Square(GHairStrands_InvalidationPosition_Threshold);
 	Parameters->bDrawInvalidElement = GHairStrands_InvalidationPosition_Debug > 0 ? 1u : 0u;
 	Parameters->HairStrandsVF_bCullingEnable = 0u;
