Change allocation API for groom resource.

diff --git a/Engine/Plugins/Runtime/HairStrands/Source/HairStrandsCore/Private/GroomResources.cpp b/Engine/Plugins/Runtime/HairStrands/Source/HairStrandsCore/Private/GroomResources.cpp
--- a/Engine/Plugins/Runtime/HairStrands/Source/HairStrandsCore/Private/GroomResources.cpp
+++ b/Engine/Plugins/Runtime/HairStrands/Source/HairStrandsCore/Private/GroomResources.cpp
@@ -636,7 +636,7 @@ void FHairCommonResource::InitRHI()
 	{
 		FRHICommandListImmediate& RHICmdList = FRHICommandListExecutor::GetImmediateCommandList();
 		FRDGBuilder GraphBuilder(RHICmdList);
-		InternalAllocate(GraphBuilder, HAIR_MAX_NUM_CURVE_PER_GROUP, DummyLODIndex);
+		InternalAllocate(GraphBuilder, HAIR_MAX_NUM_CURVE_PER_GROUP, HAIR_MAX_NUM_POINT_PER_GROUP, DummyLODIndex);
 		GraphBuilder.Execute();
 	}
 	else
@@ -689,7 +689,7 @@ void FHairCommonResource::Allocate(FRDGBuilder& GraphBuilder, EHairResourceLoadi
 		}
 		if (!bIsInitialized || !InternalIsLODDataLoaded(InRequestedCurveCount, InRequestedPointCount, InLODIndex))
 		{
-			InternalAllocate(GraphBuilder, InRequestedCurveCount, InLODIndex);
+			InternalAllocate(GraphBuilder, InRequestedCurveCount, InRequestedPointCount, InLODIndex);
 		}
 		bIsInitialized = true;
 
@@ -713,7 +713,7 @@ void FHairCommonResource::Allocate(FRDGBuilder& GraphBuilder, EHairResourceLoadi
 			{
 				FRenderResource::InitResource(GraphBuilder.RHICmdList); // Call RenderResource InitResource() so that the resources is marked as initialized
 			}
-			InternalAllocate(GraphBuilder, StreamingRequest.CurveCount, StreamingRequest.LODIndex);
+			InternalAllocate(GraphBuilder, StreamingRequest.CurveCount, StreamingRequest.PointCount, StreamingRequest.LODIndex);
 			bIsInitialized = true;
 
 			// Update the max curve count available
@@ -1258,7 +1258,7 @@ bool FHairStrandsRestRootResource::InternalGetOrRequestData(uint32 InRequestedCu
 	return InLODIndex >= 0 && StreamingRequest.IsCompleted();
 }
 
-void FHairStrandsRestRootResource::InternalAllocate(FRDGBuilder& GraphBuilder, uint32 InCurveCount, int32 InLODIndex)
+void FHairStrandsRestRootResource::InternalAllocate(FRDGBuilder& GraphBuilder, uint32 InCurveCount, uint32 InPointCount, int32 InLODIndex)
 {
 	// Once empty, the MeshProjectionLODsneeds to be repopulate as it might be re-initialized. 
 	// E.g., when a resource is updated, it is first released, then re-init. 
@@ -1366,7 +1366,7 @@ bool FHairStrandsDeformedRootResource::InternalIsLODDataLoaded(uint32 InRequeste
 	return true;
 }
 
-void FHairStrandsDeformedRootResource::InternalAllocate(FRDGBuilder& GraphBuilder, uint32 InCurveCount, int32 InLODIndex)
+void FHairStrandsDeformedRootResource::InternalAllocate(FRDGBuilder& GraphBuilder, uint32 InCurveCount, uint32 InPointCount, int32 InLODIndex)
 {
 	if (RootCount > 0 && InLODIndex >= 0 && InLODIndex < LODs.Num())
 	{
diff --git a/Engine/Plugins/Runtime/HairStrands/Source/HairStrandsCore/Public/GroomResources.h b/Engine/Plugins/Runtime/HairStrands/Source/HairStrandsCore/Public/GroomResources.h
--- a/Engine/Plugins/Runtime/HairStrands/Source/HairStrandsCore/Public/GroomResources.h
+++ b/Engine/Plugins/Runtime/HairStrands/Source/HairStrandsCore/Public/GroomResources.h
@@ -109,7 +109,7 @@ struct FHairCommonResource : public FRenderResource
 
 	virtual void InternalAllocate() {}
 	virtual void InternalAllocate(FRDGBuilder& GraphBuilder) {}
-	virtual void InternalAllocate(FRDGBuilder& GraphBuilder, uint32 InCurveCount, int32 InLODIndex) { InternalAllocate(GraphBuilder); }
+	virtual void InternalAllocate(FRDGBuilder& GraphBuilder, uint32 InCurveCount, uint32 InPointCount, int32 InLODIndex) { InternalAllocate(GraphBuilder); }
 	virtual void InternalRelease() {}
 	virtual bool InternalGetOrRequestData(uint32 InRequestedCurveCount, uint32 InRequestedPointCount, int32 InLODIndex);
 	virtual bool InternalIsLODDataLoaded(uint32 InRequestedCurveCount, uint32 InRequestedPointCount, int32 InLODIndex) const { return true; }
@@ -353,7 +353,7 @@ struct FHairStrandsRestRootResource : public FHairCommonResource
 	FHairStrandsRestRootResource(FHairStrandsRootBulkData& BulkData, EHairStrandsResourcesType CurveType, const FHairResourceName& ResourceName, const FName& OwnerName);
 
 	/* Init/Release buffers */
-	virtual void InternalAllocate(FRDGBuilder& GraphBuilder, uint32 InCurveCount, int32 InLODIndex) override;
+	virtual void InternalAllocate(FRDGBuilder& GraphBuilder, uint32 InCurveCount, uint32 InPointCount, int32 InLODIndex) override;
 	virtual void InternalRelease() override;
 	virtual bool InternalGetOrRequestData(uint32 InRequestedCurveCount, uint32 InRequestedPointCount, int32 InLODIndex) override;
 	virtual bool InternalIsLODDataLoaded(uint32 InRequestedCurveCount, uint32 InRequestedPointCount, int32 InLODIndex) const override;
@@ -431,7 +431,7 @@ struct FHairStrandsDeformedRootResource : public FHairCommonResource
 	FHairStrandsDeformedRootResource(const FHairStrandsRestRootResource* InRestResources, EHairStrandsResourcesType CurveType, const FHairResourceName& ResourceName, const FName& OwnerName);
 
 	/* Init/Release buffers */
-	virtual void InternalAllocate(FRDGBuilder& GraphBuilder, uint32 InCurveCount, int32 InLODIndex) override;
+	virtual void InternalAllocate(FRDGBuilder& GraphBuilder, uint32 InCurveCount, uint32 InPointCount, int32 InLODIndex) override;
 	virtual void InternalRelease() override;
 	virtual bool InternalIsLODDataLoaded(uint32 InRequestedCurveCount, uint32 InRequestedPointCount, int32 InLODIndex) const override;
 
