* Add r.HairStrands.Streaming CVar to enable/disable streaming
* Remove ContinousLOD toggling and clean up code making usage of that variable.

diff --git a/Engine/Plugins/Runtime/HairStrands/Source/HairStrandsCore/Private/GroomManager.cpp b/Engine/Plugins/Runtime/HairStrands/Source/HairStrandsCore/Private/GroomManager.cpp
--- a/Engine/Plugins/Runtime/HairStrands/Source/HairStrandsCore/Private/GroomManager.cpp
+++ b/Engine/Plugins/Runtime/HairStrands/Source/HairStrandsCore/Private/GroomManager.cpp
@@ -44,8 +44,8 @@ static FAutoConsoleVariableRef CVarGHairStrands_ManualSkinCache(TEXT("r.HairStra
 static int32 GHairStrands_InterpolationFrustumCullingEnable = 1;
 static FAutoConsoleVariableRef CVarHairStrands_InterpolationFrustumCullingEnable(TEXT("r.HairStrands.Interoplation.FrustumCulling"), GHairStrands_InterpolationFrustumCullingEnable, TEXT("Swap rendering buffer at the end of frame. This is an experimental toggle. Default:1"));
 
-static int32 GHairStrands_ContinousLOD = 0;
-static FAutoConsoleVariableRef CVarHairStrands_ContinousLOD(TEXT("r.HairStrands.ContinuousLOD"), GHairStrands_ContinousLOD, TEXT("Continuous LODing toggle."), ECVF_RenderThreadSafe);
+static int32 GHairStrands_Streaming = 0;
+static FAutoConsoleVariableRef CVarHairStrands_Streaming(TEXT("r.HairStrands.Streaming"), GHairStrands_Streaming, TEXT("Hair strands streaming toggle."), ECVF_RenderThreadSafe | ECVF_ReadOnly);
 
 EHairBufferSwapType GetHairSwapBufferType()
 {
@@ -523,11 +523,6 @@ static void RunHairStrandsInterpolation_Strands(
 		&ClusterData);
 }
 
-bool IsHairStrandsContinousLODEnabled()
-{
-	return GHairStrands_ContinousLOD > 0;
-}
-
 // Return the LOD which should be used for a given screen size and LOD bias value
 // This function is mirrored in HairStrandsClusterCommon.ush
 static float GetHairInstanceLODIndex(const TArray<float>& InLODScreenSizes, float InScreenSize, float InLODBias)
@@ -747,17 +742,6 @@ void AddHairStreamingRequest(FHairGroupInstance* Instance, int32 InLODIndex)
 	}
 }
 
-static uint32 GetHairVisibilityComputeRasterPointCount(float ScreenSize, uint32 InPointCount)
-{
-	if (IsHairVisibilityComputeRasterContinuousLODEnabled())
-	{
-		const float Scale = FMath::Clamp(FMath::Pow(FMath::Clamp(ScreenSize, 1.0f / 16.0f, 1.0f), 1.0f), 0.f, 1.f);
-		InPointCount *= Scale;
-	}
-
-	return InPointCount;
-}
-
 static FVector2f ComputeProjectedScreenPos(const FVector& InWorldPos, const FSceneView& View)
 {
 	// Compute the MinP/MaxP in pixel coord, relative to View.ViewRect.Min
@@ -789,6 +773,13 @@ static FVector2f ComputeProjectedScreenPos(const FVector& InWorldPos, const FSce
 	return FVector2f(ScreenPos.X, ScreenPos.Y);
 }
 
+static float ComputeActiveCurveCoverageScale(uint32 InAvailableCurveCount, uint32 InRestCurveCount)
+{
+	// Compensate lost in curve by a coverage scale increase
+	const float CurveRatio = InAvailableCurveCount / float(InRestCurveCount);
+	return 1.f/FMath::Max(CurveRatio, 0.01f);
+}
+
 static uint32 ComputeActiveCurveCount(float InScreenSize, uint32 InCurveCount)
 {
 	const float Power = 1.8f; // Per Groom settings?
@@ -890,35 +881,14 @@ static void RunHairLODSelection(
 			Instance->HairGroupPublicData->ContinuousLODBounds = SphereBound;
 			Instance->HairGroupPublicData->ContinuousLODCoverageScale = 1.f;
 
-			if (IsHairStrandsContinousLODEnabled())
-			{
-				// These values are set later, once the resources are loaded
-				Instance->HairGroupPublicData->ContinuousLODCurveCount = 0;
-				Instance->HairGroupPublicData->ContinuousLODPointCount = 0;
-			}
-			else if (IsHairStrandContinuousDecimationReorderingEnabled() && IsHairVisibilityComputeRasterContinuousLODEnabled())
-			{
-				Instance->HairGroupPublicData->ContinuousLODPointCount = GetHairVisibilityComputeRasterPointCount(MaxScreenSize_Bound, Instance->HairGroupPublicData->RestPointCount);
-				Instance->HairGroupPublicData->ContinuousLODCurveCount = Instance->HairGroupPublicData->RestCurveCount;
-			}
-			else if (Instance->Strands.ClusterCullingResource)
+			if (Instance->Strands.ClusterCullingResource)
 			{
-				const TArray<FHairLODInfo>& LODInfos = Instance->Strands.ClusterCullingResource->BulkData.Header.LODInfos;
-
-				const int32 iLODIndex = LODIndex;
-				const float S = LODIndex - iLODIndex;
-
-				const uint32 LODIndex0 = FMath::Min(iLODIndex,   LODInfos.Num()-1);
-				const uint32 LODIndex1 = FMath::Min(iLODIndex+1, LODInfos.Num()-1);
-				const uint32 EffectiveCurveCount = FMath::Lerp(LODInfos[LODIndex1].CurveCount, LODInfos[LODIndex0].CurveCount, 1.f-S);
-
-				// Compensate lost in curve by a coverage scale increase
-				const float CurveRatio = EffectiveCurveCount / float(Instance->HairGroupPublicData->RestCurveCount);
-
+				const uint32 EffectiveCurveCount = Instance->Strands.ClusterCullingResource->BulkData.GetCurveCount(LODIndex);				            
 				check(EffectiveCurveCount <= uint32(Instance->Strands.Data->Header.CurveToPointCount.Num()));
+
 				Instance->HairGroupPublicData->ContinuousLODCurveCount = EffectiveCurveCount;
 				Instance->HairGroupPublicData->ContinuousLODPointCount = EffectiveCurveCount > 0 ? Instance->Strands.Data->Header.CurveToPointCount[EffectiveCurveCount - 1] : 0;
-				Instance->HairGroupPublicData->ContinuousLODCoverageScale = 1.f/FMath::Max(CurveRatio, 0.01f);
+				Instance->HairGroupPublicData->ContinuousLODCoverageScale = ComputeActiveCurveCoverageScale(EffectiveCurveCount, Instance->HairGroupPublicData->RestCurveCount);
 			}
 		}
 
@@ -962,7 +932,7 @@ static void RunHairLODSelection(
 				{
 					const bool bNeedLODing		= LODIndex > 0;
 					const bool bNeedDeformation = Instance->Strands.DeformedResource != nullptr;
-					bCullingEnable = (bNeedLODing || bNeedDeformation) && !IsHairStrandsContinousLODEnabled();
+					bCullingEnable = (bNeedLODing || bNeedDeformation);
 				}
 			}
 
@@ -971,12 +941,14 @@ static void RunHairLODSelection(
 				GeometryType = EHairGeometryType::NoneGeometry;
 			}
 
+			// Transform ContinuousLOD CurveCount/PointCount into streaming request
 			uint32 RequestedCurveCount = Instance->HairGroupPublicData->RestCurveCount;
 			uint32 RequestedPointCount = Instance->HairGroupPublicData->RestPointCount;
-			if (GeometryType == EHairGeometryType::Strands && IsHairStrandsContinousLODEnabled() && Instance->bSupportStreaming)
+			const bool bStreamingEnabled = GHairStrands_Streaming > 0;
+			if (GeometryType == EHairGeometryType::Strands && bStreamingEnabled && Instance->bSupportStreaming)
 			{
-				RequestedCurveCount = ComputeActiveCurveCount(Instance->HairGroupPublicData->ContinuousLODScreenSize, Instance->HairGroupPublicData->RestCurveCount);
-				RequestedPointCount = (Instance->Strands.RestResource) ? Instance->Strands.RestResource->BulkData.Header.CurveToPointCount[RequestedCurveCount -1] : 0;
+				RequestedCurveCount = Instance->HairGroupPublicData->ContinuousLODCurveCount;
+				RequestedPointCount = Instance->HairGroupPublicData->ContinuousLODPointCount;
 			}
 
 			const bool bSimulationEnable			= Instance->HairGroupPublicData->IsSimulationEnable(IntLODIndex);
@@ -1063,7 +1035,7 @@ static void RunHairLODSelection(
 			}
 
 			// Only switch LOD if the data are ready to be used
-			const bool bIsLODDataReady = !ResourceStatus.HasStatus(EHairResourceStatus::EStatus::Loading) || (IsHairStrandsContinousLODEnabled() && ResourceStatus.AvailableCurveCount > 0);
+			const bool bIsLODDataReady = !ResourceStatus.HasStatus(EHairResourceStatus::EStatus::Loading) || (bStreamingEnabled && ResourceStatus.AvailableCurveCount > 0);
 			if (bIsLODDataReady)
 			{
 				const uint32 IntPrevLODIndex = FMath::FloorToInt(PrevLODIndex); 
@@ -1102,8 +1074,8 @@ static void RunHairLODSelection(
 				return false;
 			}
 
-			// Adapt the number curve/point based on Bound/CPU screen-size
-			if (GeometryType == EHairGeometryType::Strands && IsHairStrandsContinousLODEnabled())
+			// Adapt the number curve/point based on available curves/points
+			if (GeometryType == EHairGeometryType::Strands && bStreamingEnabled)
 			{
 				if (!bIsLODDataReady)
 				{
@@ -1112,9 +1084,11 @@ static void RunHairLODSelection(
 				check(bIsLODDataReady);
 				check(Instance->Strands.RestResource);
 
+				// Adapt CurveCount/PointCount/CoverageScale based on what is actually available
 				const uint32 EffectiveCurveCount = FMath::Min(ResourceStatus.AvailableCurveCount, RequestedCurveCount);
 				Instance->HairGroupPublicData->ContinuousLODCurveCount = EffectiveCurveCount;
 				Instance->HairGroupPublicData->ContinuousLODPointCount = EffectiveCurveCount > 0 ? Instance->Strands.Data->Header.CurveToPointCount[EffectiveCurveCount - 1] : 0;
+				Instance->HairGroupPublicData->ContinuousLODCoverageScale = ComputeActiveCurveCoverageScale(EffectiveCurveCount, Instance->HairGroupPublicData->RestCurveCount);
 				check(Instance->HairGroupPublicData->ContinuousLODPointCount <= Instance->HairGroupPublicData->RestPointCount);
 			}
 			return true;
diff --git a/Engine/Plugins/Runtime/HairStrands/Source/HairStrandsCore/Private/GroomResources.cpp b/Engine/Plugins/Runtime/HairStrands/Source/HairStrandsCore/Private/GroomResources.cpp
--- a/Engine/Plugins/Runtime/HairStrands/Source/HairStrandsCore/Private/GroomResources.cpp
+++ b/Engine/Plugins/Runtime/HairStrands/Source/HairStrandsCore/Private/GroomResources.cpp
@@ -52,8 +52,6 @@ static int32 GHairStrandsDebugVoxel_MaxSegmentPerVoxel = 2048;
 static FAutoConsoleVariableRef CVarHairStrandsDebugVoxel_WorldSize(TEXT("r.HairStrands.DebugData.VoxelSize"), GHairStrandsDebugVoxel_WorldSize, TEXT("Voxel size use for creating debug data."));
 static FAutoConsoleVariableRef CVarHairStrandsDebugVoxel_MaxSegmentPerVoxel(TEXT("r.HairStrands.DebugData.MaxSegmentPerVoxel"), GHairStrandsDebugVoxel_MaxSegmentPerVoxel, TEXT("Max number of segments per Voxel size when creating debug data."));
 
-bool IsHairStrandsContinousLODEnabled();
-
 bool ValidateHairBulkData()
 {
 	return GHairStrandsBulkData_Validation > 0;
diff --git a/Engine/Plugins/Runtime/HairStrands/Source/HairStrandsCore/Private/HairStrandsDatas.cpp b/Engine/Plugins/Runtime/HairStrands/Source/HairStrandsCore/Private/HairStrandsDatas.cpp
--- a/Engine/Plugins/Runtime/HairStrands/Source/HairStrandsCore/Private/HairStrandsDatas.cpp
+++ b/Engine/Plugins/Runtime/HairStrands/Source/HairStrandsCore/Private/HairStrandsDatas.cpp
@@ -731,6 +731,20 @@ void FHairStrandsClusterCullingBulkData::GetResources(FHairStrandsBulkCommon::FQ
 	}
 }
 
+uint32 FHairStrandsClusterCullingBulkData::GetCurveCount(float InLODIndex) const
+{
+	const TArray<FHairLODInfo>& LODInfos = Header.LODInfos;
+
+	InLODIndex = FMath::Max(InLODIndex, 0.f);
+
+	const int32 iLODIndex = InLODIndex;
+	const float S = InLODIndex - iLODIndex;
+
+	const uint32 LODIndex0 = FMath::Min(iLODIndex,   LODInfos.Num()-1);
+	const uint32 LODIndex1 = FMath::Min(iLODIndex+1, LODInfos.Num()-1);
+	return FMath::Lerp(LODInfos[LODIndex1].CurveCount, LODInfos[LODIndex0].CurveCount, 1.f-S);
+}
+
 /////////////////////////////////////////////////////////////////////////////////////////
 // Root data
 
diff --git a/Engine/Plugins/Runtime/HairStrands/Source/HairStrandsCore/Private/HairStrandsInterpolation.cpp b/Engine/Plugins/Runtime/HairStrands/Source/HairStrandsCore/Private/HairStrandsInterpolation.cpp
--- a/Engine/Plugins/Runtime/HairStrands/Source/HairStrandsCore/Private/HairStrandsInterpolation.cpp
+++ b/Engine/Plugins/Runtime/HairStrands/Source/HairStrandsCore/Private/HairStrandsInterpolation.cpp
@@ -77,8 +77,6 @@ int GetHairRaytracingProceduralSplits()
 	return FMath::Clamp(GHairRaytracingProceduralSplits, 1, STRANDS_PROCEDURAL_INTERSECTOR_MAX_SPLITS);
 }
 
-bool IsHairStrandsContinousLODEnabled();
-
 inline uint32 ComputeGroupSize()
 {
 	const uint32 GroupSize = IsRHIDeviceAMD() ? 64 : (IsRHIDeviceNVIDIA() ? 32 : 64);
@@ -2071,11 +2069,6 @@ void ComputeHairStrandsInterpolation(
 					}
 				}
 
-				// When CLOD is enabled, GPU AABB are not supported yet
-				if (IsHairStrandsContinousLODEnabled())
-				{
-					bNeedGPUAABB = false;
-				}
 				Instance->HairGroupPublicData->bClusterAABBValid = false;
 				Instance->HairGroupPublicData->bGroupAABBValid = false;
 
diff --git a/Engine/Plugins/Runtime/HairStrands/Source/HairStrandsCore/Public/HairStrandsDatas.h b/Engine/Plugins/Runtime/HairStrands/Source/HairStrandsCore/Public/HairStrandsDatas.h
--- a/Engine/Plugins/Runtime/HairStrands/Source/HairStrandsCore/Public/HairStrandsDatas.h
+++ b/Engine/Plugins/Runtime/HairStrands/Source/HairStrandsCore/Public/HairStrandsDatas.h
@@ -716,6 +716,7 @@ struct HAIRSTRANDSCORE_API FHairStrandsClusterCullingBulkData : FHairStrandsBulk
 	virtual void SerializeHeader(FArchive& Ar, UObject* Owner) override;
 	virtual uint32 GetResourceCount() const override;
 	virtual void GetResources(FQuery& Out) override;
+	uint32 GetCurveCount(float InLODIndex) const;
 
 	bool IsValid() const { return Header.ClusterCount > 0 && Header.PointCount > 0; }
 
diff --git a/Engine/Source/Runtime/Renderer/Private/HairStrands/HairStrandsVoxelization.cpp b/Engine/Source/Runtime/Renderer/Private/HairStrands/HairStrandsVoxelization.cpp
--- a/Engine/Source/Runtime/Renderer/Private/HairStrands/HairStrandsVoxelization.cpp
+++ b/Engine/Source/Runtime/Renderer/Private/HairStrands/HairStrandsVoxelization.cpp
@@ -124,12 +124,6 @@ static FAutoConsoleVariableRef CVarHairVirtualVoxelAdaptive_CorrectionThreshold(
 static int32 GHairVirtualVoxel_JitterMode = 1;
 static FAutoConsoleVariableRef CVarHairVirtualVoxel_JitterMode(TEXT("r.HairStrands.Voxelization.Virtual.Jitter"), GHairVirtualVoxel_JitterMode, TEXT("Change jittered for voxelization/traversal. 0: No jitter 1: Regular randomized jitter: 2: Constant Jitter (default = 1)"), ECVF_RenderThreadSafe);
 
-static bool IsHairStrandsContinousLODEnabled()
-{
-	static IConsoleVariable* CVar = IConsoleManager::Get().FindConsoleVariable(TEXT("r.HairStrands.ContinuousLOD"));
-	return CVar && CVar->GetInt() > 0;
-}
-
 void GetVoxelPageResolution(uint32& OutPageResolution, uint32& OutPageResolutionLog2)
 {
 	OutPageResolutionLog2 = FMath::CeilLogTwo(FMath::Clamp(GHairVirtualVoxel_PageResolution, 2, 256));
@@ -138,8 +132,7 @@ void GetVoxelPageResolution(uint32& OutPageResolution, uint32& OutPageResolution
 
 static bool IsHairStrandsAdaptiveVoxelAllocationEnable()
 {
-	// HAIR_TODO: GPU driven should be compatible with CLOD
-	return !IsHairStrandsContinousLODEnabled() && GHairVirtualVoxelAdaptive_Enable > 0;
+	return GHairVirtualVoxelAdaptive_Enable > 0;
 }
 
 bool IsHairStrandsVoxelizationEnable()
@@ -154,8 +147,7 @@ bool IsHairStrandsForVoxelTransmittanceAndShadowEnable()
 
 bool IsHairStrandsGPUDrivenVoxelizationEnabled()
 {
-	// HAIR_TODO: GPU driven should be compatible with CLOD
-	return !IsHairStrandsContinousLODEnabled() && GHairVirtualVoxelGPUDriven > 0; 
+	return GHairVirtualVoxelGPUDriven > 0; 
 }
 
 
@@ -1171,7 +1163,7 @@ static void AddVirtualVoxelizationComputeRasterPass(
 				PassParameters->VertexCount = PointCount;
 				PassParameters->CoverageScale = CoverageScale;
 			
-				const bool bCullingEnable = !IsHairStrandsContinousLODEnabled() && !IsHairStrandContinuousDecimationReorderingEnabled() && HairGroupPublicData->GetCullingResultAvailable();
+				const bool bCullingEnable = HairGroupPublicData->GetCullingResultAvailable();
 
 				PassParameters->HairInstance = GetHairStrandsInstanceParameters(GraphBuilder, *ViewInfo, HairGroupPublicData, bCullingEnable, true);
 
