Add auto LOD settings to groom.

diff --git a/Engine/Plugins/Runtime/HairStrands/Source/HairStrandsCore/Private/GroomComponent.cpp b/Engine/Plugins/Runtime/HairStrands/Source/HairStrandsCore/Private/GroomComponent.cpp
--- a/Engine/Plugins/Runtime/HairStrands/Source/HairStrandsCore/Private/GroomComponent.cpp
+++ b/Engine/Plugins/Runtime/HairStrands/Source/HairStrandsCore/Private/GroomComponent.cpp
@@ -2600,6 +2600,7 @@ void UGroomComponent::InitResources(bool bIsBindingReloading)
 			HairGroupInstance->HairGroupPublicData->SetLODScreenSizes(CPULODScreenSize);
 			HairGroupInstance->HairGroupPublicData->SetLODVisibilities(LODVisibility);
 			HairGroupInstance->HairGroupPublicData->SetLODGeometryTypes(LODGeometryTypes);
+			HairGroupInstance->HairGroupPublicData->bAutoLOD = GroomAsset->HairGroupsLOD[GroupIt].LODType == EGroomLODType::Auto;
 		}
 
 		FHairResourceName ResourceName(GetFName(), GroupIt);
diff --git a/Engine/Plugins/Runtime/HairStrands/Source/HairStrandsCore/Private/GroomDebug.cpp b/Engine/Plugins/Runtime/HairStrands/Source/HairStrandsCore/Private/GroomDebug.cpp
--- a/Engine/Plugins/Runtime/HairStrands/Source/HairStrandsCore/Private/GroomDebug.cpp
+++ b/Engine/Plugins/Runtime/HairStrands/Source/HairStrandsCore/Private/GroomDebug.cpp
@@ -981,7 +981,8 @@ static void AddHairDebugPrintInstancePass(
 
 				D.Data3.Z |= FFloat16(Instance->HairGroupPublicData->VFInput.Strands.Common.Length).Encoded;
 				D.Data3.Z |= FFloat16(Instance->HairGroupPublicData->VFInput.Strands.Common.LengthScale).Encoded << 16u;
-				D.Data3.W  = Instance->HairGroupPublicData->GetActiveStrandsCoverageScale();
+				D.Data3.W |= FFloat16(Instance->HairGroupPublicData->GetActiveStrandsCoverageScale()).Encoded;
+				D.Data3.W |= Instance->HairGroupPublicData->bAutoLOD ? (1u << 16u) : 0;
 			}
 			break;
 		case EHairGeometryType::Cards:
diff --git a/Engine/Plugins/Runtime/HairStrands/Source/HairStrandsCore/Private/GroomManager.cpp b/Engine/Plugins/Runtime/HairStrands/Source/HairStrandsCore/Private/GroomManager.cpp
--- a/Engine/Plugins/Runtime/HairStrands/Source/HairStrandsCore/Private/GroomManager.cpp
+++ b/Engine/Plugins/Runtime/HairStrands/Source/HairStrandsCore/Private/GroomManager.cpp
@@ -47,6 +47,11 @@ static FAutoConsoleVariableRef CVarHairStrands_InterpolationFrustumCullingEnable
 static int32 GHairStrands_Streaming = 0;
 static FAutoConsoleVariableRef CVarHairStrands_Streaming(TEXT("r.HairStrands.Streaming"), GHairStrands_Streaming, TEXT("Hair strands streaming toggle."), ECVF_RenderThreadSafe | ECVF_ReadOnly);
 
+static float GHairStrands_AutoLOD_Scale = 1.f;
+static float GHairStrands_AutoLOD_Bias = 0.f;
+static FAutoConsoleVariableRef CVarHairStrands_AutoLOD_Scale(TEXT("r.HairStrands.AutoLOD.Scale"), GHairStrands_Streaming, TEXT("Hair strands Auto LOD rate at which curves get decimated based on screen coverage."), ECVF_RenderThreadSafe | ECVF_ReadOnly);
+static FAutoConsoleVariableRef CVarHairStrands_AutoLOD_Bias(TEXT("r.HairStrands.AutoLOD.Bias"), GHairStrands_Streaming, TEXT("Hair strands Auto LOD screen size bias at which curves get decimated."), ECVF_RenderThreadSafe | ECVF_ReadOnly);
+
 EHairBufferSwapType GetHairSwapBufferType()
 {
 	switch (GHairStrands_SwapBufferType)
@@ -782,8 +787,9 @@ static float ComputeActiveCurveCoverageScale(uint32 InAvailableCurveCount, uint3
 
 static uint32 ComputeActiveCurveCount(float InScreenSize, uint32 InCurveCount)
 {
-	const float Power = 1.8f; // Per Groom settings?
-	const uint32 OutCurveCount = InCurveCount * FMath::Pow(FMath::Clamp(InScreenSize, 0.f, 1.0f), Power);
+	const float Power = FMath::Max(0.1f, GHairStrands_AutoLOD_Scale);
+	const float ScreenSizeBias = FMath::Clamp(GHairStrands_AutoLOD_Bias, 0.f, 1.f);
+	const uint32 OutCurveCount = InCurveCount * FMath::Pow(FMath::Clamp(InScreenSize + ScreenSizeBias, 0.f, 1.0f), Power);
 	return FMath::Clamp(OutCurveCount, 1, InCurveCount);
 }
 
@@ -883,7 +889,16 @@ static void RunHairLODSelection(
 
 			if (Instance->Strands.ClusterCullingResource)
 			{
-				const uint32 EffectiveCurveCount = Instance->Strands.ClusterCullingResource->BulkData.GetCurveCount(LODIndex);				            
+				uint32 EffectiveCurveCount = 0;
+				if (Instance->HairGroupPublicData->bAutoLOD)
+				{
+					EffectiveCurveCount = ComputeActiveCurveCount(Instance->HairGroupPublicData->ContinuousLODScreenSize, Instance->HairGroupPublicData->RestCurveCount);
+					LODIndex = 0;
+				}
+				else
+				{
+					EffectiveCurveCount = Instance->Strands.ClusterCullingResource->BulkData.GetCurveCount(LODIndex);
+				}
 				check(EffectiveCurveCount <= uint32(Instance->Strands.Data->Header.CurveToPointCount.Num()));
 
 				Instance->HairGroupPublicData->ContinuousLODCurveCount = EffectiveCurveCount;
diff --git a/Engine/Plugins/Runtime/HairStrands/Source/HairStrandsCore/Public/GroomAssetInterpolation.h b/Engine/Plugins/Runtime/HairStrands/Source/HairStrandsCore/Public/GroomAssetInterpolation.h
--- a/Engine/Plugins/Runtime/HairStrands/Source/HairStrandsCore/Public/GroomAssetInterpolation.h
+++ b/Engine/Plugins/Runtime/HairStrands/Source/HairStrandsCore/Public/GroomAssetInterpolation.h
@@ -50,7 +50,6 @@ enum class EGroomOverrideType : uint8
 	Disable UMETA(DisplayName = "Disable", ToolTip = "Override the asset value, and force disabled")
 };
 
-
 UENUM(BlueprintType)
 enum class EGroomGuideType : uint8
 {
@@ -59,6 +58,13 @@ enum class EGroomGuideType : uint8
 	Rigged 		UMETA(DisplayName = "Rigged  Guides", ToolTip = "Generated rigged guides from imported strands")
 };
 
+UENUM(BlueprintType)
+enum class EGroomLODType : uint8
+{
+	Manual 		UMETA(DisplayName = "Manual", ToolTip = "Hair strands curves & points adapt based on LOD settings"),
+	Auto		UMETA(DisplayName = "Auto (Experimental)", ToolTip = "Hair strands curves & points automatically adapt based on screen coverage. LOD settings are ignored.")
+};
+
 USTRUCT(BlueprintType)
 struct HAIRSTRANDSCORE_API FHairLODSettings
 {
@@ -214,6 +220,9 @@ struct HAIRSTRANDSCORE_API FHairGroupsLOD
 {
 	GENERATED_USTRUCT_BODY()
 
+	UPROPERTY(EditAnywhere, Category = "InterpolationSettings", meta = (ToolTip = "Define how LOD adapts curves & points for a given group."))
+	EGroomLODType LODType = EGroomLODType::Manual;
+
 	/** LODs  */
 	UPROPERTY(EditAnywhere, Category = "DecimationSettings", meta = (ToolTip = "Reduce the number of hair strands in a uniform manner", ClampMin = "0", UIMin = "0", UIMax = "1.0"))
 	TArray<FHairLODSettings> LODs;
diff --git a/Engine/Plugins/Runtime/HairStrands/Source/HairStrandsEditor/Private/GroomAssetDetails.cpp b/Engine/Plugins/Runtime/HairStrands/Source/HairStrandsEditor/Private/GroomAssetDetails.cpp
--- a/Engine/Plugins/Runtime/HairStrands/Source/HairStrandsEditor/Private/GroomAssetDetails.cpp
+++ b/Engine/Plugins/Runtime/HairStrands/Source/HairStrandsEditor/Private/GroomAssetDetails.cpp
@@ -1084,16 +1084,19 @@ void FGroomRenderingDetails::ExpandStructForLOD(TSharedRef<IPropertyHandle>& Pro
 			continue;
 		}
 
+		const bool bAutoLOD = GroomAsset->HairGroupsLOD[GroupIndex].LODType == EGroomLODType::Auto;
 		if (bOverrideReset)
 		{
 			FIsResetToDefaultVisible IsResetVisible = FIsResetToDefaultVisible::CreateSP(this, &FGroomRenderingDetails::ShouldResetToDefault, GroupIndex, LODIndex);
 			FResetToDefaultHandler ResetHandler = FResetToDefaultHandler::CreateSP(this, &FGroomRenderingDetails::ResetToDefault, GroupIndex, LODIndex);
 			FResetToDefaultOverride ResetOverride = FResetToDefaultOverride::Create(IsResetVisible, ResetHandler);
-			ChildrenBuilder.AddProperty(ChildHandle.ToSharedRef()).OverrideResetToDefault(ResetOverride);
+			IDetailPropertyRow& Row = ChildrenBuilder.AddProperty(ChildHandle.ToSharedRef()).OverrideResetToDefault(ResetOverride);
+			Row.IsEnabled(!bAutoLOD);
 		}
 		else
 		{
-			ChildrenBuilder.AddProperty(ChildHandle.ToSharedRef());
+			IDetailPropertyRow& Row = ChildrenBuilder.AddProperty(ChildHandle.ToSharedRef());
+			Row.IsEnabled(!bAutoLOD);
 		}
 	}
 }
@@ -1756,6 +1759,12 @@ void FGroomRenderingDetails::OnGenerateElementForHairGroup(TSharedRef<IPropertyH
 		break;
 		case EMaterialPanelType::LODs:
 		{
+			// Reload layout when the LOD type change
+			if (PropertyName == GET_MEMBER_NAME_CHECKED(FHairGroupsLOD, LODType))
+			{
+				ChildHandle->SetOnPropertyValueChanged(FSimpleDelegate::CreateSP(this, &FGroomRenderingDetails::OnGroomLODTypeChanged));
+			}
+
 			if (PropertyName == GET_MEMBER_NAME_CHECKED(FHairGroupsLOD, LODs))
 			{
 				// Add a custom builder for each LOD arrays within each group. This way we can customize this 'nested/inner' array
@@ -1812,6 +1821,11 @@ void FGroomRenderingDetails::OnGenerateElementForHairGroup(TSharedRef<IPropertyH
 	}
 }
 
+void FGroomRenderingDetails::OnGroomLODTypeChanged()
+{
+	GroomDetailLayout->ForceRefreshDetails();
+}
+
 // Hair binding display
 void FGroomRenderingDetails::OnGenerateElementForBindingAsset(TSharedRef<IPropertyHandle> StructProperty, int32 BindingIndex, IDetailChildrenBuilder& ChildrenBuilder, IDetailLayoutBuilder* DetailLayout)
 {
diff --git a/Engine/Plugins/Runtime/HairStrands/Source/HairStrandsEditor/Private/GroomAssetDetails.h b/Engine/Plugins/Runtime/HairStrands/Source/HairStrandsEditor/Private/GroomAssetDetails.h
--- a/Engine/Plugins/Runtime/HairStrands/Source/HairStrandsEditor/Private/GroomAssetDetails.h
+++ b/Engine/Plugins/Runtime/HairStrands/Source/HairStrandsEditor/Private/GroomAssetDetails.h
@@ -87,6 +87,7 @@ private:
 	FReply OnSaveCards(int32 GroupIndex, FProperty* Property);
 	FReply OnGenerateCardDataUsingPlugin(int32 GroupIndex);
 	FReply OnSelectBinding(int32 BindingIndex, FProperty* Property);
+	void OnGroomLODTypeChanged();
 
 	void AddPropertySeparator(FName PropertyName, IDetailChildrenBuilder& ChildrenBuilder);
 	void ExpandStructForLOD(TSharedRef<IPropertyHandle>& PropertyHandle, IDetailChildrenBuilder& ChildrenBuilder, int32 GroupIndex, int32 LODIndex, bool bOverrideReset);
diff --git a/Engine/Shaders/Private/HairStrands/HairStrandsDebug.usf b/Engine/Shaders/Private/HairStrands/HairStrandsDebug.usf
--- a/Engine/Shaders/Private/HairStrands/HairStrandsDebug.usf
+++ b/Engine/Shaders/Private/HairStrands/HairStrandsDebug.usf
@@ -1201,6 +1201,7 @@ struct FInstanceInfo
 	bool CacheRen;
 	bool CacheSim;
 	bool bHasLODSwitched;
+	bool bAutoLOD;
 
 	// Binding
 	uint UniqueTriangle;
@@ -1279,6 +1280,7 @@ FInstanceInfo ReadInstanceInfo(uint Index)
 	Out.ActivePointCount		= Data2.y;
 	Out.ActiveCurveCount		= Data2.z;
 	Out.Attributes				= Data2.w;
+	Out.bAutoLOD				= Data3.w >> 16u;
 
 	Out.Radius					= f16tof32(Data3.x & 0xFFFF);
 	Out.Density					= f16tof32(Data3.x >> 16u);
@@ -1286,7 +1288,7 @@ FInstanceInfo ReadInstanceInfo(uint Index)
 	Out.TipScale				= f16tof32(Data3.y >> 16u);
 	Out.Length					= f16tof32(Data3.z & 0xFFFF);
 	Out.LengthScale				= f16tof32(Data3.z >> 16u);
-	Out.CoverageScale			= Data3.w;
+	Out.CoverageScale			= f16tof32(Data3.w & 0xFFFF);
 
 	Out.bHasLODSwitched			= Data4.x > 0;
 
@@ -1424,7 +1426,7 @@ void MainCS(uint GroupIndex : SV_GroupIndex, uint3 DispatchThreadId : SV_Dispatc
 	Newline(Context);
 
 	// Draw instance details
-	Print(Context, TEXT("   Group  LOD     Geometry Binding Cache Sim RBF Clip Screen CovSc. Strands            Points             Bind.Tri/Sect Vis. Shad. SwLOD Name"));
+	Print(Context, TEXT("   Group  LOD     Geometry Binding Cache Sim RBF Clip Screen CovSc. Strands            Points             Bind.Tri/Sect Vis. Shad. SwLOD ALOD Name"));
 	Newline(Context);
 
 	for (uint InstanceIndex = 0; InstanceIndex < InstanceCount; ++InstanceIndex)
@@ -1516,6 +1518,10 @@ void MainCS(uint GroupIndex : SV_GroupIndex, uint3 DispatchThreadId : SV_Dispatc
 		PrintBool(Context, Info.bHasLODSwitched);
 		Print(Context, TEXT("  "));
 
+		// Auto LOD
+		PrintBool(Context, Info.bAutoLOD);
+		Print(Context, TEXT("  "));
+		
 		// Name
 		PrintInstanceName(Context, InstanceIndex, FontWhite);
 
diff --git a/Engine/Source/Runtime/Renderer/Public/HairStrandsInterface.h b/Engine/Source/Runtime/Renderer/Public/HairStrandsInterface.h
--- a/Engine/Source/Runtime/Renderer/Public/HairStrandsInterface.h
+++ b/Engine/Source/Runtime/Renderer/Public/HairStrandsInterface.h
@@ -410,6 +410,7 @@ public:
 	float LODIndex = -1;		// Current LOD used for all views
 	float LODBias = 0;			// Current LOD bias
 	bool bLODVisibility = true; // Enable/disable hair rendering for this component
+	bool bAutoLOD = false;
 
 	// Active/used point/curved based on select continuous LOD
 	uint32 ContinuousLODPointCount = 0;
