Allocate deformed resources based on requested point/curve count.

diff --git a/Engine/Plugins/Runtime/HairStrands/Source/HairStrandsCore/Private/GroomManager.cpp b/Engine/Plugins/Runtime/HairStrands/Source/HairStrandsCore/Private/GroomManager.cpp
--- a/Engine/Plugins/Runtime/HairStrands/Source/HairStrandsCore/Private/GroomManager.cpp
+++ b/Engine/Plugins/Runtime/HairStrands/Source/HairStrandsCore/Private/GroomManager.cpp
@@ -29,6 +29,8 @@
 #include "HairStrandsClusterCulling.h"
 #include "GroomVisualizationData.h"
 
+UE_DISABLE_OPTIMIZATION
+
 static int32 GHairStrandsMinLOD = 0;
 static FAutoConsoleVariableRef CVarGHairStrandsMinLOD(TEXT("r.HairStrands.MinLOD"), GHairStrandsMinLOD, TEXT("Clamp the min hair LOD to this value, preventing to reach lower/high-quality LOD."), ECVF_Scalability);
 
@@ -1037,7 +1039,7 @@ static void RunHairLODSelection(
 				if (Instance->Strands.InterpolationResource)	{ Instance->Strands.InterpolationResource->Allocate(GraphBuilder, EHairResourceLoadingType::Async, ResourceStatus, RequestedCurveCount, RequestedPointCount); }
 
 				if (Instance->Strands.DeformedRootResource)		{ Instance->Strands.DeformedRootResource->Allocate(GraphBuilder, LoadingType, ResourceStatus, MeshLODIndex); }
-				if (Instance->Strands.DeformedResource)			{ Instance->Strands.DeformedResource->Allocate(GraphBuilder, LoadingType, ResourceStatus); }
+				if (Instance->Strands.DeformedResource)			{ Instance->Strands.DeformedResource->Allocate(GraphBuilder, EHairResourceLoadingType::Async, ResourceStatus, RequestedCurveCount, RequestedPointCount); }
 				#if RHI_RAYTRACING
 				if (bHasPathTracingView)						{ AllocateRaytracingResources(Instance); }
 				if (Instance->Strands.RenRaytracingResource)	{ Instance->Strands.RenRaytracingResource->Allocate(GraphBuilder, LoadingType, ResourceStatus); }
diff --git a/Engine/Plugins/Runtime/HairStrands/Source/HairStrandsCore/Private/GroomResources.cpp b/Engine/Plugins/Runtime/HairStrands/Source/HairStrandsCore/Private/GroomResources.cpp
--- a/Engine/Plugins/Runtime/HairStrands/Source/HairStrandsCore/Private/GroomResources.cpp
+++ b/Engine/Plugins/Runtime/HairStrands/Source/HairStrandsCore/Private/GroomResources.cpp
@@ -35,6 +35,8 @@
 #include "DerivedDataRequestOwner.h"
 #endif
 
+UE_DISABLE_OPTIMIZATION
+
 static int32 GHairStrandsBulkData_ReleaseAfterUse = 0;
 static FAutoConsoleVariableRef CVarHairStrandsBulkData_ReleaseAfterUse(TEXT("r.HairStrands.Strands.BulkData.ReleaseAfterUse"), GHairStrandsBulkData_ReleaseAfterUse, TEXT("Release CPU bulk data once hair groom/groom binding asset GPU resources are created. This saves memory"));
 
@@ -1103,19 +1105,22 @@ FHairStrandsDeformedResource::FHairStrandsDeformedResource(FHairStrandsBulkData&
 	GetPositionOffset(FHairStrandsDeformedResource::EFrameType::Previous) = BulkData.GetPositionOffset();
 }
 
-void FHairStrandsDeformedResource::InternalAllocate(FRDGBuilder& GraphBuilder)
+void FHairStrandsDeformedResource::InternalAllocate(FRDGBuilder& GraphBuilder, uint32 InCurveCount, uint32 InPointCount, int32 InLODIndex)
 {
-	const uint32 PointCount = BulkData.GetNumPoints();
+	// Need to update these to by dynamic based on the active number of point
+	const uint32 PointCount = FMath::Min(InPointCount, BulkData.GetNumPoints());
 
 	InternalCreateVertexBufferRDG<FHairStrandsPositionFormat>(GraphBuilder, PointCount, DeformedPositionBuffer[0], ToHairResourceDebugName(HAIRSTRANDS_RESOUCE_NAME(CurveType, Hair.StrandsDeformed_DeformedPositionBuffer0), ResourceName), OwnerName, EHairResourceUsageType::Dynamic);
 	InternalCreateVertexBufferRDG<FHairStrandsPositionFormat>(GraphBuilder, PointCount, DeformedPositionBuffer[1], ToHairResourceDebugName(HAIRSTRANDS_RESOUCE_NAME(CurveType, Hair.StrandsDeformed_DeformedPositionBuffer1), ResourceName), OwnerName, EHairResourceUsageType::Dynamic);
 	InternalCreateVertexBufferRDG<FHairStrandsTangentFormat>(GraphBuilder,  PointCount * FHairStrandsTangentFormat::ComponentCount, TangentBuffer, ToHairResourceDebugName(HAIRSTRANDS_RESOUCE_NAME(CurveType, Hair.StrandsDeformed_TangentBuffer), ResourceName), OwnerName, EHairResourceUsageType::Dynamic);
 
-	TArray<FVector4f> DefaultOffsets;
-	DefaultOffsets.Add((FVector3f)BulkData.GetPositionOffset()); // LWC_TODO: precision loss
-	InternalCreateVertexBufferRDG<FHairStrandsPositionOffsetFormat>(GraphBuilder, DefaultOffsets, DeformedOffsetBuffer[0], ToHairResourceDebugName(HAIRSTRANDS_RESOUCE_NAME(CurveType, Hair.StrandsDeformed_DeformedOffsetBuffer0), ResourceName), OwnerName, EHairResourceUsageType::Dynamic, ERDGInitialDataFlags::None);
-	InternalCreateVertexBufferRDG<FHairStrandsPositionOffsetFormat>(GraphBuilder, DefaultOffsets, DeformedOffsetBuffer[1], ToHairResourceDebugName(HAIRSTRANDS_RESOUCE_NAME(CurveType, Hair.StrandsDeformed_DeformedOffsetBuffer1), ResourceName), OwnerName, EHairResourceUsageType::Dynamic, ERDGInitialDataFlags::None);
-
+	if (DeformedOffsetBuffer[0].Buffer == nullptr)
+	{
+		TArray<FVector4f> DefaultOffsets;
+		DefaultOffsets.Add((FVector3f)BulkData.GetPositionOffset()); // LWC_TODO: precision loss
+		InternalCreateVertexBufferRDG<FHairStrandsPositionOffsetFormat>(GraphBuilder, DefaultOffsets, DeformedOffsetBuffer[0], ToHairResourceDebugName(HAIRSTRANDS_RESOUCE_NAME(CurveType, Hair.StrandsDeformed_DeformedOffsetBuffer0), ResourceName), OwnerName, EHairResourceUsageType::Dynamic, ERDGInitialDataFlags::None);
+		InternalCreateVertexBufferRDG<FHairStrandsPositionOffsetFormat>(GraphBuilder, DefaultOffsets, DeformedOffsetBuffer[1], ToHairResourceDebugName(HAIRSTRANDS_RESOUCE_NAME(CurveType, Hair.StrandsDeformed_DeformedOffsetBuffer1), ResourceName), OwnerName, EHairResourceUsageType::Dynamic, ERDGInitialDataFlags::None);
+	}
 	// Note: DeformerBuffer is optionally/lazily allocated by a mesh-deformer graph
 }
 
diff --git a/Engine/Plugins/Runtime/HairStrands/Source/HairStrandsCore/Public/GroomResources.h b/Engine/Plugins/Runtime/HairStrands/Source/HairStrandsCore/Public/GroomResources.h
--- a/Engine/Plugins/Runtime/HairStrands/Source/HairStrandsCore/Public/GroomResources.h
+++ b/Engine/Plugins/Runtime/HairStrands/Source/HairStrandsCore/Public/GroomResources.h
@@ -204,7 +204,7 @@ struct FHairStrandsDeformedResource : public FHairCommonResource
 	FHairStrandsDeformedResource(FHairStrandsBulkData& BulkData, EHairStrandsResourcesType CurveType, const FHairResourceName& ResourceName, const FName& OwnerName);
 
 	/* Init/Release buffers */
-	virtual void InternalAllocate(FRDGBuilder& GraphBuilder) override;
+	virtual void InternalAllocate(FRDGBuilder& GraphBuilder, uint32 InCurveCount, uint32 InPointCount, int32 InLODIndex) override;
 	virtual void InternalRelease() override;
 
 	/* Get the resource name */
