Fix hair access to deprecated members.

diff --git a/Engine/Plugins/Runtime/HairStrands/Source/HairStrandsCore/Private/GroomComponent.cpp b/Engine/Plugins/Runtime/HairStrands/Source/HairStrandsCore/Private/GroomComponent.cpp
--- a/Engine/Plugins/Runtime/HairStrands/Source/HairStrandsCore/Private/GroomComponent.cpp
+++ b/Engine/Plugins/Runtime/HairStrands/Source/HairStrandsCore/Private/GroomComponent.cpp
@@ -2600,7 +2600,7 @@ void UGroomComponent::InitResources(bool bIsBindingReloading)
 			HairGroupInstance->HairGroupPublicData->SetLODScreenSizes(CPULODScreenSize);
 			HairGroupInstance->HairGroupPublicData->SetLODVisibilities(LODVisibility);
 			HairGroupInstance->HairGroupPublicData->SetLODGeometryTypes(LODGeometryTypes);
-			HairGroupInstance->HairGroupPublicData->bAutoLOD = GroomAsset->HairGroupsLOD[GroupIt].LODType == EGroomLODType::Auto;
+			HairGroupInstance->HairGroupPublicData->bAutoLOD = GroomAsset->GetHairGroupsLOD()[GroupIt].LODType == EGroomLODType::Auto;
 		}
 
 		FHairResourceName ResourceName(GetFName(), GroupIt);
diff --git a/Engine/Plugins/Runtime/HairStrands/Source/HairStrandsEditor/Private/GroomAssetDetails.cpp b/Engine/Plugins/Runtime/HairStrands/Source/HairStrandsEditor/Private/GroomAssetDetails.cpp
--- a/Engine/Plugins/Runtime/HairStrands/Source/HairStrandsEditor/Private/GroomAssetDetails.cpp
+++ b/Engine/Plugins/Runtime/HairStrands/Source/HairStrandsEditor/Private/GroomAssetDetails.cpp
@@ -1084,7 +1084,7 @@ void FGroomRenderingDetails::ExpandStructForLOD(TSharedRef<IPropertyHandle>& Pro
 			continue;
 		}
 
-		const bool bAutoLOD = GroomAsset->HairGroupsLOD[GroupIndex].LODType == EGroomLODType::Auto;
+		const bool bAutoLOD = GroomAsset->GetHairGroupsLOD()[GroupIndex].LODType == EGroomLODType::Auto;
 		if (bOverrideReset)
 		{
 			FIsResetToDefaultVisible IsResetVisible = FIsResetToDefaultVisible::CreateSP(this, &FGroomRenderingDetails::ShouldResetToDefault, GroupIndex, LODIndex);
