Fix crashes due to inconsistent usage of the point counter buffer.

The cluster culling shader was expecting a typed buffer, but everything else was using a byte address buffer.

#rnx

diff --git a/Engine/Plugins/Runtime/HairStrands/Source/HairStrandsCore/Private/HairStrandsClusterCulling.cpp b/Engine/Plugins/Runtime/HairStrands/Source/HairStrandsCore/Private/HairStrandsClusterCulling.cpp
--- a/Engine/Plugins/Runtime/HairStrands/Source/HairStrandsCore/Private/HairStrandsClusterCulling.cpp
+++ b/Engine/Plugins/Runtime/HairStrands/Source/HairStrandsCore/Private/HairStrandsClusterCulling.cpp
@@ -121,8 +121,8 @@ static void AddClusterCullingPass(
 	ShaderPrint::RequestSpaceForLines(2048);
 
 	// 0. Glogal counter for visible points
-	FRDGBufferRef PointCounter 	= GraphBuilder.CreateBuffer(FRDGBufferDesc::CreateByteAddressDesc(4), TEXT("Hair.ClusterPointCounter"));
-	FRDGBufferUAVRef PointCounterUAV = GraphBuilder.CreateUAV(PointCounter);
+	FRDGBufferRef PointCounter 	= GraphBuilder.CreateBuffer(FRDGBufferDesc::CreateBufferDesc(sizeof(uint32), 1), TEXT("Hair.ClusterPointCounter"));
+	FRDGBufferUAVRef PointCounterUAV = GraphBuilder.CreateUAV(PointCounter, PF_R32_UINT);
 	AddClearUAVPass(GraphBuilder, PointCounterUAV, 0);
 
 	// 1. Build culled index buffer
@@ -182,7 +182,7 @@ static void AddClusterCullingPass(
 		FRDGImportedBuffer DrawIndirectParametersRasterComputeBuffer = Register(GraphBuilder, ClusterData.HairGroupPublicPtr->GetDrawIndirectRasterComputeBuffer(), ERDGImportedBufferFlags::CreateViews);
 
 		FHairClusterCullArgsCS::FParameters* Parameters = GraphBuilder.AllocParameters<FHairClusterCullArgsCS::FParameters>();
-		Parameters->PointCounterBuffer = GraphBuilder.CreateSRV(PointCounter);
+		Parameters->PointCounterBuffer = GraphBuilder.CreateSRV(PointCounter, PF_R32_UINT);
 		Parameters->RWIndirectDrawArgsBuffer = DrawIndirectParametersBuffer.UAV;
 		Parameters->RWIndirectDispatchArgsBuffer = DrawIndirectParametersRasterComputeBuffer.UAV;
 
diff --git a/Engine/Shaders/Private/HairStrands/HairStrandsClusterCulling.usf b/Engine/Shaders/Private/HairStrands/HairStrandsClusterCulling.usf
--- a/Engine/Shaders/Private/HairStrands/HairStrandsClusterCulling.usf
+++ b/Engine/Shaders/Private/HairStrands/HairStrandsClusterCulling.usf
@@ -607,14 +607,14 @@ void Main(uint3 DispatchThreadId : SV_DispatchThreadID, uint ThreadIndex : SV_Gr
 
 #if SHADER_CLUSTER_CULL_ARGS
 
-ByteAddressBuffer PointCounterBuffer;
+Buffer<uint> PointCounterBuffer;
 RWBuffer<uint> RWIndirectDrawArgsBuffer;
 RWBuffer<uint> RWIndirectDispatchArgsBuffer;
 
 [numthreads(1, 1, 1)]
 void Main(uint DispatchThreadId : SV_DispatchThreadID, uint ThreadIndex : SV_GroupIndex)
 {
-	const uint PointCount = PointCounterBuffer.Load(0);
+	const uint PointCount = PointCounterBuffer[0];
 	const uint DispatchCount = DivideAndRoundUp(PointCount, HAIR_VERTEXCOUNT_GROUP_SIZE);
 
 	RWIndirectDrawArgsBuffer[0] = PointCount * HAIR_POINT_TO_VERTEX;
