Change Groom instance print to use ShaderPrint's strings.

diff --git a/Engine/Plugins/Runtime/HairStrands/Source/HairStrandsCore/Private/GroomDebug.cpp b/Engine/Plugins/Runtime/HairStrands/Source/HairStrandsCore/Private/GroomDebug.cpp
--- a/Engine/Plugins/Runtime/HairStrands/Source/HairStrandsCore/Private/GroomDebug.cpp
+++ b/Engine/Plugins/Runtime/HairStrands/Source/HairStrandsCore/Private/GroomDebug.cpp
@@ -18,6 +18,7 @@
 #include "GroomVisualizationData.h"
 #include "GroomCacheData.h"
 #include "HairStrandsClusterCulling.h"
+#include "ShaderPrintParameters.h"
 
 ///////////////////////////////////////////////////////////////////////////////////////////////////
 
@@ -818,17 +819,8 @@ class FHairDebugPrintInstanceCS : public FGlobalShader
 		SHADER_PARAMETER(uint32, InstanceCount_CardsOrMeshesPrimaryView)
 		SHADER_PARAMETER(uint32, InstanceCount_CardsOrMeshesShadowView)
 		SHADER_PARAMETER(FVector4f, InstanceScreenSphereBound)
-		// Instance names
-		SHADER_PARAMETER(uint32, InstanceNameInfoCount)
-		SHADER_PARAMETER(uint32, InstanceNameCharacterCount)
-		SHADER_PARAMETER_RDG_BUFFER_SRV(StructuredBuffer<uint2>, InstanceNameInfos)
-		SHADER_PARAMETER_RDG_BUFFER_SRV(Buffer<uint8>, InstanceNames)
-		// Attribute names
-		SHADER_PARAMETER(uint32, AttributeNameInfoCount)
-		SHADER_PARAMETER(uint32, AttributeNameCharacterCount)
-		SHADER_PARAMETER_RDG_BUFFER_SRV(StructuredBuffer<uint2>, AttributeNameInfos)
-		SHADER_PARAMETER_RDG_BUFFER_SRV(Buffer<uint8>, AttributeNames)
-
+		SHADER_PARAMETER_STRUCT(ShaderPrint::FStrings::FShaderParameters, InstanceNames)
+		SHADER_PARAMETER_STRUCT(ShaderPrint::FStrings::FShaderParameters, AttributeNames)
 		SHADER_PARAMETER_RDG_BUFFER_SRV(StructuredBuffer<uint4>, Infos)
 		SHADER_PARAMETER_RDG_BUFFER_SRV(Buffer<int>, InstanceAABB)
 		SHADER_PARAMETER_STRUCT_INCLUDE(ShaderPrint::FShaderParameters, ShaderPrintUniformBuffer)
@@ -847,14 +839,6 @@ public:
 
 IMPLEMENT_GLOBAL_SHADER(FHairDebugPrintInstanceCS, "/Engine/Private/HairStrands/HairStrandsDebug.usf", "MainCS", SF_Compute);
 
-struct FHairDebugNameInfo
-{
-	uint32 PrimitiveID;
-	uint16 Offset;
-	uint8  Length;
-	uint8  Pad0;
-};
-
 const TCHAR* GetHairAttributeText(EHairAttribute In, uint32 InFlags);
 uint32 GetHairAttributeIndex(EHairAttribute In);
 EGroomCacheType GetHairInstanceCacheType(const FHairGroupInstance* Instance);
@@ -877,13 +861,8 @@ static void AddHairDebugPrintInstancePass(
 
 	if (!ShaderPrintData || InstanceCount == 0) { return; }
 
-	const uint32 MaxPrimitiveNameCount = 128u;
-	check(sizeof(FHairDebugNameInfo) == 8);
-
-	TArray<FHairDebugNameInfo> InstanceNameInfos;
-	TArray<uint8> InstanceNames;
-	InstanceNames.Reserve(MaxPrimitiveNameCount * 30u);
 
+	ShaderPrint::FStrings InstanceNames(InstanceCount);
 	struct FInstanceInfos
 	{
 		FUintVector4 Data0 = { 0,0,0,0 };
@@ -896,25 +875,10 @@ static void AddHairDebugPrintInstancePass(
 	Infos.Reserve(InstanceCount);
 	for (uint32 InstanceIndex = 0; InstanceIndex < InstanceCount; ++InstanceIndex)
 	{
-		const FHairStrandsInstance* AbstractInstance = Instances[InstanceIndex];
-		const FHairGroupInstance* Instance = static_cast<const FHairGroupInstance*>(AbstractInstance);
+		const FHairGroupInstance* Instance = static_cast<const FHairGroupInstance*>(Instances[InstanceIndex]);
 
 		// Collect InstanceNames
-		if (InstanceIndex < MaxPrimitiveNameCount)
-		{
-			const FString Name = *Instance->Debug.GroomAssetName;
-			const uint32 NameOffset = InstanceNames.Num();
-			const uint32 NameLength = Name.Len();
-			for (TCHAR C : Name)
-			{
-				InstanceNames.Add(uint8(C));
-			}
-
-			FHairDebugNameInfo& NameInfo = InstanceNameInfos.AddDefaulted_GetRef();
-			NameInfo.PrimitiveID = InstanceIndex;
-			NameInfo.Length = NameLength;
-			NameInfo.Offset = NameOffset;
-		}
+		InstanceNames.Add(*Instance->Debug.GroomAssetName, InstanceIndex);
 
 		const float LODIndex = Instance->HairGroupPublicData->LODIndex;
 		const uint32 IntLODIndex = Instance->HairGroupPublicData->LODIndex;
@@ -1010,51 +974,19 @@ static void AddHairDebugPrintInstancePass(
 		}
 	}
 
-	if (InstanceNameInfos.IsEmpty())
-	{
-		FHairDebugNameInfo& NameInfo = InstanceNameInfos.AddDefaulted_GetRef();
-		NameInfo.PrimitiveID = ~0;
-		NameInfo.Length = 4;
-		NameInfo.Offset = 0;
-		InstanceNames.Add(uint8('N'));
-		InstanceNames.Add(uint8('o'));
-		InstanceNames.Add(uint8('n'));
-		InstanceNames.Add(uint8('e'));
-	}	
-
-	TArray<FHairDebugNameInfo> AttributeNameInfos;
-	TArray<uint8> AttributeNames;
-	AttributeNames.Reserve(uint32(EHairAttribute::Count) * 30u);
+	ShaderPrint::FStrings AttributeNames;
 	for (uint32 AttributeIt = 0; AttributeIt < uint32(EHairAttribute::Count); ++AttributeIt)
 	{
 		// Only add valid optional attribute 
 		// HAIR_ATTRIBUTE_XXX and EHairAttribute don't have a 1:1 mapping
 		const EHairAttribute Attribute = (EHairAttribute)AttributeIt;
-		{
-			const FString Name = GetHairAttributeText(Attribute, 0u);
-			FHairDebugNameInfo& NameInfo = AttributeNameInfos.AddDefaulted_GetRef();
-			NameInfo.PrimitiveID = ~0;
-			NameInfo.Length = Name.Len();
-			NameInfo.Offset = AttributeNames.Num();	
-			for (TCHAR C : Name)
-			{
-				AttributeNames.Add(uint8(C));
-			}
-		}
+		AttributeNames.Add(GetHairAttributeText(Attribute, 0u), AttributeIt);
 	}
 
 	const uint32 InfoInBytes  = sizeof(FInstanceInfos);
 	const uint32 InfoInUints = sizeof(FInstanceInfos) / sizeof(uint32);
 	FRDGBufferRef InfoBuffer = CreateVertexBuffer(GraphBuilder, TEXT("Hair.Debug.InstanceInfos"), FRDGBufferDesc::CreateBufferDesc(4, InfoInUints * Infos.Num()), Infos.GetData(), InfoInBytes * Infos.Num());
 
-	// Resource for instance names
-	FRDGBufferRef InstanceNameBuffer = CreateVertexBuffer(GraphBuilder, TEXT("Hair.Debug.InstanceNames"), FRDGBufferDesc::CreateBufferDesc(1, InstanceNames.Num()), InstanceNames.GetData(), InstanceNames.Num());
-	FRDGBufferRef InstanceNameInfoBuffer = CreateStructuredBuffer(GraphBuilder, TEXT("Hair.Debug.InstanceNameInfos"), InstanceNameInfos);
-
-	// Resource for attribute names
-	FRDGBufferRef AttributeNameBuffer = CreateVertexBuffer(GraphBuilder, TEXT("Hair.Debug.AttributeNames"), FRDGBufferDesc::CreateBufferDesc(1, AttributeNames.Num()), AttributeNames.GetData(), AttributeNames.Num());
-	FRDGBufferRef AttributeNameInfoBuffer = CreateStructuredBuffer(GraphBuilder, TEXT("Hair.Debug.AttributeNameInfos"), AttributeNameInfos);
-
 	// Draw general information for all instances (one pass for all instances)
 	{
 		FHairDebugPrintInstanceCS::FParameters* Parameters = GraphBuilder.AllocParameters<FHairDebugPrintInstanceCS::FParameters>();
@@ -1063,19 +995,8 @@ static void AddHairDebugPrintInstancePass(
 		Parameters->InstanceCount_StrandsShadowView = InstanceCountPerType[HairInstanceCount_StrandsShadowView];
 		Parameters->InstanceCount_CardsOrMeshesPrimaryView = InstanceCountPerType[HairInstanceCount_CardsOrMeshesPrimaryView];
 		Parameters->InstanceCount_CardsOrMeshesShadowView = InstanceCountPerType[HairInstanceCount_CardsOrMeshesShadowView];
-
-		// Instance name
-		Parameters->InstanceNameInfoCount = InstanceNameInfos.Num();
-		Parameters->InstanceNameCharacterCount = InstanceNames.Num();
-		Parameters->InstanceNames = GraphBuilder.CreateSRV(InstanceNameBuffer, PF_R8_UINT);
-		Parameters->InstanceNameInfos = GraphBuilder.CreateSRV(InstanceNameInfoBuffer);
-
-		// Attribute name
-		Parameters->AttributeNameInfoCount = AttributeNameInfos.Num();
-		Parameters->AttributeNameCharacterCount = AttributeNames.Num();
-		Parameters->AttributeNames = GraphBuilder.CreateSRV(AttributeNameBuffer, PF_R8_UINT);
-		Parameters->AttributeNameInfos = GraphBuilder.CreateSRV(AttributeNameInfoBuffer);
-
+		Parameters->InstanceNames = InstanceNames.GetParameters(GraphBuilder);
+		Parameters->AttributeNames = AttributeNames.GetParameters(GraphBuilder);
 		Parameters->Infos = GraphBuilder.CreateSRV(InfoBuffer, PF_R32_UINT);
 		ShaderPrint::SetParameters(GraphBuilder, *ShaderPrintData, Parameters->ShaderPrintUniformBuffer);
 		FHairDebugPrintInstanceCS::FPermutationDomain PermutationVector;
@@ -1262,4 +1183,4 @@ void RunHairStrandsDebug(
 			}
 		}
 	}
-}
\ No newline at end of file
+}
diff --git a/Engine/Shaders/Private/HairStrands/HairStrandsDebug.usf b/Engine/Shaders/Private/HairStrands/HairStrandsDebug.usf
--- a/Engine/Shaders/Private/HairStrands/HairStrandsDebug.usf
+++ b/Engine/Shaders/Private/HairStrands/HairStrandsDebug.usf
@@ -1154,19 +1154,11 @@ uint					InstanceCount_CardsOrMeshesPrimaryView;
 uint					InstanceCount_CardsOrMeshesShadowView;
 
 Buffer<uint>			Infos;
-
-uint					InstanceNameInfoCount;
-uint					InstanceNameCharacterCount;
-StructuredBuffer<uint2> InstanceNameInfos;
-Buffer<uint>			InstanceNames;
 Buffer<int>				InstanceAABB;
-
 float4					InstanceScreenSphereBound;
 
-uint					AttributeNameInfoCount;
-uint					AttributeNameCharacterCount;
-StructuredBuffer<uint2> AttributeNameInfos;
-Buffer<uint>			AttributeNames;
+FSTRINGS(InstanceNames)
+FSTRINGS(AttributeNames)
 
 struct FInstanceInfo
 {
@@ -1300,70 +1292,6 @@ bool HasAttribute(FInstanceInfo In, uint InAttributeIt)
 	return (In.Attributes & (1u<< InAttributeIt)) != 0;
 }
 
-struct FNameInfo
-{
-	uint PrimitiveID;
-	uint Length;
-	uint Offset;
-	uint Pad0;
-	uint Pad1;
-};
-
-FNameInfo UnpackDebugNameInfo(uint2 In)
-{
-	FNameInfo Out = (FNameInfo)0;
-	Out.PrimitiveID = In.x;
-	Out.Offset = (In.y) & 0xFFFF;
-	Out.Length = (In.y >> 16) & 0xFF;
-
-	return Out;
-}
-
-FNameInfo FindNameInfo(uint PrimitiveID)
-{
-	FNameInfo Out = (FNameInfo)0;
-	Out.PrimitiveID = ~0;
-	for (uint It = 0; It < InstanceNameInfoCount; ++It)
-	{
-		if (InstanceNameInfos[It].x == PrimitiveID)
-		{
-			return UnpackDebugNameInfo(InstanceNameInfos[It]);
-		}
-	}
-	return Out;
-}
-
-void PrintInstanceName(inout FShaderPrintContext Context, uint PrimitiveID, FFontColor InColor)
-{
-	const FNameInfo Info = FindNameInfo(PrimitiveID);
-	if (Info.PrimitiveID != ~0 && (Info.Length + Info.Offset) <= InstanceNameCharacterCount)
-	{
-		for (uint It = 0; It < Info.Length; ++It)
-		{
-			const uint Char = InstanceNames[It + Info.Offset];
-			PrintSymbol(Context, Char, InColor);
-		}
-	}
-}
-
-void PrintAttributeName(inout FShaderPrintContext Context, uint AttributeIt, FFontColor InColor)
-{
-	if (AttributeIt >= AttributeNameInfoCount)
-	{
-		return;
-	}
-
-	const FNameInfo Info = UnpackDebugNameInfo(AttributeNameInfos[AttributeIt]);
-	if ((Info.Length + Info.Offset) <= AttributeNameCharacterCount)
-	{
-		for (uint It = 0; It < Info.Length; ++It)
-		{
-			const uint Char = AttributeNames[It + Info.Offset];
-			PrintSymbol(Context, Char, InColor);
-		}
-	}
-}
-
 [numthreads(1, 1, 1)]
 void MainCS(uint GroupIndex : SV_GroupIndex, uint3 DispatchThreadId : SV_DispatchThreadID)
 #if PERMUTATION_OUTPUT_TYPE == 0
@@ -1523,7 +1451,7 @@ void MainCS(uint GroupIndex : SV_GroupIndex, uint3 DispatchThreadId : SV_Dispatc
 		Print(Context, TEXT("  "));
 		
 		// Name
-		PrintInstanceName(Context, InstanceIndex, FontWhite);
+		PrintInstanceNames(Context, InstanceIndex, FontWhite);
 
 		Newline(Context);
 	}
@@ -1540,9 +1468,10 @@ void MainCS(uint GroupIndex : SV_GroupIndex, uint3 DispatchThreadId : SV_Dispatc
 		Print(Context, TEXT("Radius "), FontWhite);
 		Print(Context, TEXT("Root  "), FontWhite);
 		Print(Context, TEXT("Tip   "), FontWhite);
+		const uint AttributeNameInfoCount = NumAttributeNames();
 		for (uint AttributeIt=0; AttributeIt < AttributeNameInfoCount; ++AttributeIt)
 		{
-			PrintAttributeName(Context, AttributeIt, FontWhite);
+			PrintAttributeNames(Context, AttributeIt, FontWhite);
 			Print(Context, TEXT(" "));
 		}
 		Newline(Context);
@@ -1575,10 +1504,10 @@ void MainCS(uint GroupIndex : SV_GroupIndex, uint3 DispatchThreadId : SV_Dispatc
 			for (uint AttributeIt = 0; AttributeIt < AttributeNameInfoCount; ++AttributeIt)
 			{
 				PrintBool(Context, HasAttribute(Info, AttributeIt));
-				const FNameInfo Info = UnpackDebugNameInfo(AttributeNameInfos[AttributeIt]);
-				if (Info.Length > 3)
+				const uint AttributeLen = LenAttributeNames(AttributeIt);
+				if (AttributeLen > 3)
 				{
-					const uint SpaceCount = Info.Length - 3u;
+					const uint SpaceCount = AttributeLen - 3u;
 					for (uint SpaceIt = 0; SpaceIt < SpaceCount; ++SpaceIt)
 					{
 						Print(Context, TEXT(" "));
@@ -1586,7 +1515,7 @@ void MainCS(uint GroupIndex : SV_GroupIndex, uint3 DispatchThreadId : SV_Dispatc
 				}
 			}
 
-			PrintInstanceName(Context, StrandInstanceIndex, FontWhite);
+			PrintInstanceNames(Context, StrandInstanceIndex, FontWhite);
 
 			Newline(Context);
 		}
diff --git a/Engine/Shaders/Private/ShaderPrintCommon.ush b/Engine/Shaders/Private/ShaderPrintCommon.ush
--- a/Engine/Shaders/Private/ShaderPrintCommon.ush
+++ b/Engine/Shaders/Private/ShaderPrintCommon.ush
@@ -1395,7 +1395,7 @@ uint LenString(uint InEntryID, uint InInfoCount, StructuredBuffer<uint2> InInfoB
 	return Info.EntryID != ~0 ? Info.Length : 0;
 }
 
-#define SHADERPRINT_STRINGS(N) \
+#define FSTRINGS(N) \
 	uint					N##_InfoCount;\
 	uint					N##_CharCount;\
 	StructuredBuffer<uint2> N##_InfoBuffer;\
diff --git a/Engine/Source/Runtime/Renderer/Public/ShaderPrintParameters.h b/Engine/Source/Runtime/Renderer/Public/ShaderPrintParameters.h
--- a/Engine/Source/Runtime/Renderer/Public/ShaderPrintParameters.h
+++ b/Engine/Source/Runtime/Renderer/Public/ShaderPrintParameters.h
@@ -150,7 +150,7 @@ namespace ShaderPrint
 	//	 MyVariable.Add(FString(...), StringId);
 	//
 	// In shader:
-	//   SHADERPRINT_STRINGS(MyVariable)
+	//   FSTRINGS(MyVariable)
 	//   void foo()
 	//   {  
 	//     InitShaderPrintContext Ctx = InitShaderPrintContext(...);
