Change how hair Groom/Binding/Component memory stats are gathered.

diff --git a/Engine/Plugins/Runtime/HairStrands/Source/HairStrandsCore/Private/GroomAsset.cpp b/Engine/Plugins/Runtime/HairStrands/Source/HairStrandsCore/Private/GroomAsset.cpp
--- a/Engine/Plugins/Runtime/HairStrands/Source/HairStrandsCore/Private/GroomAsset.cpp
+++ b/Engine/Plugins/Runtime/HairStrands/Source/HairStrandsCore/Private/GroomAsset.cpp
@@ -132,6 +132,68 @@ uint32 FHairGroupPlatformData::FStrands::GetDataSize() const
 	return BTotal + Total + CTotal;
 }
 
+FGroomAssetMemoryStats FGroomAssetMemoryStats::Get(const FHairGroupPlatformData& In)
+{
+	FGroomAssetMemoryStats Out;
+	Out.CPU.Guides  = In.Guides.GetDataSize();
+	Out.CPU.Strands = In.Strands.GetDataSize();
+	Out.CPU.Cards   = In.Cards.GetDataSize();
+	Out.CPU.Meshes  = In.Meshes.GetDataSize();
+
+	Out.GPU.Guides  = In.Guides.GetResourcesSize();
+	Out.GPU.Strands = In.Strands.GetResourcesSize();
+	Out.GPU.Cards   = In.Cards.GetResourcesSize();
+	Out.GPU.Meshes  = In.Meshes.GetResourcesSize();
+
+	
+	// Strands only
+	if (In.Strands.RestResource) 			Out.Memory.Rest 			= In.Strands.RestResource->GetResourcesSize();
+	if (In.Strands.InterpolationResource) 	Out.Memory.Interpolation 	= In.Strands.InterpolationResource->GetResourcesSize();
+	if (In.Strands.ClusterCullingResource)	Out.Memory.Cluster 			= In.Strands.ClusterCullingResource->GetResourcesSize();
+#if RHI_RAYTRACING
+	if (In.Strands.RaytracingResource) 		Out.Memory.Raytracing 		= In.Strands.RaytracingResource->GetResourcesSize();
+#endif
+	if (In.Strands.RestResource) 			Out.Curves.Rest 			= In.Strands.RestResource->MaxAvailableCurveCount;
+	if (In.Strands.InterpolationResource) 	Out.Curves.Interpolation 	= In.Strands.InterpolationResource->MaxAvailableCurveCount;
+	if (In.Strands.ClusterCullingResource)	Out.Curves.Cluster 			= In.Strands.ClusterCullingResource->MaxAvailableCurveCount;
+#if RHI_RAYTRACING
+	if (In.Strands.RaytracingResource)		Out.Curves.Raytracing 		= In.Strands.RaytracingResource->MaxAvailableCurveCount;
+#endif
+	return Out;
+}
+
+void FGroomAssetMemoryStats::Accumulate(const FGroomAssetMemoryStats& In)
+{
+	CPU.Guides += In.CPU.Guides ;
+	CPU.Strands+= In.CPU.Strands;
+	CPU.Cards  += In.CPU.Cards  ;
+	CPU.Meshes += In.CPU.Meshes ;
+
+	GPU.Guides += In.GPU.Guides ;
+	GPU.Strands+= In.GPU.Strands;
+	GPU.Cards  += In.GPU.Cards  ;
+	GPU.Meshes += In.GPU.Meshes ;
+}
+uint32 FGroomAssetMemoryStats::GetTotalCPUSize() const
+{
+	uint32 Out = 0;
+	Out += GPU.Guides ;
+	Out += GPU.Strands;
+	Out += GPU.Cards  ;
+	Out += GPU.Meshes ;
+	return Out;
+}
+
+uint32 FGroomAssetMemoryStats::GetTotalGPUSize() const
+{
+	uint32 Out = 0;
+	Out += CPU.Guides ;
+	Out += CPU.Strands;
+	Out += CPU.Cards  ;
+	Out += CPU.Meshes ;
+	return Out;
+}
+
 void DumpLoadedGroomAssets(IConsoleVariable* InCVarPakTesterEnabled)
 {
 	const bool bDetails = true;
@@ -140,14 +202,7 @@ void DumpLoadedGroomAssets(IConsoleVariable* InCVarPakTesterEnabled)
 	uint32 Total_Asset = 0;
 	uint32 Total_Group = 0;
 
-	uint32 Total_CPUMemorySize_Guides = 0;
-	uint32 Total_CPUMemorySize_Strands= 0;
-	uint32 Total_CPUMemorySize_Cards  = 0;
-	uint32 Total_CPUMemorySize_Meshes = 0;
-	uint32 Total_GPUMemorySize_Guides = 0;
-	uint32 Total_GPUMemorySize_Strands= 0;
-	uint32 Total_GPUMemorySize_Cards  = 0;
-	uint32 Total_GPUMemorySize_Meshes = 0;
+	FGroomAssetMemoryStats Total;
 
 	UE_LOG(LogHairStrands, Log, TEXT("[Groom] ##### UGroomAssets #####"));
 	UE_LOG(LogHairStrands, Log, TEXT("----------------------------------------------------------------------------------------------------------------------------------------------------------------"));
@@ -161,37 +216,8 @@ void DumpLoadedGroomAssets(IConsoleVariable* InCVarPakTesterEnabled)
 			const uint32 GroupCount = AssetIt->GetHairGroupsPlatformData().Num();
 			for (uint32 GroupIt = 0; GroupIt < GroupCount; ++GroupIt)
 			{
-				const uint32 CPUMemorySize_Guides  = AssetIt->GetHairGroupsPlatformData()[GroupIt].Guides.GetDataSize();
-				const uint32 CPUMemorySize_Strands = AssetIt->GetHairGroupsPlatformData()[GroupIt].Strands.GetDataSize();
-				const uint32 CPUMemorySize_Cards   = AssetIt->GetHairGroupsPlatformData()[GroupIt].Cards.GetDataSize();
-				const uint32 CPUMemorySize_Meshes  = AssetIt->GetHairGroupsPlatformData()[GroupIt].Meshes.GetDataSize();
-
-				const uint32 GPUMemorySize_Guides  = AssetIt->GetHairGroupsPlatformData()[GroupIt].Guides.GetResourcesSize();
-				const uint32 GPUMemorySize_Strands = AssetIt->GetHairGroupsPlatformData()[GroupIt].Strands.GetResourcesSize();
-				const uint32 GPUMemorySize_Cards   = AssetIt->GetHairGroupsPlatformData()[GroupIt].Cards.GetResourcesSize();
-				const uint32 GPUMemorySize_Meshes  = AssetIt->GetHairGroupsPlatformData()[GroupIt].Meshes.GetResourcesSize();
-
-				uint32 CPUMemorySize = 0;
-				uint32 GPUMemorySize = 0;
-				CPUMemorySize += CPUMemorySize_Guides;
-				CPUMemorySize += CPUMemorySize_Strands;
-				CPUMemorySize += CPUMemorySize_Cards;
-				CPUMemorySize += CPUMemorySize_Meshes;
-
-				GPUMemorySize += GPUMemorySize_Guides;
-				GPUMemorySize += GPUMemorySize_Strands;
-				GPUMemorySize += GPUMemorySize_Cards;
-				GPUMemorySize += GPUMemorySize_Meshes;
-
-				Total_CPUMemorySize_Guides += CPUMemorySize_Guides;
-				Total_CPUMemorySize_Strands+= CPUMemorySize_Strands;
-				Total_CPUMemorySize_Cards  += CPUMemorySize_Cards;
-				Total_CPUMemorySize_Meshes += CPUMemorySize_Meshes;
-				Total_GPUMemorySize_Guides += GPUMemorySize_Guides;
-				Total_GPUMemorySize_Strands+= GPUMemorySize_Strands;
-				Total_GPUMemorySize_Cards  += GPUMemorySize_Cards;
-				Total_GPUMemorySize_Meshes += GPUMemorySize_Meshes;
-
+				const FGroomAssetMemoryStats Group = FGroomAssetMemoryStats::Get(AssetIt->GetHairGroupsPlatformData()[GroupIt]);
+				Total.Accumulate(Group);
 				const uint32 LODCount = AssetIt->GetHairGroupsLOD()[GroupIt].LODs.Num();
 				if (bDetails)
 				{
@@ -200,20 +226,17 @@ void DumpLoadedGroomAssets(IConsoleVariable* InCVarPakTesterEnabled)
 					UE_LOG(LogHairStrands, Log, TEXT("-- %2d/%2d -  %2d -  %9.3fMb (%9.3fMb|%9.3fMb|%9.3fMb|%9.3fMb) -  %9.3fMb (%9.3fMb|%9.3fMb|%9.3fMb|%9.3fMb) - %s"), 
 						GroupIt,
 						GroupCount,
-
 						LODCount, 
-
-						CPUMemorySize * ToMb,
-						CPUMemorySize_Guides * ToMb,
-						CPUMemorySize_Strands* ToMb,
-						CPUMemorySize_Cards  * ToMb,
-						CPUMemorySize_Meshes * ToMb,
-
-						GPUMemorySize* ToMb,
-						GPUMemorySize_Guides * ToMb,
-						GPUMemorySize_Strands* ToMb,
-						GPUMemorySize_Cards  * ToMb,
-						GPUMemorySize_Meshes * ToMb,
+						Group.GetTotalCPUSize() * ToMb,
+						Group.CPU.Guides * ToMb,
+						Group.CPU.Strands* ToMb,
+						Group.CPU.Cards  * ToMb,
+						Group.CPU.Meshes * ToMb,
+                        Group.GetTotalGPUSize()* ToMb,
+						Group.GPU.Guides * ToMb,
+						Group.GPU.Strands* ToMb,
+						Group.GPU.Cards  * ToMb,
+						Group.GPU.Meshes * ToMb,
 						GroupIt == 0 ? *AssetIt->GetPathName() : TEXT("."));
 				}
 			}
@@ -223,34 +246,76 @@ void DumpLoadedGroomAssets(IConsoleVariable* InCVarPakTesterEnabled)
 		}
 	}
 
-	uint32 Total_CPUMemorySize = 0;
-	Total_CPUMemorySize += Total_CPUMemorySize_Guides;
-	Total_CPUMemorySize += Total_CPUMemorySize_Strands;
-	Total_CPUMemorySize += Total_CPUMemorySize_Cards;
-	Total_CPUMemorySize += Total_CPUMemorySize_Meshes;
-
-	uint32 Total_GPUMemorySize = 0;
-	Total_GPUMemorySize += Total_GPUMemorySize_Guides;
-	Total_GPUMemorySize += Total_GPUMemorySize_Strands;
-	Total_GPUMemorySize += Total_GPUMemorySize_Cards;
-	Total_GPUMemorySize += Total_GPUMemorySize_Meshes;
-
 	UE_LOG(LogHairStrands, Log, TEXT("----------------------------------------------------------------------------------------------------------------------------------------------------------------"));
 	UE_LOG(LogHairStrands, Log, TEXT("-- A:%3d|G:%3d -  %9.3fMb (%9.3fMb|%9.3fMb|%9.3fMb|%9.3fMb) -  %9.3fMb (%9.3fMb|%9.3fMb|%9.3fMb|%9.3fMb)"),
 		Total_Asset,
 		Total_Group,
+		Total.GetTotalCPUSize() * ToMb,
+		Total.CPU.Guides * ToMb,
+		Total.CPU.Strands * ToMb,
+		Total.CPU.Cards * ToMb,
+		Total.CPU.Meshes * ToMb,
+        Total.GetTotalGPUSize() * ToMb,
+		Total.GPU.Guides * ToMb,
+		Total.GPU.Strands * ToMb,
+		Total.GPU.Cards * ToMb,
+		Total.GPU.Meshes * ToMb);
+}
 
-		Total_CPUMemorySize * ToMb,
-		Total_CPUMemorySize_Guides * ToMb,
-		Total_CPUMemorySize_Strands * ToMb,
-		Total_CPUMemorySize_Cards * ToMb,
-		Total_CPUMemorySize_Meshes * ToMb,
+FGroomBindingAssetMemoryStats FGroomBindingAssetMemoryStats::Get(const UGroomBindingAsset::FHairGroupPlatformData& InCPU, const UGroomBindingAsset::FHairGroupResource& InGPU)
+{
+	FGroomBindingAssetMemoryStats Out;
+	Out.CPU.Guides  = InCPU.SimRootBulkData.GetDataSize();
+	Out.CPU.Strands = InCPU.RenRootBulkData.GetDataSize();
 
-		Total_GPUMemorySize * ToMb,
-		Total_GPUMemorySize_Guides * ToMb,
-		Total_GPUMemorySize_Strands * ToMb,
-		Total_GPUMemorySize_Cards * ToMb,
-		Total_GPUMemorySize_Meshes * ToMb);
+	if (const FHairStrandsRestRootResource* R = InGPU.SimRootResources)
+	{
+		Out.GPU.Guides += R->GetResourcesSize();
+	}
+	if (const FHairStrandsRestRootResource* R = InGPU.RenRootResources)
+	{
+		Out.GPU.Strands += R->GetResourcesSize();
+	}
+
+	const uint32 CardCount = InCPU.CardsRootBulkData.Num();
+	for (uint32 CardIt = 0; CardIt < CardCount; ++CardIt)
+	{
+		Out.CPU.Cards += InCPU.CardsRootBulkData[CardIt].GetDataSize();
+		if (const FHairStrandsRestRootResource* R = InGPU.CardsRootResources[CardIt])
+		{
+			Out.GPU.Cards += R->GetResourcesSize();
+		}
+	}
+
+	return Out;
+}
+
+void FGroomBindingAssetMemoryStats::Accumulate(const FGroomBindingAssetMemoryStats& In)
+{
+	CPU.Guides += In.CPU.Guides ;
+	CPU.Strands+= In.CPU.Strands;
+	CPU.Cards  += In.CPU.Cards  ;
+
+	GPU.Guides += In.GPU.Guides ;
+	GPU.Strands+= In.GPU.Strands;
+	GPU.Cards  += In.GPU.Cards  ;
+}
+uint32 FGroomBindingAssetMemoryStats::GetTotalCPUSize() const
+{
+	uint32 Out = 0;
+	Out += GPU.Guides ;
+	Out += GPU.Strands;
+	Out += GPU.Cards  ;
+	return Out;
+}
+
+uint32 FGroomBindingAssetMemoryStats::GetTotalGPUSize() const
+{
+	uint32 Out = 0;
+	Out += CPU.Guides ;
+	Out += CPU.Strands;
+	Out += CPU.Cards  ;
+	return Out;
 }
 
 void DumpLoadedGroomBindingAssets(IConsoleVariable* InCVarPakTesterEnabled)
@@ -261,15 +326,8 @@ void DumpLoadedGroomBindingAssets(IConsoleVariable* InCVarPakTesterEnabled)
 	uint32 Total_Asset = 0;
 	uint32 Total_Group = 0;
 
-	uint32 Total_CPUMemorySize_Guides = 0;
-	uint32 Total_CPUMemorySize_Strands= 0;
-	uint32 Total_CPUMemorySize_Cards  = 0;
-	
-	uint32 Total_GPUMemorySize_Guides = 0;
-	uint32 Total_GPUMemorySize_Strands= 0;
-	uint32 Total_GPUMemorySize_Cards  = 0;
+	FGroomBindingAssetMemoryStats Total;
 	
-
 	UE_LOG(LogHairStrands, Log, TEXT("[Groom] ##### UGroomBindingAssets #####"));
 	UE_LOG(LogHairStrands, Log, TEXT("----------------------------------------------------------------------------------------------------------------------------------------------------------------"));
 	UE_LOG(LogHairStrands, Log, TEXT("--  No.  - LOD -    CPU Total (     Guides|    Strands|      Cards) -    GPU Total (     Guides|    Strands|      Cards) - Asset Name "));
@@ -282,42 +340,8 @@ void DumpLoadedGroomBindingAssets(IConsoleVariable* InCVarPakTesterEnabled)
 			const uint32 GroupCount = AssetIt->GetHairGroupsPlatformData().Num();
 			for (uint32 GroupIt = 0; GroupIt < GroupCount; ++GroupIt)
 			{
-				const uint32 CPUMemorySize_Guides  = AssetIt->GetHairGroupsPlatformData()[GroupIt].SimRootBulkData.GetDataSize();
-				const uint32 CPUMemorySize_Strands = AssetIt->GetHairGroupsPlatformData()[GroupIt].RenRootBulkData.GetDataSize();
-
-				uint32 GPUMemorySize_Guides = 0;
-				if (const FHairStrandsRestRootResource* RootResource = AssetIt->GetHairGroupResources()[GroupIt].SimRootResources)
-				{
-					GPUMemorySize_Guides += RootResource->GetResourcesSize();
-				}
-				uint32 GPUMemorySize_Strands = 0;
-				if (const FHairStrandsRestRootResource* RootResource = AssetIt->GetHairGroupResources()[GroupIt].RenRootResources)
-				{
-					GPUMemorySize_Strands += RootResource->GetResourcesSize();
-				}
-
-				uint32 CPUMemorySize_Cards = 0;
-				uint32 GPUMemorySize_Cards = 0;
-				const uint32 CardCount = AssetIt->GetHairGroupsPlatformData()[GroupIt].CardsRootBulkData.Num();
-				for (uint32 CardIt = 0; CardIt < CardCount; ++CardIt)
-				{
-					CPUMemorySize_Cards += AssetIt->GetHairGroupsPlatformData()[GroupIt].CardsRootBulkData[CardIt].GetDataSize();
-					if (const FHairStrandsRestRootResource* RootResource = AssetIt->GetHairGroupResources()[GroupIt].CardsRootResources[CardIt])
-					{
-						GPUMemorySize_Cards += RootResource->GetResourcesSize();
-					}
-				}
-
-				const uint32 CPUMemorySize = CPUMemorySize_Guides + CPUMemorySize_Strands + CPUMemorySize_Cards;
-				const uint32 GPUMemorySize = GPUMemorySize_Guides + GPUMemorySize_Strands + GPUMemorySize_Cards;
-
-				Total_CPUMemorySize_Guides += CPUMemorySize_Guides;
-				Total_CPUMemorySize_Strands+= CPUMemorySize_Strands;
-				Total_CPUMemorySize_Cards  += CPUMemorySize_Cards;
-				
-				Total_GPUMemorySize_Guides += GPUMemorySize_Guides;
-				Total_GPUMemorySize_Strands+= GPUMemorySize_Strands;
-				Total_GPUMemorySize_Cards  += GPUMemorySize_Cards;
+				const FGroomBindingAssetMemoryStats Group = FGroomBindingAssetMemoryStats::Get(AssetIt->GetHairGroupsPlatformData()[GroupIt], AssetIt->GetHairGroupResources()[GroupIt]);
+				Total.Accumulate(Group);
 
 				const uint32 SkelLODCount = AssetIt->GetHairGroupsPlatformData()[GroupIt].RenRootBulkData.GetLODCount();
 				if (bDetails)
@@ -327,18 +351,15 @@ void DumpLoadedGroomBindingAssets(IConsoleVariable* InCVarPakTesterEnabled)
 					UE_LOG(LogHairStrands, Log, TEXT("-- %2d/%2d -  %2d -  %9.3fMb (%9.3fMb|%9.3fMb|%9.3fMb) -  %9.3fMb (%9.3fMb|%9.3fMb|%9.3fMb) - %s"), 
 						GroupIt,
 						GroupCount,
-
 						SkelLODCount,
-
-						CPUMemorySize * ToMb,
-						CPUMemorySize_Guides * ToMb,
-						CPUMemorySize_Strands* ToMb,
-						CPUMemorySize_Cards  * ToMb,
-
-						GPUMemorySize* ToMb,
-						GPUMemorySize_Guides * ToMb,
-						GPUMemorySize_Strands* ToMb,
-						GPUMemorySize_Cards  * ToMb,
+                        Group.GetTotalCPUSize() * ToMb,
+						Group.CPU.Guides * ToMb,
+						Group.CPU.Strands* ToMb,
+						Group.CPU.Cards  * ToMb,
+                        Group.GetTotalGPUSize()* ToMb,
+						Group.GPU.Guides * ToMb,
+						Group.GPU.Strands* ToMb,
+						Group.GPU.Cards  * ToMb,
 						GroupIt == 0 ? *AssetIt->GetPathName() : TEXT("."));
 				}
 			}
@@ -347,31 +368,18 @@ void DumpLoadedGroomBindingAssets(IConsoleVariable* InCVarPakTesterEnabled)
 			Total_Group += GroupCount;
 		}
 	}
-
-	uint32 Total_CPUMemorySize = 0;
-	Total_CPUMemorySize += Total_CPUMemorySize_Guides;
-	Total_CPUMemorySize += Total_CPUMemorySize_Strands;
-	Total_CPUMemorySize += Total_CPUMemorySize_Cards;
-
-	uint32 Total_GPUMemorySize = 0;
-	Total_GPUMemorySize += Total_GPUMemorySize_Guides;
-	Total_GPUMemorySize += Total_GPUMemorySize_Strands;
-	Total_GPUMemorySize += Total_GPUMemorySize_Cards;
-
 	UE_LOG(LogHairStrands, Log, TEXT("----------------------------------------------------------------------------------------------------------------------------------------------------------------"));
 	UE_LOG(LogHairStrands, Log, TEXT("-- A:%3d|G:%3d -  %9.3fMb (%9.3fMb|%9.3fMb|%9.3fMb) -  %9.3fMb (%9.3fMb|%9.3fMb|%9.3fMb)"),
 		Total_Asset,
 		Total_Group,
-
-		Total_CPUMemorySize * ToMb,
-		Total_CPUMemorySize_Guides * ToMb,
-		Total_CPUMemorySize_Strands * ToMb,
-		Total_CPUMemorySize_Cards * ToMb,
-
-		Total_GPUMemorySize * ToMb,
-		Total_GPUMemorySize_Guides * ToMb,
-		Total_GPUMemorySize_Strands * ToMb,
-		Total_GPUMemorySize_Cards * ToMb);
+		Total.GetTotalCPUSize() * ToMb,
+		Total.CPU.Guides * ToMb,
+		Total.CPU.Strands * ToMb,
+		Total.CPU.Cards * ToMb,
+		Total.GetTotalGPUSize() * ToMb,
+		Total.GPU.Guides * ToMb,
+		Total.GPU.Strands * ToMb,
+		Total.GPU.Cards * ToMb);
 }
 
 void DumpLoadedGroomData(IConsoleVariable* InCVarPakTesterEnabled)
diff --git a/Engine/Plugins/Runtime/HairStrands/Source/HairStrandsCore/Private/GroomComponent.cpp b/Engine/Plugins/Runtime/HairStrands/Source/HairStrandsCore/Private/GroomComponent.cpp
--- a/Engine/Plugins/Runtime/HairStrands/Source/HairStrandsCore/Private/GroomComponent.cpp
+++ b/Engine/Plugins/Runtime/HairStrands/Source/HairStrandsCore/Private/GroomComponent.cpp
@@ -3971,32 +3971,12 @@ void UGroomComponent::GetResourceSizeEx(FResourceSizeEx& CumulativeResourceSize)
 
 uint32 UGroomComponent::GetResourcesSize() const
 {
-	uint32 Total = 0;
+	FGroomComponentMemoryStats Total;
 	for (const FHairGroupInstance* Instance : HairGroupInstances)
 	{
-		Total += Instance->HairGroupPublicData ? Instance->HairGroupPublicData->GetResourcesSize() : 0;
-
-		Total += Instance->Guides.DeformedResource ? Instance->Guides.DeformedResource->GetResourcesSize() : 0;
-		Total += Instance->Guides.DeformedRootResource ? Instance->Guides.DeformedRootResource->GetResourcesSize() : 0;
-
-		Total += Instance->Strands.DeformedResource ? Instance->Strands.DeformedResource->GetResourcesSize() : 0;
-		Total += Instance->Strands.DeformedRootResource ? Instance->Strands.DeformedRootResource->GetResourcesSize() : 0;
-	#if RHI_RAYTRACING
-		Total += (Instance->Strands.RenRaytracingResourceOwned && Instance->Strands.RenRaytracingResource) ? Instance->Strands.RenRaytracingResource->GetResourcesSize() : 0;
-	#endif
-
-		for (const FHairGroupInstance::FCards::FLOD& LOD : Instance->Cards.LODs)
-		{
-			Total += LOD.DeformedResource ? LOD.DeformedResource->GetResourcesSize() : 0;
-		}
-
-		for (const FHairGroupInstance::FMeshes::FLOD& LOD : Instance->Meshes.LODs)
-		{
-			Total += LOD.DeformedResource ? LOD.DeformedResource->GetResourcesSize() : 0;
-		}
+		Total.Accumulate(FGroomComponentMemoryStats::Get(Instance));
 	}
-
-	return Total;
+	return Total.GetTotalSize();
 }
 
 /////////////////////////////////////////////////////////////////////////////////////////
@@ -4019,11 +3999,7 @@ void DumpLoadedGroomComponent(IConsoleVariable* InCVarPakTesterEnabled)
 	uint32 Total_Component = 0;
 	uint32 Total_Group = 0;
 
-	uint32 Total_GPUMemorySize_Common = 0;
-	uint32 Total_GPUMemorySize_Guides = 0;
-	uint32 Total_GPUMemorySize_Strands= 0;
-	uint32 Total_GPUMemorySize_Cards  = 0;
-	uint32 Total_GPUMemorySize_Meshes = 0;
+	FGroomComponentMemoryStats TotalMemory;
 
 	UE_LOG(LogHairStrands, Log, TEXT("[Groom] ##### UGroomComponent #####"));
 	UE_LOG(LogHairStrands, Log, TEXT("----------------------------------------------------------------------------------------------------------------------------------------------------------------"));
@@ -4037,49 +4013,10 @@ void DumpLoadedGroomComponent(IConsoleVariable* InCVarPakTesterEnabled)
 			const uint32 GroupCount = ComponentIt->GetGroupCount();
 			for (uint32 GroupIt = 0; GroupIt < GroupCount; ++GroupIt)
 			{
-				FHairGroupInstance* Instance = ComponentIt->GetGroupInstance(GroupIt);
-
-				uint32 GPUMemorySize_Common = 0;
-				GPUMemorySize_Common += Instance->HairGroupPublicData ? Instance->HairGroupPublicData->GetResourcesSize() : 0;
-
-				uint32 GPUMemorySize_Guides = 0;
-				GPUMemorySize_Guides += Instance->Guides.DeformedResource ? Instance->Guides.DeformedResource->GetResourcesSize() : 0;
-				GPUMemorySize_Guides += Instance->Guides.DeformedRootResource ? Instance->Guides.DeformedRootResource->GetResourcesSize() : 0;
-
-				uint32 GPUMemorySize_Strands = 0;
-				GPUMemorySize_Strands += Instance->Strands.DeformedResource ? Instance->Strands.DeformedResource->GetResourcesSize() : 0;
-				GPUMemorySize_Strands += Instance->Strands.DeformedRootResource ? Instance->Strands.DeformedRootResource->GetResourcesSize() : 0;
-			#if RHI_RAYTRACING
-				GPUMemorySize_Strands += (Instance->Strands.RenRaytracingResourceOwned && Instance->Strands.RenRaytracingResource) ? Instance->Strands.RenRaytracingResource->GetResourcesSize() : 0;
-			#endif
-
-				uint32 GPUMemorySize_Cards  = 0;
-				for (const FHairGroupInstance::FCards::FLOD& LOD : Instance->Cards.LODs)
-				{
-					GPUMemorySize_Cards += LOD.DeformedResource ? LOD.DeformedResource->GetResourcesSize() : 0;
-				}
-
-				uint32 GPUMemorySize_Meshes = 0;
-				for (const FHairGroupInstance::FMeshes::FLOD& LOD : Instance->Meshes.LODs)
-				{
-					GPUMemorySize_Meshes += LOD.DeformedResource ? LOD.DeformedResource->GetResourcesSize() : 0;
-				}
-
-				uint32 LODCount = Instance->Cards.LODs.Num();
-
-				uint32 GPUMemorySize = 0;
-				GPUMemorySize += GPUMemorySize_Common;
-				GPUMemorySize += GPUMemorySize_Guides;
-				GPUMemorySize += GPUMemorySize_Strands;
-				GPUMemorySize += GPUMemorySize_Cards;
-				GPUMemorySize += GPUMemorySize_Meshes;
-
-				Total_GPUMemorySize_Common += GPUMemorySize_Common;
-				Total_GPUMemorySize_Guides += GPUMemorySize_Guides;
-				Total_GPUMemorySize_Strands+= GPUMemorySize_Strands;
-				Total_GPUMemorySize_Cards  += GPUMemorySize_Cards;
-				Total_GPUMemorySize_Meshes += GPUMemorySize_Meshes;
-
+				const FHairGroupInstance* Instance = ComponentIt->GetGroupInstance(GroupIt);
+				const FGroomComponentMemoryStats InstanceMemory = FGroomComponentMemoryStats::Get(ComponentIt->GetGroupInstance(GroupIt));
+				TotalMemory.Accumulate(InstanceMemory);
+				
 				if (bDetails)
 				{
 //					UE_LOG(LogHairStrands, Log, TEXT("--  No.  - LOD -    GPU Total (    Common|     Guides|    Strands|      Cards|     Meshes) - Asset Name "));
@@ -4087,15 +4024,13 @@ void DumpLoadedGroomComponent(IConsoleVariable* InCVarPakTesterEnabled)
 					UE_LOG(LogHairStrands, Log, TEXT("-- %2d/%2d -  %2d -  %9.3fMb (%9.3fMb|%9.3fMb|%9.3fMb|%9.3fMb|%9.3fMb) - %s"), 
 						GroupIt,
 						GroupCount,
-
-						LODCount, 
-
-						GPUMemorySize* ToMb,
-						GPUMemorySize_Common * ToMb,
-						GPUMemorySize_Guides * ToMb,
-						GPUMemorySize_Strands* ToMb,
-						GPUMemorySize_Cards  * ToMb,
-						GPUMemorySize_Meshes * ToMb,
+                        Instance->Cards.LODs.Num(), 
+                        InstanceMemory.GetTotalSize()* ToMb,
+                        InstanceMemory.Common * ToMb,
+						InstanceMemory.Guides * ToMb,
+						InstanceMemory.Strands* ToMb,
+						InstanceMemory.Cards  * ToMb,
+						InstanceMemory.Meshes * ToMb,
 						GroupIt == 0 && ComponentIt->GroomAsset ? *ComponentIt->GroomAsset->GetPathName() : TEXT("."));
 				}
 			}
@@ -4105,24 +4040,16 @@ void DumpLoadedGroomComponent(IConsoleVariable* InCVarPakTesterEnabled)
 		}
 	}
 
-	uint32 Total_GPUMemorySize = 0;
-	Total_GPUMemorySize += Total_GPUMemorySize_Common;
-	Total_GPUMemorySize += Total_GPUMemorySize_Guides;
-	Total_GPUMemorySize += Total_GPUMemorySize_Strands;
-	Total_GPUMemorySize += Total_GPUMemorySize_Cards;
-	Total_GPUMemorySize += Total_GPUMemorySize_Meshes;
-
 	UE_LOG(LogHairStrands, Log, TEXT("----------------------------------------------------------------------------------------------------------------------------------------------------------------"));
 	UE_LOG(LogHairStrands, Log, TEXT("-- C:%3d|G:%3d -  %9.3fMb (%9.3fMb|%9.3fMb|%9.3fMb|%9.3fMb|%9.3fMb)"),
 		Total_Component,
 		Total_Group,
-
-		Total_GPUMemorySize * ToMb,
-		Total_GPUMemorySize_Common* ToMb,
-		Total_GPUMemorySize_Guides * ToMb,
-		Total_GPUMemorySize_Strands * ToMb,
-		Total_GPUMemorySize_Cards * ToMb,
-		Total_GPUMemorySize_Meshes * ToMb);
+        TotalMemory.GetTotalSize() * ToMb,
+		TotalMemory.Common* ToMb,
+		TotalMemory.Guides * ToMb,
+		TotalMemory.Strands * ToMb,
+		TotalMemory.Cards * ToMb,
+		TotalMemory.Meshes * ToMb);
 }
 
 /////////////////////////////////////////////////////////////////////////////////////////
@@ -4198,5 +4125,44 @@ FGroomComponentRecreateRenderStateContext::~FGroomComponentRecreateRenderStateCo
 }
 #endif
 
+FGroomComponentMemoryStats FGroomComponentMemoryStats::Get(const FHairGroupInstance* In)
+{
+	FGroomComponentMemoryStats Out;
+	Out.Common  += In->HairGroupPublicData ? In->HairGroupPublicData->GetResourcesSize() : 0;
+	Out.Guides  += In->Guides.DeformedResource ? In->Guides.DeformedResource->GetResourcesSize() : 0;
+	Out.Guides  += In->Guides.DeformedRootResource ? In->Guides.DeformedRootResource->GetResourcesSize() : 0;
+	Out.Strands += In->Strands.DeformedResource ? In->Strands.DeformedResource->GetResourcesSize() : 0;
+	Out.Strands += In->Strands.DeformedRootResource ? In->Strands.DeformedRootResource->GetResourcesSize() : 0;
+#if RHI_RAYTRACING
+	Out.Strands += (In->Strands.RenRaytracingResourceOwned && In->Strands.RenRaytracingResource) ? In->Strands.RenRaytracingResource->GetResourcesSize() : 0;
+#endif
+	for (const FHairGroupInstance::FCards::FLOD& LOD : In->Cards.LODs)
+	{
+		Out.Cards += LOD.DeformedResource ? LOD.DeformedResource->GetResourcesSize() : 0;
+	}
+	for (const FHairGroupInstance::FMeshes::FLOD& LOD : In->Meshes.LODs)
+	{
+		Out.Meshes += LOD.DeformedResource ? LOD.DeformedResource->GetResourcesSize() : 0;
+	}
+	return Out;
+}
+
+void FGroomComponentMemoryStats::Accumulate(const FGroomComponentMemoryStats& In)
+{
+	Common += In.Common ;
+	Guides += In.Guides ;
+	Strands+= In.Strands;
+	Cards  += In.Cards  ;
+	Meshes += In.Meshes ;
+}
+uint32 FGroomComponentMemoryStats::GetTotalSize() const
+{
+	uint32 Out = 0;
+	Out += Guides ;
+	Out += Strands;
+	Out += Cards  ;
+	Out += Meshes ;
+	return Out;
+}
 #undef LOCTEXT_NAMESPACE
 
diff --git a/Engine/Plugins/Runtime/HairStrands/Source/HairStrandsCore/Public/GroomAsset.h b/Engine/Plugins/Runtime/HairStrands/Source/HairStrandsCore/Public/GroomAsset.h
--- a/Engine/Plugins/Runtime/HairStrands/Source/HairStrandsCore/Public/GroomAsset.h
+++ b/Engine/Plugins/Runtime/HairStrands/Source/HairStrandsCore/Public/GroomAsset.h
@@ -822,3 +822,31 @@ private:
 	TQueue<FHairGroupCardsTextures*> AssetToSave_Textures;
 #endif // WITH_EDITOR
 };
+
+struct FGroomAssetMemoryStats
+{
+	struct FStats
+	{
+		uint32 Guides = 0;
+		uint32 Strands= 0;
+		uint32 Cards  = 0;
+		uint32 Meshes = 0;
+	};
+	FStats CPU;
+	FStats GPU;
+
+	struct FStrandsDetails
+	{
+		uint32 Rest = 0;
+		uint32 Interpolation= 0;
+		uint32 Cluster  = 0;
+		uint32 Raytracing = 0;
+	};
+	FStrandsDetails Memory;
+	FStrandsDetails Curves;
+
+	static FGroomAssetMemoryStats Get(const FHairGroupPlatformData& In);
+	void Accumulate(const FGroomAssetMemoryStats& In);
+	uint32 GetTotalCPUSize() const;
+	uint32 GetTotalGPUSize() const;
+};
\ No newline at end of file
diff --git a/Engine/Plugins/Runtime/HairStrands/Source/HairStrandsCore/Public/GroomBindingAsset.h b/Engine/Plugins/Runtime/HairStrands/Source/HairStrandsCore/Public/GroomBindingAsset.h
--- a/Engine/Plugins/Runtime/HairStrands/Source/HairStrandsCore/Public/GroomBindingAsset.h
+++ b/Engine/Plugins/Runtime/HairStrands/Source/HairStrandsCore/Public/GroomBindingAsset.h
@@ -288,4 +288,21 @@ class HAIRSTRANDSCORE_API UGroomBindingAssetList : public UObject
 public:
 	UPROPERTY(EditAnywhere, Transient, EditFixedSize, Category = "Bindings")
 	TArray<TObjectPtr<UGroomBindingAsset>> Bindings;
+};
+
+struct FGroomBindingAssetMemoryStats
+{
+	struct FStats
+	{
+		uint32 Guides = 0;
+		uint32 Strands= 0;
+		uint32 Cards  = 0;
+	};
+	FStats CPU;
+	FStats GPU;
+
+	static FGroomBindingAssetMemoryStats Get(const UGroomBindingAsset::FHairGroupPlatformData& InCPU, const UGroomBindingAsset::FHairGroupResource& InGPU);
+	void Accumulate(const FGroomBindingAssetMemoryStats& In);
+	uint32 GetTotalCPUSize() const;
+	uint32 GetTotalGPUSize() const;
 };
\ No newline at end of file
diff --git a/Engine/Plugins/Runtime/HairStrands/Source/HairStrandsCore/Public/GroomComponent.h b/Engine/Plugins/Runtime/HairStrands/Source/HairStrandsCore/Public/GroomComponent.h
--- a/Engine/Plugins/Runtime/HairStrands/Source/HairStrandsCore/Public/GroomComponent.h
+++ b/Engine/Plugins/Runtime/HairStrands/Source/HairStrandsCore/Public/GroomComponent.h
@@ -369,3 +369,17 @@ HAIRSTRANDSCORE_API const FLinearColor GetHairGroupDebugColor(int32 GroupIt);
 
 #endif
 
+struct FGroomComponentMemoryStats
+{
+	uint32 Common = 0;
+	uint32 Guides = 0;
+	uint32 Strands= 0;
+	uint32 Cards  = 0;
+	uint32 Meshes = 0;
+
+	static FGroomComponentMemoryStats Get(const FHairGroupInstance* In);
+	void Accumulate(const FGroomComponentMemoryStats& In);
+	uint32 GetTotalSize() const;
+};
+
+
