Add force hair AutoLOD settings.

This allows to force all hair strands groom to use the AutoLOD feature (experimental).

diff --git a/Engine/Plugins/Runtime/HairStrands/Source/HairStrandsCore/Private/GroomDebug.cpp b/Engine/Plugins/Runtime/HairStrands/Source/HairStrandsCore/Private/GroomDebug.cpp
--- a/Engine/Plugins/Runtime/HairStrands/Source/HairStrandsCore/Private/GroomDebug.cpp
+++ b/Engine/Plugins/Runtime/HairStrands/Source/HairStrandsCore/Private/GroomDebug.cpp
@@ -840,6 +840,7 @@ public:
 
 IMPLEMENT_GLOBAL_SHADER(FHairDebugPrintInstanceCS, "/Engine/Private/HairStrands/HairStrandsDebug.usf", "MainCS", SF_Compute);
 
+bool IsHairStrandsForceAutoLODEnabled();
 const TCHAR* GetHairAttributeText(EHairAttribute In, uint32 InFlags);
 uint32 GetHairAttributeIndex(EHairAttribute In);
 EGroomCacheType GetHairInstanceCacheType(const FHairGroupInstance* Instance);
@@ -947,7 +948,7 @@ static void AddHairDebugPrintInstancePass(
 				D.Data3.Z |= FFloat16(Instance->HairGroupPublicData->VFInput.Strands.Common.Length).Encoded;
 				D.Data3.Z |= FFloat16(Instance->HairGroupPublicData->VFInput.Strands.Common.LengthScale).Encoded << 16u;
 				D.Data3.W |= FFloat16(Instance->HairGroupPublicData->GetActiveStrandsCoverageScale()).Encoded;
-				D.Data3.W |= Instance->HairGroupPublicData->bAutoLOD ? (1u << 16u) : 0;
+				D.Data3.W |= Instance->HairGroupPublicData->bAutoLOD || IsHairStrandsForceAutoLODEnabled() ? (1u << 16u) : 0;
 			}
 			break;
 		case EHairGeometryType::Cards:
diff --git a/Engine/Plugins/Runtime/HairStrands/Source/HairStrandsCore/Private/GroomManager.cpp b/Engine/Plugins/Runtime/HairStrands/Source/HairStrandsCore/Private/GroomManager.cpp
--- a/Engine/Plugins/Runtime/HairStrands/Source/HairStrandsCore/Private/GroomManager.cpp
+++ b/Engine/Plugins/Runtime/HairStrands/Source/HairStrandsCore/Private/GroomManager.cpp
@@ -47,10 +47,12 @@ static FAutoConsoleVariableRef CVarHairStrands_InterpolationFrustumCullingEnable
 static int32 GHairStrands_Streaming = 0;
 static FAutoConsoleVariableRef CVarHairStrands_Streaming(TEXT("r.HairStrands.Streaming"), GHairStrands_Streaming, TEXT("Hair strands streaming toggle."), ECVF_RenderThreadSafe | ECVF_ReadOnly);
 
+static int32 GHairStrands_AutoLOD_Force = 0;
 static float GHairStrands_AutoLOD_Scale = 1.f;
 static float GHairStrands_AutoLOD_Bias = 0.f;
-static FAutoConsoleVariableRef CVarHairStrands_AutoLOD_Scale(TEXT("r.HairStrands.AutoLOD.Scale"), GHairStrands_Streaming, TEXT("Hair strands Auto LOD rate at which curves get decimated based on screen coverage."), ECVF_RenderThreadSafe | ECVF_ReadOnly);
-static FAutoConsoleVariableRef CVarHairStrands_AutoLOD_Bias(TEXT("r.HairStrands.AutoLOD.Bias"), GHairStrands_Streaming, TEXT("Hair strands Auto LOD screen size bias at which curves get decimated."), ECVF_RenderThreadSafe | ECVF_ReadOnly);
+static FAutoConsoleVariableRef CVarHairStrands_AutoLOD_Force(TEXT("r.HairStrands.AutoLOD.Force"), GHairStrands_AutoLOD_Force, TEXT("Force all groom to use Auto LOD (experimental)."), ECVF_RenderThreadSafe);
+static FAutoConsoleVariableRef CVarHairStrands_AutoLOD_Scale(TEXT("r.HairStrands.AutoLOD.Scale"), GHairStrands_AutoLOD_Scale, TEXT("Hair strands Auto LOD rate at which curves get decimated based on screen coverage."), ECVF_RenderThreadSafe);
+static FAutoConsoleVariableRef CVarHairStrands_AutoLOD_Bias(TEXT("r.HairStrands.AutoLOD.Bias"), GHairStrands_AutoLOD_Bias, TEXT("Hair strands Auto LOD screen size bias at which curves get decimated."), ECVF_RenderThreadSafe);
 
 EHairBufferSwapType GetHairSwapBufferType()
 {
@@ -64,6 +66,11 @@ EHairBufferSwapType GetHairSwapBufferType()
 	return EHairBufferSwapType::EndOfFrame;
 }
 
+bool IsHairStrandsForceAutoLODEnabled()
+{
+	return GHairStrands_AutoLOD_Force > 0;
+}
+
 DEFINE_LOG_CATEGORY_STATIC(LogGroomManager, Log, All);
 
 ///////////////////////////////////////////////////////////////////////////////////////////////////
@@ -890,7 +897,7 @@ static void RunHairLODSelection(
 			if (Instance->Strands.ClusterCullingResource)
 			{
 				uint32 EffectiveCurveCount = 0;
-				if (Instance->HairGroupPublicData->bAutoLOD)
+				if (Instance->HairGroupPublicData->bAutoLOD || IsHairStrandsForceAutoLODEnabled())
 				{
 					EffectiveCurveCount = ComputeActiveCurveCount(Instance->HairGroupPublicData->ContinuousLODScreenSize, Instance->HairGroupPublicData->RestCurveCount);
 					LODIndex = 0;
diff --git a/Engine/Shaders/Private/HairStrands/HairStrandsDebug.usf b/Engine/Shaders/Private/HairStrands/HairStrandsDebug.usf
--- a/Engine/Shaders/Private/HairStrands/HairStrandsDebug.usf
+++ b/Engine/Shaders/Private/HairStrands/HairStrandsDebug.usf
@@ -1354,7 +1354,7 @@ void MainCS(uint GroupIndex : SV_GroupIndex, uint3 DispatchThreadId : SV_Dispatc
 	Newline(Context);
 
 	// Draw instance details
-	Print(Context, TEXT("   Group  LOD     Geometry Binding Cache Sim RBF Clip Screen CovSc. Strands            Points             Bind.Tri/Sect Vis. Shad. SwLOD ALOD Name"));
+	Print(Context, TEXT("   Group  LOD     Geometry Binding Cache Sim RBF Clip Screen CovSc. Strands            Points             Bind.Tri/Sect Vis. Shad. SwLOD Name"));
 	Newline(Context);
 
 	for (uint InstanceIndex = 0; InstanceIndex < InstanceCount; ++InstanceIndex)
@@ -1372,10 +1372,17 @@ void MainCS(uint GroupIndex : SV_GroupIndex, uint3 DispatchThreadId : SV_Dispatc
 		Print(Context, TEXT("    "));
 
 		// LOD
-		Print(Context,Info.LOD, FontYellow, 4, 2);
-		Print(Context, TEXT("/"));
-		Print(Context,Info.LODCount-1, FontYellow, 1);
-		Print(Context, TEXT("  "));
+		if (Info.bAutoLOD)
+		{
+			Print(Context, TEXT("Auto    "), FontYellow);
+		}
+		else
+		{
+			Print(Context,Info.LOD, FontYellow, 4, 2);
+			Print(Context, TEXT("/"));
+			Print(Context,Info.LODCount-1, FontYellow, 1);
+			Print(Context, TEXT("  "));
+		}
 
 		// Geometry
 		if (Info.GeometryType == 0)		{ Print(Context, TEXT("Strands "), FontEmerald); }
@@ -1445,10 +1452,6 @@ void MainCS(uint GroupIndex : SV_GroupIndex, uint3 DispatchThreadId : SV_Dispatc
 		// Has LOD switched
 		PrintBool(Context, Info.bHasLODSwitched);
 		Print(Context, TEXT("  "));
-
-		// Auto LOD
-		PrintBool(Context, Info.bAutoLOD);
-		Print(Context, TEXT("  "));
 		
 		// Name
 		PrintInstanceNames(Context, InstanceIndex, FontWhite);
