Stream and allocate curve as 'page' to avoid too frequent streaming request  and reallocation.

diff --git a/Engine/Plugins/Runtime/HairStrands/Source/HairStrandsCore/Private/GroomManager.cpp b/Engine/Plugins/Runtime/HairStrands/Source/HairStrandsCore/Private/GroomManager.cpp
--- a/Engine/Plugins/Runtime/HairStrands/Source/HairStrandsCore/Private/GroomManager.cpp
+++ b/Engine/Plugins/Runtime/HairStrands/Source/HairStrandsCore/Private/GroomManager.cpp
@@ -46,6 +46,8 @@ static FAutoConsoleVariableRef CVarHairStrands_InterpolationFrustumCullingEnable
 
 static int32 GHairStrands_Streaming = 0;
 static FAutoConsoleVariableRef CVarHairStrands_Streaming(TEXT("r.HairStrands.Streaming"), GHairStrands_Streaming, TEXT("Hair strands streaming toggle."), ECVF_RenderThreadSafe | ECVF_ReadOnly);
+static int32 GHairStrands_Streaming_CurvePage = 2048;
+static FAutoConsoleVariableRef CVarHairStrands_StreamingCurvePage(TEXT("r.HairStrands.Streaming.CurvePage"), GHairStrands_Streaming_CurvePage, TEXT("Number of strands curve per streaming page"));
 
 static int32 GHairStrands_AutoLOD_Force = 0;
 static float GHairStrands_AutoLOD_Scale = 1.f;
@@ -71,6 +73,17 @@ bool IsHairStrandsForceAutoLODEnabled()
 	return GHairStrands_AutoLOD_Force > 0;
 }
 
+uint32 GetStreamingCurvePage()
+{
+	return FMath::Clamp(uint32(GHairStrands_Streaming_CurvePage), 1u, 32768u);
+}
+
+static uint32 GetRoundedCurveCount(uint32 InRequest, uint32 InMaxCurve)
+{
+	const uint32 Page = GetStreamingCurvePage();
+	return FMath::Min(FMath::DivideAndRoundUp(InRequest, Page) * Page,  InMaxCurve);
+}
+
 DEFINE_LOG_CATEGORY_STATIC(LogGroomManager, Log, All);
 
 ///////////////////////////////////////////////////////////////////////////////////////////////////
@@ -969,8 +982,9 @@ static void RunHairLODSelection(
 			const bool bStreamingEnabled = GHairStrands_Streaming > 0;
 			if (GeometryType == EHairGeometryType::Strands && bStreamingEnabled && Instance->bSupportStreaming)
 			{
-				RequestedCurveCount = Instance->HairGroupPublicData->ContinuousLODCurveCount;
-				RequestedPointCount = Instance->HairGroupPublicData->ContinuousLODPointCount;
+				// Round Curve/Point request to curve 'page'
+				RequestedCurveCount = GetRoundedCurveCount(Instance->HairGroupPublicData->ContinuousLODCurveCount, Instance->HairGroupPublicData->RestCurveCount);
+				RequestedPointCount = RequestedCurveCount > 0 ? Instance->Strands.Data->Header.CurveToPointCount[RequestedCurveCount - 1] : 0;
 			}
 
 			const bool bSimulationEnable			= Instance->HairGroupPublicData->IsSimulationEnable(IntLODIndex);
@@ -1107,7 +1121,7 @@ static void RunHairLODSelection(
 				check(Instance->Strands.RestResource);
 
 				// Adapt CurveCount/PointCount/CoverageScale based on what is actually available
-				const uint32 EffectiveCurveCount = FMath::Min(ResourceStatus.AvailableCurveCount, RequestedCurveCount);
+				const uint32 EffectiveCurveCount = FMath::Min(ResourceStatus.AvailableCurveCount, Instance->HairGroupPublicData->ContinuousLODCurveCount);
 				Instance->HairGroupPublicData->ContinuousLODCurveCount = EffectiveCurveCount;
 				Instance->HairGroupPublicData->ContinuousLODPointCount = EffectiveCurveCount > 0 ? Instance->Strands.Data->Header.CurveToPointCount[EffectiveCurveCount - 1] : 0;
 				Instance->HairGroupPublicData->ContinuousLODCoverageScale = ComputeActiveCurveCoverageScale(EffectiveCurveCount, Instance->HairGroupPublicData->RestCurveCount);
