Change culling resource to be non-blocking.

diff --git a/Engine/Plugins/Runtime/HairStrands/Source/HairStrandsCore/Private/GroomResources.cpp b/Engine/Plugins/Runtime/HairStrands/Source/HairStrandsCore/Private/GroomResources.cpp
--- a/Engine/Plugins/Runtime/HairStrands/Source/HairStrandsCore/Private/GroomResources.cpp
+++ b/Engine/Plugins/Runtime/HairStrands/Source/HairStrandsCore/Private/GroomResources.cpp
@@ -1182,7 +1182,7 @@ bool FHairStrandsClusterCullingResource::InternalGetOrRequestData(uint32 InReque
 {
 	if (StreamingRequest.IsNone())
 	{
-		StreamingRequest.Request(InRequestedCurveCount, InRequestedPointCount, InLODIndex, BulkData, true, true, false, OwnerName);
+		StreamingRequest.Request(InRequestedCurveCount, InRequestedPointCount, InLODIndex, BulkData, false, false, false, OwnerName);
 	}
 	return StreamingRequest.IsCompleted();
 }
