Fix missing ShaderPrint set enable when running with groom memory view mode.

diff --git a/Engine/Plugins/Runtime/HairStrands/Source/HairStrandsCore/Private/GroomDebug.cpp b/Engine/Plugins/Runtime/HairStrands/Source/HairStrandsCore/Private/GroomDebug.cpp
--- a/Engine/Plugins/Runtime/HairStrands/Source/HairStrandsCore/Private/GroomDebug.cpp
+++ b/Engine/Plugins/Runtime/HairStrands/Source/HairStrandsCore/Private/GroomDebug.cpp
@@ -1088,19 +1088,8 @@ static void AddHairDebugPrintMemoryPass(
 	const FSceneView& View,
 	const FShaderPrintData* ShaderPrintData)
 {
-
+	ShaderPrint::SetEnabled(true);
 	
-#if 0
-	// Streaming enabled
-	LOD Index | Auto LOD
-		rest
-		deformed
-		cluster
-		root
-
-		Total CPU | Total GPU
-#endif
-
 	if (!ShaderPrintData) { return; }
 
 	struct FInfos
@@ -1273,7 +1262,6 @@ static void AddHairDebugPrintMemoryPass(
 	if (BindingBuffer.IsEmpty())	{ BindingBuffer.AddDefaulted(); }
 
 	// Force ShaderPrint on.
-	ShaderPrint::SetEnabled(true);
 	ShaderPrint::RequestSpaceForCharacters((ComponentCount + GroomCount + BindingCount) * 256 + 512);
 
 	FHairDebugPrintMemoryCS::FParameters* Parameters = GraphBuilder.AllocParameters<FHairDebugPrintMemoryCS::FParameters>();
