Adapt groom memory allocation when data are streamed out.

diff --git a/Engine/Plugins/Runtime/HairStrands/Source/HairStrandsCore/Private/GroomManager.cpp b/Engine/Plugins/Runtime/HairStrands/Source/HairStrandsCore/Private/GroomManager.cpp
--- a/Engine/Plugins/Runtime/HairStrands/Source/HairStrandsCore/Private/GroomManager.cpp
+++ b/Engine/Plugins/Runtime/HairStrands/Source/HairStrandsCore/Private/GroomManager.cpp
@@ -48,6 +48,8 @@ static int32 GHairStrands_Streaming = 0;
 static FAutoConsoleVariableRef CVarHairStrands_Streaming(TEXT("r.HairStrands.Streaming"), GHairStrands_Streaming, TEXT("Hair strands streaming toggle."), ECVF_RenderThreadSafe | ECVF_ReadOnly);
 static int32 GHairStrands_Streaming_CurvePage = 2048;
 static FAutoConsoleVariableRef CVarHairStrands_StreamingCurvePage(TEXT("r.HairStrands.Streaming.CurvePage"), GHairStrands_Streaming_CurvePage, TEXT("Number of strands curve per streaming page"));
+static int32 GHairStrands_Streaming_StreamOutThreshold = 2;
+static FAutoConsoleVariableRef CVarHairStrands_Streaming_StreamOutThreshold(TEXT("r.HairStrands.Streaming.StreamOutThreshold"), GHairStrands_Streaming_StreamOutThreshold, TEXT("Threshold used for streaming out data. In curve page. Default:2."));
 
 static int32 GHairStrands_AutoLOD_Force = 0;
 static float GHairStrands_AutoLOD_Scale = 1.f;
@@ -84,6 +86,22 @@ static uint32 GetRoundedCurveCount(uint32 InRequest, uint32 InMaxCurve)
 	return FMath::Min(FMath::DivideAndRoundUp(InRequest, Page) * Page,  InMaxCurve);
 }
 
+bool NeedDeallocation(uint32 InRequest, uint32 InAvailable)
+{
+	bool bNeedDeallocate = false;
+	if (GHairStrands_Streaming > 0)
+	{
+
+		const uint32 Page = GetStreamingCurvePage();
+		const uint32 RequestedPageCount = FMath::DivideAndRoundUp(InRequest, Page);
+		const uint32 AvailablePageCount = FMath::DivideAndRoundUp(InAvailable, Page);
+	
+		const uint32 DeallactionThreshold = FMath::Max(GHairStrands_Streaming_StreamOutThreshold, 2);
+		bNeedDeallocate = AvailablePageCount - RequestedPageCount >= DeallactionThreshold;
+	}
+	return bNeedDeallocate;
+}
+
 DEFINE_LOG_CATEGORY_STATIC(LogGroomManager, Log, All);
 
 ///////////////////////////////////////////////////////////////////////////////////////////////////
diff --git a/Engine/Plugins/Runtime/HairStrands/Source/HairStrandsCore/Private/GroomResources.cpp b/Engine/Plugins/Runtime/HairStrands/Source/HairStrandsCore/Private/GroomResources.cpp
--- a/Engine/Plugins/Runtime/HairStrands/Source/HairStrandsCore/Private/GroomResources.cpp
+++ b/Engine/Plugins/Runtime/HairStrands/Source/HairStrandsCore/Private/GroomResources.cpp
@@ -257,13 +257,26 @@ static FRDGBufferRef InternalCreateBufferRDG_FromHairBulkData(FRDGBuilder& Graph
 	check(InData.ChunkRequest);
 	FHairStreamingRequest::FChunk& InChunk = *InData.ChunkRequest;
 
+	const bool bDeallocate  = In != nullptr && InData.ChunkRequest->Status == FHairStreamingRequest::FChunk::EStatus::Unloading;
 	const bool bCreate 		= In == nullptr || BufferDesc.GetSize() == UploadDesc.GetSize();
 	const bool bCopy 		= In != nullptr && InChunk.Size > 0;
 	const bool bReallocate 	= In != nullptr && In->Desc.GetSize() < InChunk.TotalSize;
 
 	check(InChunk.Size > 0);
 	// Either create a new buffer or append new data to existing buffer
-	if (bCreate)
+	if (bDeallocate)
+	{
+		// 1. Create new (smaller) buffer 
+		FRDGBufferDesc NewBufferDesc = BufferDesc;
+		NewBufferDesc.Usage |= EBufferUsageFlags::UnorderedAccess;
+		FRDGBufferRef Out = GraphBuilder.CreateBuffer(NewBufferDesc, DebugName, ERDGBufferFlags::MultiFrame);
+
+		// 2. Copy existing data from the old buffer to the new buffer
+		check (Out->Desc.GetSize() <= In->Desc.GetSize());
+		AddCopyBufferPass(GraphBuilder, Out, 0, In, 0, Out->Desc.GetSize());
+		return Out;
+	}
+	else if (bCreate)
 	{	
 		check(BufferDesc.GetSize() >= UploadDesc.GetSize());
 
@@ -673,6 +686,7 @@ void FHairCommonResource::Allocate(FRDGBuilder& GraphBuilder, EHairResourceLoadi
 	Allocate(GraphBuilder, LoadingType, Status, HAIR_MAX_NUM_CURVE_PER_GROUP, HAIR_MAX_NUM_POINT_PER_GROUP, InLODIndex);
 }
 
+bool NeedDeallocation(uint32 InRequest, uint32 InAvailable);
 void FHairCommonResource::Allocate(FRDGBuilder& GraphBuilder, EHairResourceLoadingType LoadingType, EHairResourceStatus& Status, uint32 InRequestedCurveCount, uint32 InRequestedPointCount, int32 InLODIndex)
 {
 	check(AllocationType == EHairStrandsAllocationType::Deferred);
@@ -701,6 +715,25 @@ void FHairCommonResource::Allocate(FRDGBuilder& GraphBuilder, EHairResourceLoadi
 		// 1. If all requested curve are already loaded, nothing to do
 		if (bIsInitialized && MaxAvailableCurveCount >= InRequestedCurveCount && InternalIsLODDataLoaded(InRequestedCurveCount, InRequestedPointCount, InLODIndex)) 
 		{ 
+			// Trim/Un-stream data if needed
+			if (NeedDeallocation(InRequestedCurveCount, MaxAvailableCurveCount))
+			{
+				StreamingRequest.CurveCount = InRequestedCurveCount;
+				StreamingRequest.PointCount = InRequestedPointCount;
+				StreamingRequest.LODIndex = InLODIndex;
+				if (FHairStrandsBulkCommon* Bulk = InternalGetBulkData())
+				{
+					Bulk->Unload(&StreamingRequest);
+				}
+				InternalAllocate(GraphBuilder, InRequestedCurveCount, InRequestedPointCount, StreamingRequest.LODIndex); // Do we need to pass something else for the deallocation
+
+				// Update the max curve count available
+				MaxAvailableCurveCount = StreamingRequest.CurveCount;
+
+				// Reset streaming request. When the request is delete, the DDC request becomes cancelled. 
+				StreamingRequest = FHairStreamingRequest();
+			}
+
 			Status |= EHairResourceStatus::EStatus::Valid; 
 		}
 		// 2. If more curves are requested, issue a streaming request
diff --git a/Engine/Plugins/Runtime/HairStrands/Source/HairStrandsCore/Private/HairStrandsDatas.cpp b/Engine/Plugins/Runtime/HairStrands/Source/HairStrandsCore/Private/HairStrandsDatas.cpp
--- a/Engine/Plugins/Runtime/HairStrands/Source/HairStrandsCore/Private/HairStrandsDatas.cpp
+++ b/Engine/Plugins/Runtime/HairStrands/Source/HairStrandsCore/Private/HairStrandsDatas.cpp
@@ -175,6 +175,14 @@ void FHairStrandsBulkCommon::Write_IO(UObject* Owner, FArchive& Out, int32 InLOD
 	Q.LODIndex = InLODIndex;
 	GetResources(Q);
 }
+void FHairStrandsBulkCommon::Unload(FHairStreamingRequest* In)
+{
+	FHairStrandsBulkCommon::FQuery Q;
+	Q.Type = FHairStrandsBulkCommon::FQuery::UnloadData;
+	Q.StreamingRequest = In;
+	Q.StreamingRequest->Chunks.Reserve(GetResourceCount()); // This ensures that Chunks array is never reallocated, which would invalidate pointers to FChunk
+	GetResources(Q);
+}
 
 void FHairStrandsBulkCommon::Serialize(FArchive& Ar, UObject* Owner)
 {
@@ -187,7 +195,7 @@ void FHairStrandsBulkCommon::SerializeData(FArchive& Ar, UObject* Owner, int32 I
 	Write_IO(Owner, Ar, InLODIndex);
 }
 
-void FHairStrandsBulkCommon::FQuery::Add(FHairBulkContainer& In, const TCHAR* InSuffix, uint32 InOffset, uint32 InSize) 
+void FHairStrandsBulkCommon::FQuery::Add(FHairBulkContainer& In, const TCHAR* InSuffix, uint32& InOffset, uint32 InSize) 
 {
 	check(Type != None);
 #if WITH_EDITORONLY_DATA
@@ -267,9 +275,8 @@ void FHairStrandsBulkCommon::FQuery::Add(FHairBulkContainer& In, const TCHAR* In
 		check(OutReadIO);
 		OutReadIO->Read(In.Data, InOffset, InSize, EAsyncIOPriorityAndFlags::AIOP_Normal, Chunk.Data_IO);
 	}
-	else
+	else if (Type == ReadWriteIO)
 	{
-		check(Type == ReadWriteIO)
 		check(OutWriteIO);
 
 		if (OutWriteIO->IsSaving())
@@ -279,6 +286,22 @@ void FHairStrandsBulkCommon::FQuery::Add(FHairBulkContainer& In, const TCHAR* In
 		}
 		In.Data.Serialize(*OutWriteIO, Owner, 0/*ChunkIndex*/, false /*bAttemptFileMapping*/);
 	}
+	else if (Type == UnloadData)
+	{
+		InOffset = FMath::Min(InOffset, InSize);
+		check (StreamingRequest != nullptr);
+		FHairStreamingRequest::FChunk& Chunk = StreamingRequest->Chunks.AddDefaulted_GetRef();
+		Chunk.Status 	= FHairStreamingRequest::FChunk::EStatus::Unloading;
+		Chunk.Container = &In;
+		Chunk.Size 		= InSize;
+		Chunk.Offset 	= InOffset;
+		Chunk.TotalSize = InSize;
+		In.ChunkRequest = &Chunk;
+	}
+	else
+	{
+		checkNoEntry();
+	}
 }
 
 /////////////////////////////////////////////////////////////////////////////////////////
@@ -294,6 +317,11 @@ bool FHairStreamingRequest::IsNone() const
 }
 bool FHairStreamingRequest::IsCompleted()
 { 
+	if (IsUnloading())
+	{
+		return true;
+	}
+
 #if !WITH_EDITORONLY_DATA
 	if (IORequest.IsCompleted())
 	{
@@ -316,6 +344,17 @@ bool FHairStreamingRequest::IsCompleted()
 	return true; 
 #endif
 }
+bool FHairStreamingRequest::IsUnloading() const
+{ 
+	for (const FHairStreamingRequest::FChunk& Chunk : Chunks)
+	{
+		if (Chunk.Status == FHairStreamingRequest::FChunk::Unloading)
+		{
+			return true;
+		}
+	}
+	return false; 
+}
 #if WITH_EDITORONLY_DATA
 static void RequestWarmCache(UE::DerivedData::FRequestOwner* RequestOwner, const TArray<UE::DerivedData::FCacheGetChunkRequest>& Requests)
 {
@@ -529,7 +568,7 @@ void FHairStrandsBulkData::GetResources(FHairStrandsBulkCommon::FQuery& Out)
 	// Translate requested curve count into chunk/offset/size to be read
 	uint32 PointCount = 0;
 	uint32 CurveCount = 0;
-	if (Out.Type == FHairStrandsBulkCommon::FQuery::ReadIO || Out.Type == FHairStrandsBulkCommon::FQuery::ReadDDC)
+	if (Out.Type == FHairStrandsBulkCommon::FQuery::ReadIO || Out.Type == FHairStrandsBulkCommon::FQuery::ReadDDC || Out.Type == FHairStrandsBulkCommon::FQuery::UnloadData)
 	{
 		CurveCount = FMath::Min(Header.CurveCount, Out.GetCurveCount());
 		PointCount = CurveCount > 0 ? Header.CurveToPointCount[CurveCount -1] : 0;
@@ -632,7 +671,7 @@ void FHairStrandsInterpolationBulkData::GetResources(FHairStrandsBulkCommon::FQu
 	{
 		// Translate requested curve count into chunk/offset/size to be read
 		uint32 PointCount = 0;
-		if (Out.Type == FHairStrandsBulkCommon::FQuery::ReadIO || Out.Type == FHairStrandsBulkCommon::FQuery::ReadDDC)
+		if (Out.Type == FHairStrandsBulkCommon::FQuery::ReadIO || Out.Type == FHairStrandsBulkCommon::FQuery::ReadDDC || Out.Type == FHairStrandsBulkCommon::FQuery::UnloadData)
 		{
 			PointCount = FMath::Min(Header.PointCount, Out.GetPointCount());
 		}
@@ -717,7 +756,7 @@ void FHairStrandsClusterCullingBulkData::GetResources(FHairStrandsBulkCommon::FQ
 		// Translate requested curve count into chunk/offset/size to be read
 		uint32 PointCount = 0;
 		uint32 CurveCount = 0;
-		if (Out.Type == FHairStrandsBulkCommon::FQuery::ReadIO || Out.Type == FHairStrandsBulkCommon::FQuery::ReadDDC)
+		if (Out.Type == FHairStrandsBulkCommon::FQuery::ReadIO || Out.Type == FHairStrandsBulkCommon::FQuery::ReadDDC || Out.Type == FHairStrandsBulkCommon::FQuery::UnloadData)
 		{
 			CurveCount = FMath::Min(Header.CurveCount, Out.GetCurveCount());
 			PointCount = FMath::Min(Header.PointCount, Out.GetPointCount());
@@ -841,7 +880,7 @@ void FHairStrandsRootBulkData::GetResources(FQuery& Out)
 	// * Unique triangle data are loaded all at one
 	// * Root data are loaded based on the curve request
 	uint32 RootCount = Header.RootCount;
-	if (Out.Type == FHairStrandsBulkCommon::FQuery::ReadIO || Out.Type == FHairStrandsBulkCommon::FQuery::ReadDDC)
+	if (Out.Type == FHairStrandsBulkCommon::FQuery::ReadIO || Out.Type == FHairStrandsBulkCommon::FQuery::ReadDDC || Out.Type == FHairStrandsBulkCommon::FQuery::UnloadData)
 	{
 		RootCount = FMath::Min(Header.RootCount, Out.GetCurveCount());
 	}
diff --git a/Engine/Plugins/Runtime/HairStrands/Source/HairStrandsCore/Public/GroomResources.h b/Engine/Plugins/Runtime/HairStrands/Source/HairStrandsCore/Public/GroomResources.h
--- a/Engine/Plugins/Runtime/HairStrands/Source/HairStrandsCore/Public/GroomResources.h
+++ b/Engine/Plugins/Runtime/HairStrands/Source/HairStrandsCore/Public/GroomResources.h
@@ -114,6 +114,7 @@ struct FHairCommonResource : public FRenderResource
 	virtual bool InternalGetOrRequestData(uint32 InRequestedCurveCount, uint32 InRequestedPointCount, int32 InLODIndex);
 	virtual bool InternalIsLODDataLoaded(uint32 InRequestedCurveCount, uint32 InRequestedPointCount, int32 InLODIndex) const { return true; }
 	virtual void InternalResetLoadedSize() { }
+	virtual FHairStrandsBulkCommon* InternalGetBulkData() { return nullptr; }
 
 	bool bUseRenderGraph = true;
 	bool bIsInitialized = false;
@@ -142,6 +143,7 @@ struct FHairStrandsRestResource : public FHairCommonResource
 	virtual void InternalRelease() override;
 	virtual bool InternalGetOrRequestData(uint32 InRequestedCurveCount, uint32 InRequestedPointCount, int32 InLODIndex) override;
 	virtual void InternalResetLoadedSize() override;
+	virtual FHairStrandsBulkCommon* InternalGetBulkData() override { return &BulkData; }
 
 	/* Get the resource name */
 	virtual FString GetFriendlyName() const override { return TEXT("FHairStrandsResource"); }
@@ -288,6 +290,7 @@ struct FHairStrandsInterpolationResource : public FHairCommonResource
 	virtual void InternalRelease() override;
 	virtual bool InternalGetOrRequestData(uint32 InRequestedCurveCount, uint32 InRequestedPointCount, int32 InLODIndex) override;
 	virtual void InternalResetLoadedSize() override;
+	virtual FHairStrandsBulkCommon* InternalGetBulkData() override { return &BulkData; }
 
 	/* Get the resource name */
 	virtual FString GetFriendlyName() const override { return TEXT("FHairStrandsInterplationResource"); }
@@ -320,6 +323,7 @@ struct FHairStrandsClusterCullingResource : public FHairCommonResource
 	virtual void InternalRelease() override;
 	virtual bool InternalGetOrRequestData(uint32 InRequestedCurveCount, uint32 InRequestedPointCount, int32 InLODIndex) override;
 	virtual void InternalResetLoadedSize() override;
+	virtual FHairStrandsBulkCommon* InternalGetBulkData() override { return &BulkData; }
 
 	/* Get the resource name */
 	virtual FString GetFriendlyName() const override { return TEXT("FHairStrandsClusterResource"); }
@@ -358,6 +362,7 @@ struct FHairStrandsRestRootResource : public FHairCommonResource
 	virtual bool InternalGetOrRequestData(uint32 InRequestedCurveCount, uint32 InRequestedPointCount, int32 InLODIndex) override;
 	virtual bool InternalIsLODDataLoaded(uint32 InRequestedCurveCount, uint32 InRequestedPointCount, int32 InLODIndex) const override;
 	virtual void InternalResetLoadedSize() override;
+	virtual FHairStrandsBulkCommon* InternalGetBulkData() override { return &BulkData; }
 
 	/* Get the resource name */
 	virtual FString GetFriendlyName() const override { return TEXT("FHairStrandsRestRootResource"); }
diff --git a/Engine/Plugins/Runtime/HairStrands/Source/HairStrandsCore/Public/HairStrandsDatas.h b/Engine/Plugins/Runtime/HairStrands/Source/HairStrandsCore/Public/HairStrandsDatas.h
--- a/Engine/Plugins/Runtime/HairStrands/Source/HairStrandsCore/Public/HairStrandsDatas.h
+++ b/Engine/Plugins/Runtime/HairStrands/Source/HairStrandsCore/Public/HairStrandsDatas.h
@@ -328,7 +328,7 @@ struct FHairStreamingRequest
 	// Hold request at the container level (transient to the request)
 	struct FChunk
 	{
-		enum EStatus { None, Pending, Completed, Failed };
+		enum EStatus { None, Pending, Completed, Failed, Unloading };
 		FSharedBuffer Data_DDC;
 		FIoBuffer Data_IO;
 		uint32 Offset = 0;				// Offset to the requested data
@@ -344,6 +344,7 @@ struct FHairStreamingRequest
 	void Request(uint32 InRequestedCurveCount, uint32 InRequestedPointCount, int32 InLODIndex, FHairStrandsBulkCommon& In, bool bWait=false, bool bFillBulkData=false, bool bWarmCache=false, const FName& InOwnerName = NAME_None);
 	bool IsNone() const;
 	bool IsCompleted();
+	bool IsUnloading() const;
 
 #if WITH_EDITORONLY_DATA
 	void WarmCache(uint32 InRequestedCurveCount, uint32 InRequestedPointCount, int32 InLODIndex, FHairStrandsBulkCommon& In);
@@ -397,14 +398,15 @@ struct HAIRSTRANDSCORE_API FHairStrandsBulkCommon
 	void Read_DDC(FHairStreamingRequest* In, TArray<UE::DerivedData::FCacheGetChunkRequest>& Out);
 	void Write_IO(UObject* Owner, FArchive& Out, int32 LODIndex=-1);
 	void Read_IO(FHairStreamingRequest* In, FBulkDataBatchRequest& Out);
+	void Unload(FHairStreamingRequest* In);
 
 	struct FQuery
 	{
-		void Add(FHairBulkContainer& In, const TCHAR* InSuffix, uint32 InOffset=0, uint32 InSize=0);
+		void Add(FHairBulkContainer& In, const TCHAR* InSuffix, uint32& InOffset, uint32 InSize=0);
 		uint32 GetCurveCount() const { check(StreamingRequest); return StreamingRequest->CurveCount; }
 		uint32 GetPointCount() const { check(StreamingRequest); return StreamingRequest->PointCount; }
 		int32  GetLODIndex() const	 { return StreamingRequest ? StreamingRequest->LODIndex : LODIndex; }
-		enum EQueryType { None, ReadDDC, WriteDDC, ReadIO, ReadWriteIO /* i.e. regular Serialize() */};
+		enum EQueryType { None, ReadDDC, WriteDDC, ReadIO, ReadWriteIO /* i.e. regular Serialize() */, UnloadData};
 		EQueryType Type = None;
 		FHairStreamingRequest* StreamingRequest = nullptr;
 		FBulkDataBatchRequest::FBatchBuilder* 			OutReadIO = nullptr;
