Fix cached tangent with streaming / progressive LOD for static groom.

diff --git a/Engine/Plugins/Runtime/HairStrands/Source/HairStrandsCore/Private/GroomResources.cpp b/Engine/Plugins/Runtime/HairStrands/Source/HairStrandsCore/Private/GroomResources.cpp
--- a/Engine/Plugins/Runtime/HairStrands/Source/HairStrandsCore/Private/GroomResources.cpp
+++ b/Engine/Plugins/Runtime/HairStrands/Source/HairStrandsCore/Private/GroomResources.cpp
@@ -1097,6 +1097,14 @@ void FHairStrandsRestResource::InternalAllocate(FRDGBuilder& GraphBuilder)
 		InternalCreateVertexBufferRDG_FromHairBulkData<FHairStrandsPointToCurveFormat32>(GraphBuilder, BulkData.Data.PointToCurve, PointCount, PointToCurveBuffer, ToHairResourceDebugName(HAIRSTRANDS_RESOUCE_NAME(CurveType, Hair.StrandsRest_PointToCurveBuffer), ResourceName), OwnerName, EHairResourceUsageType::Static);
 	}
 
+	// If the buffer is not null, it means it was been lazily allocated once. In such a case, The tangent buffer is recreated the new required size
+	const bool bStaticTangentNeeded = TangentBuffer.Buffer != nullptr;
+	if (bStaticTangentNeeded)
+	{
+		InternalCreateVertexBufferRDG<FHairStrandsTangentFormat>(GraphBuilder, PointCount * FHairStrandsTangentFormat::ComponentCount, TangentBuffer, ToHairResourceDebugName(TEXT("Hair.StrandsRest_TangentBuffer"), ResourceName), OwnerName, EHairResourceUsageType::Dynamic);
+		CachedTangentPointCount = 0; // Force recaching
+	}
+
 	if (PositionOffsetBuffer.Buffer == nullptr)
 	{	
 		TArray<FVector4f> RestOffset;
@@ -1109,25 +1117,34 @@ void FHairStrandsRestResource::InternalAllocate(FRDGBuilder& GraphBuilder)
 void AddHairTangentPass(
 	FRDGBuilder& GraphBuilder,
 	FGlobalShaderMap* ShaderMap,
-	uint32 VertexCount,
+	uint32 PointCount,
 	FHairGroupPublicData* HairGroupPublicData,
 	FRDGBufferSRVRef PositionBuffer,
 	FRDGImportedBuffer OutTangentBuffer);
 
-FRDGExternalBuffer FHairStrandsRestResource::GetTangentBuffer(FRDGBuilder& GraphBuilder, FGlobalShaderMap* ShaderMap)
+FRDGExternalBuffer FHairStrandsRestResource::GetTangentBuffer(FRDGBuilder& GraphBuilder, FGlobalShaderMap* ShaderMap, uint32 ActivePointCount, uint32 ActiveCurveCount)
 {
 	// Lazy allocation and update
 	if (TangentBuffer.Buffer == nullptr)
 	{
-		InternalCreateVertexBufferRDG<FHairStrandsTangentFormat>(GraphBuilder, BulkData.GetNumPoints() * FHairStrandsTangentFormat::ComponentCount, TangentBuffer, ToHairResourceDebugName(TEXT("Hair.StrandsRest_TangentBuffer"), ResourceName), OwnerName, EHairResourceUsageType::Dynamic);
+		InternalCreateVertexBufferRDG<FHairStrandsTangentFormat>(GraphBuilder, ActivePointCount * FHairStrandsTangentFormat::ComponentCount, TangentBuffer, ToHairResourceDebugName(TEXT("Hair.StrandsRest_TangentBuffer"), ResourceName), OwnerName, EHairResourceUsageType::Dynamic);
+	}
 
+	const uint32 AllocatedPoint = TangentBuffer.Buffer->Desc.NumElements / FHairStrandsTangentFormat::ComponentCount;
+	const uint32 RequestedPoint = ActivePointCount;
+
+	const bool bRecomputeTangent = (AllocatedPoint >= RequestedPoint) && (ActivePointCount > CachedTangentPointCount);
+	if (bRecomputeTangent)
+	{
 		AddHairTangentPass(
 			GraphBuilder,
 			ShaderMap,
-			BulkData.GetNumPoints(),
+			ActivePointCount,
 			nullptr,
 			RegisterAsSRV(GraphBuilder, PositionBuffer),
 			Register(GraphBuilder, TangentBuffer, ERDGImportedBufferFlags::CreateUAV));
+
+		CachedTangentPointCount = ActivePointCount;
 	}
 
 	return TangentBuffer;
diff --git a/Engine/Plugins/Runtime/HairStrands/Source/HairStrandsCore/Private/HairStrandsInterpolation.cpp b/Engine/Plugins/Runtime/HairStrands/Source/HairStrandsCore/Private/HairStrandsInterpolation.cpp
--- a/Engine/Plugins/Runtime/HairStrands/Source/HairStrandsCore/Private/HairStrandsInterpolation.cpp
+++ b/Engine/Plugins/Runtime/HairStrands/Source/HairStrandsCore/Private/HairStrandsInterpolation.cpp
@@ -1086,7 +1086,7 @@ class FHairTangentCS : public FGlobalShader
 	using FPermutationDomain = TShaderPermutationDomain<FCulling>;
 
 	BEGIN_SHADER_PARAMETER_STRUCT(FParameters, )
-		SHADER_PARAMETER(uint32, VertexCount)
+		SHADER_PARAMETER(uint32, PointCount)
 		SHADER_PARAMETER(uint32, HairStrandsVF_bCullingEnable)
 		SHADER_PARAMETER_RDG_BUFFER_SRV(Buffer, PositionBuffer)
 		SHADER_PARAMETER_RDG_BUFFER_SRV(Buffer<uint>,	HairStrandsVF_CullingIndirectBuffer)
@@ -1111,7 +1111,7 @@ IMPLEMENT_GLOBAL_SHADER(FHairTangentCS, "/Engine/Private/HairStrands/HairStrands
 void AddHairTangentPass(
 	FRDGBuilder& GraphBuilder,
 	FGlobalShaderMap* ShaderMap,
-	uint32 VertexCount,
+	uint32 PointCount,
 	FHairGroupPublicData* HairGroupPublicData,
 	FRDGBufferSRVRef PositionBuffer,
 	FRDGImportedBuffer OutTangentBuffer)
@@ -1121,7 +1121,7 @@ void AddHairTangentPass(
 	FHairTangentCS::FParameters* Parameters = GraphBuilder.AllocParameters<FHairTangentCS::FParameters>();
 	Parameters->PositionBuffer = PositionBuffer;
 	Parameters->OutputTangentBuffer = OutTangentBuffer.UAV;
-	Parameters->VertexCount = VertexCount;
+	Parameters->PointCount = PointCount;
 	Parameters->HairStrandsVF_bCullingEnable = bCullingEnable ? 1 : 0;
 
 	FHairTangentCS::FPermutationDomain PermutationVector;
@@ -1145,7 +1145,7 @@ void AddHairTangentPass(
 	}
 	else
 	{
-		const FIntVector DispatchCount(FMath::DivideAndRoundUp(VertexCount, FHairTangentCS::GetGroupSize()), 1, 1);
+		const FIntVector DispatchCount(FMath::DivideAndRoundUp(PointCount, FHairTangentCS::GetGroupSize()), 1, 1);
 		FComputeShaderUtils::AddPass(
 			GraphBuilder,
 			RDG_EVENT_NAME("HairStrands::Tangent(culling=off)"),
@@ -2015,8 +2015,8 @@ void ComputeHairStrandsInterpolation(
 			{
 				Strands_PositionSRV = RegisterAsSRV(GraphBuilder, Instance->Strands.RestResource->PositionBuffer);
 				Strands_PositionOffsetSRV = RegisterAsSRV(GraphBuilder, Instance->Strands.RestResource->PositionOffsetBuffer);
-				// Generated the static tangent if they haven;t been generated yet
-				Instance->Strands.RestResource->GetTangentBuffer(GraphBuilder, ShaderMap);
+				// Generated the static tangent if they haven't been generated yet
+				Instance->Strands.RestResource->GetTangentBuffer(GraphBuilder, ShaderMap, ActivePointCount, ActiveCurveCount);
 				Strands_TangentSRV = RegisterAsSRV(GraphBuilder, Instance->Strands.RestResource->TangentBuffer);
 			}
 
diff --git a/Engine/Plugins/Runtime/HairStrands/Source/HairStrandsCore/Public/GroomResources.h b/Engine/Plugins/Runtime/HairStrands/Source/HairStrandsCore/Public/GroomResources.h
--- a/Engine/Plugins/Runtime/HairStrands/Source/HairStrandsCore/Public/GroomResources.h
+++ b/Engine/Plugins/Runtime/HairStrands/Source/HairStrandsCore/Public/GroomResources.h
@@ -148,7 +148,7 @@ struct FHairStrandsRestResource : public FHairCommonResource
 	/* Get the resource name */
 	virtual FString GetFriendlyName() const override { return TEXT("FHairStrandsResource"); }
 
-	FRDGExternalBuffer GetTangentBuffer(class FRDGBuilder& GraphBuilder, class FGlobalShaderMap* ShaderMap);
+	FRDGExternalBuffer GetTangentBuffer(class FRDGBuilder& GraphBuilder, class FGlobalShaderMap* ShaderMap, uint32 PointCount, uint32 CurveCount);
 	
 	FVector GetPositionOffset() const { return BulkData.GetPositionOffset(); }
 
@@ -198,6 +198,8 @@ struct FHairStrandsRestResource : public FHairCommonResource
 
 	inline uint32 GetPointCount() const { return BulkData.GetNumPoints(); }
 	inline uint32 GetCurveCount() const { return BulkData.GetNumCurves(); }
+
+	uint32 CachedTangentPointCount = 0;
 };
 
 struct FHairStrandsDeformedResource : public FHairCommonResource
diff --git a/Engine/Shaders/Private/HairStrands/HairStrandsTangent.usf b/Engine/Shaders/Private/HairStrands/HairStrandsTangent.usf
--- a/Engine/Shaders/Private/HairStrands/HairStrandsTangent.usf
+++ b/Engine/Shaders/Private/HairStrands/HairStrandsTangent.usf
@@ -6,7 +6,7 @@
 #include "HairStrandsClusterCommon.ush"
 
 #if SHADER_TANGENT
-uint VertexCount;
+uint PointCount;
 
 Buffer<uint4>  	 PositionBuffer;
 RWBuffer<float4> OutputTangentBuffer;
@@ -27,22 +27,22 @@ float4 PackNormal(float3 V, uint A)
 [numthreads(GROUP_SIZE, 1, 1)]
 void MainCS(uint DispatchThreadId : SV_DispatchThreadID)
 {
-	const int MaxVertexCount = int(VertexCount);
+	const int MaxPointCount = int(PointCount);
 	int IndexCurr = DispatchThreadId;
-	int IndexPrev = clamp(IndexCurr-1, 0, MaxVertexCount-1);
-	int IndexNext = clamp(IndexCurr+1, 0, MaxVertexCount-1);
-	bool bValid = IndexCurr < MaxVertexCount;
+	int IndexPrev = clamp(IndexCurr-1, 0, MaxPointCount-1);
+	int IndexNext = clamp(IndexCurr+1, 0, MaxPointCount-1);
+	bool bValid = IndexCurr < MaxPointCount;
 #if PERMUTATION_CULLING
 	if (HairStrandsVF_bCullingEnable)
 	{
 		const int FetchIndex = DispatchThreadId;
-		const int VertexCountAfterCulling = clamp((int)HairStrandsVF_CullingIndirectBuffer[3], 0, MaxVertexCount);
+		const int PointCountAfterCulling = clamp((int)HairStrandsVF_CullingIndirectBuffer[3], 0, MaxPointCount);
 		bValid = false;
-		if (FetchIndex < VertexCountAfterCulling)
+		if (FetchIndex < PointCountAfterCulling)
 		{
-			const int FetchIndexPrev = clamp(FetchIndex-1, 0, VertexCountAfterCulling-1);
-			const int FetchIndexCurr = clamp(FetchIndex  , 0, VertexCountAfterCulling-1);
-			const int FetchIndexNext = clamp(FetchIndex+1, 0, VertexCountAfterCulling-1);
+			const int FetchIndexPrev = clamp(FetchIndex-1, 0, PointCountAfterCulling-1);
+			const int FetchIndexCurr = clamp(FetchIndex  , 0, PointCountAfterCulling-1);
+			const int FetchIndexNext = clamp(FetchIndex+1, 0, PointCountAfterCulling-1);
 
 			IndexPrev = HairStrandsVF_CullingIndexBuffer[FetchIndexPrev];
 			IndexCurr = HairStrandsVF_CullingIndexBuffer[FetchIndexCurr];
@@ -53,7 +53,7 @@ void MainCS(uint DispatchThreadId : SV_DispatchThreadID)
 	}
 #endif
 
-	if (IndexCurr < MaxVertexCount)
+	if (IndexCurr < MaxPointCount)
 	{
 		const FHairControlPoint CPCurr = UnpackHairControlPoint(PositionBuffer[IndexCurr]);
 		const FHairControlPoint CPPrev = UnpackHairControlPoint(PositionBuffer[IndexPrev]);
