Reallocate RT groom resources based on streaming data.

diff --git a/Engine/Plugins/Runtime/HairStrands/Source/HairStrandsCore/Private/GroomDebug.cpp b/Engine/Plugins/Runtime/HairStrands/Source/HairStrandsCore/Private/GroomDebug.cpp
--- a/Engine/Plugins/Runtime/HairStrands/Source/HairStrandsCore/Private/GroomDebug.cpp
+++ b/Engine/Plugins/Runtime/HairStrands/Source/HairStrandsCore/Private/GroomDebug.cpp
@@ -901,8 +901,7 @@ static void AddHairDebugPrintInstancePass(
 			(FFloat16(LODIndex).Encoded) |
 			(FFloat16(Instance->Strands.Modifier.HairLengthScale_Override ? Instance->Strands.Modifier.HairLengthScale : -1.f).Encoded << 16);
 		
-		D.Data4.X = Instance->HairGroupPublicData->VFInput.bHasLODSwitch ? 1u : 0u;
-
+		bool bHasRaytracing = false;
 		switch (Instance->GeometryType) 
 		{
 		case EHairGeometryType::Strands:
@@ -949,6 +948,10 @@ static void AddHairDebugPrintInstancePass(
 				D.Data3.Z |= FFloat16(Instance->HairGroupPublicData->VFInput.Strands.Common.LengthScale).Encoded << 16u;
 				D.Data3.W |= FFloat16(Instance->HairGroupPublicData->GetActiveStrandsCoverageScale()).Encoded;
 				D.Data3.W |= Instance->HairGroupPublicData->bAutoLOD || IsHairStrandsForceAutoLODEnabled() ? (1u << 16u) : 0;
+
+				#if RHI_RAYTRACING
+				bHasRaytracing = Instance->Strands.RenRaytracingResource != nullptr;
+				#endif
 			}
 			break;
 		case EHairGeometryType::Cards:
@@ -960,6 +963,10 @@ static void AddHairDebugPrintInstancePass(
 				D.Data2 = FUintVector4(0);
 				D.Data2.X |= InstanceIndex < InstanceCountPerType[HairInstanceCount_CardsOrMeshesPrimaryView] ? 0x1u : 0u;
 				D.Data2.X |= InstanceIndex < InstanceCountPerType[HairInstanceCount_CardsOrMeshesShadowView]  ? 0x2u : 0u;
+
+				#if RHI_RAYTRACING
+				bHasRaytracing = Instance->Cards.LODs[IntLODIndex].RaytracingResource != nullptr;
+				#endif
 			}
 			break;
 		case EHairGeometryType::Meshes:
@@ -971,9 +978,17 @@ static void AddHairDebugPrintInstancePass(
 				D.Data2 = FUintVector4(0);
 				D.Data2.X |= InstanceIndex < InstanceCountPerType[HairInstanceCount_CardsOrMeshesPrimaryView] ? 0x1u : 0u;
 				D.Data2.X |= InstanceIndex < InstanceCountPerType[HairInstanceCount_CardsOrMeshesShadowView]  ? 0x2u : 0u;
+
+				#if RHI_RAYTRACING
+				bHasRaytracing = Instance->Meshes.LODs[IntLODIndex].RaytracingResource != nullptr;
+				#endif
 			}
 			break;
 		}
+
+		D.Data4.X = 0;
+		D.Data4.X |= Instance->HairGroupPublicData->VFInput.bHasLODSwitch ? 1u : 0u;
+		D.Data4.X |= bHasRaytracing ? 2u : 0u;
 	}
 
 	ShaderPrint::FStrings AttributeNames;
diff --git a/Engine/Plugins/Runtime/HairStrands/Source/HairStrandsCore/Private/GroomManager.cpp b/Engine/Plugins/Runtime/HairStrands/Source/HairStrandsCore/Private/GroomManager.cpp
--- a/Engine/Plugins/Runtime/HairStrands/Source/HairStrandsCore/Private/GroomManager.cpp
+++ b/Engine/Plugins/Runtime/HairStrands/Source/HairStrandsCore/Private/GroomManager.cpp
@@ -1079,7 +1079,7 @@ static void RunHairLODSelection(
 				if (Instance->Strands.DeformedResource)			{ Instance->Strands.DeformedResource->Allocate(GraphBuilder, EHairResourceLoadingType::Async, ResourceStatus, RequestedCurveCount, RequestedPointCount); }
 				#if RHI_RAYTRACING
 				if (bHasPathTracingView)						{ AllocateRaytracingResources(Instance); }
-				if (Instance->Strands.RenRaytracingResource)	{ Instance->Strands.RenRaytracingResource->Allocate(GraphBuilder, LoadingType, ResourceStatus); }
+				if (Instance->Strands.RenRaytracingResource)	{ Instance->Strands.RenRaytracingResource->Allocate(GraphBuilder, EHairResourceLoadingType::Async, ResourceStatus, RequestedCurveCount, RequestedPointCount); }
 				#endif
 				Instance->Strands.VertexFactory->InitResources(GraphBuilder.RHICmdList);
 
diff --git a/Engine/Plugins/Runtime/HairStrands/Source/HairStrandsCore/Private/GroomResources.cpp b/Engine/Plugins/Runtime/HairStrands/Source/HairStrandsCore/Private/GroomResources.cpp
--- a/Engine/Plugins/Runtime/HairStrands/Source/HairStrandsCore/Private/GroomResources.cpp
+++ b/Engine/Plugins/Runtime/HairStrands/Source/HairStrandsCore/Private/GroomResources.cpp
@@ -1637,6 +1637,22 @@ void FHairCardsInterpolationResource::InternalRelease()
 /////////////////////////////////////////////////////////////////////////////////////////
 
 #if RHI_RAYTRACING
+bool GetSupportHairStrandsProceduralPrimitive(EShaderPlatform InShaderPlatform);
+
+static void ConvertPointToRaytracingVertexCount(bool bProceduralPrimitive, uint32 InPointCount, uint32& OutVertexCount, uint32& OutIndexCount)
+{
+	if (bProceduralPrimitive)
+	{
+		// only allocate space for primitive AABBs
+		OutVertexCount = InPointCount * 2 * STRANDS_PROCEDURAL_INTERSECTOR_MAX_SPLITS;
+	}
+	else
+	{
+		OutVertexCount = InPointCount * 4;
+		OutIndexCount  = InPointCount * 8 * 3;
+	}
+}
+
 // RT geometry for strands is built as a 4 sided cylinder
 //   each vertex of the curve becomes 4 points
 //   each curve segment turns into 2*4=8 triangles (3 indices for each)
@@ -1646,39 +1662,37 @@ void FHairCardsInterpolationResource::InternalRelease()
 // which used 6*float4 = 48 bytes per vertex
 // NOTE: the vertex buffer is a float4 because it is registered as a UAV for the compute shader to work
 // TODO: use a plain float vertex buffer with 3x the entries instead to save memory? (float3 UAVs are not allowed)
-bool GetSupportHairStrandsProceduralPrimitive(EShaderPlatform InShaderPlatform);
 FHairStrandsRaytracingResource::FHairStrandsRaytracingResource(const FHairStrandsBulkData& InData, const FHairResourceName& ResourceName, const FName& InOwnerName) :
 	FHairCommonResource(EHairStrandsAllocationType::Deferred, ResourceName, InOwnerName), bOwnBuffers(true)
 {
+	MaxVertexCount = 0;
+	MaxIndexCount = 0;
+	MaxAvailableCurveCount = 0;
 	bProceduralPrimitive = GetSupportHairStrandsProceduralPrimitive(GMaxRHIShaderPlatform);
-	if (bProceduralPrimitive)
-	{
-		// only allocate space for primitive AABBs
-		VertexCount = InData.GetNumPoints() * 2 * STRANDS_PROCEDURAL_INTERSECTOR_MAX_SPLITS;
-	}
-	else
-	{
-		VertexCount = InData.GetNumPoints() * 4;
-		IndexCount = InData.GetNumPoints() * 8 * 3;
-	}
 }
 
 FHairStrandsRaytracingResource::FHairStrandsRaytracingResource(const FHairCardsBulkData& InData, const FHairResourceName& ResourceName, const FName& InOwnerName) :
 	FHairCommonResource(EHairStrandsAllocationType::Deferred, ResourceName, InOwnerName),
-	VertexCount(InData.GetNumVertices()), bOwnBuffers(false)
+	MaxVertexCount(InData.GetNumVertices()), bOwnBuffers(false)
 {}
 
 FHairStrandsRaytracingResource::FHairStrandsRaytracingResource(const FHairMeshesBulkData& InData, const FHairResourceName& ResourceName, const FName& InOwnerName) :
 	FHairCommonResource(EHairStrandsAllocationType::Deferred, ResourceName, InOwnerName),
-	VertexCount(InData.GetNumVertices()), bOwnBuffers(false)
+	MaxVertexCount(InData.GetNumVertices()), bOwnBuffers(false)
 {}
 
-void FHairStrandsRaytracingResource::InternalAllocate(FRDGBuilder& GraphBuilder)
+void FHairStrandsRaytracingResource::InternalAllocate(FRDGBuilder& GraphBuilder, uint32 InCurveCount, uint32 InPointCount, int32 InLODIndex)
 {
 	if (bOwnBuffers)
 	{
-		InternalCreateVertexBufferRDG<FHairStrandsRaytracingFormat>(GraphBuilder, VertexCount, PositionBuffer, ToHairResourceDebugName(TEXT("Hair.StrandsRaytracing_PositionBuffer"), ResourceName), OwnerName, EHairResourceUsageType::Dynamic);
-		InternalCreateStructuredBufferRDG<FHairStrandsIndexFormat>(GraphBuilder, IndexCount, IndexBuffer, ToHairResourceDebugName(TEXT("Hair.StrandsRaytracing_IndexBuffer"), ResourceName), OwnerName, EHairResourceUsageType::Dynamic);
+		ConvertPointToRaytracingVertexCount(bProceduralPrimitive, InPointCount, MaxVertexCount, MaxIndexCount);
+
+		InternalCreateVertexBufferRDG<FHairStrandsRaytracingFormat>(GraphBuilder, MaxVertexCount, PositionBuffer, ToHairResourceDebugName(TEXT("Hair.StrandsRaytracing_PositionBuffer"), ResourceName), OwnerName, EHairResourceUsageType::Dynamic);
+		InternalCreateStructuredBufferRDG<FHairStrandsIndexFormat>(GraphBuilder, MaxIndexCount, IndexBuffer, ToHairResourceDebugName(TEXT("Hair.StrandsRaytracing_IndexBuffer"), ResourceName), OwnerName, EHairResourceUsageType::Dynamic);
+
+		// Force reallocating and rebuilding the RT geometry
+		RayTracingGeometry.ReleaseResource();
+		bIsRTGeometryInitialized = false;
 	}
 }
 
diff --git a/Engine/Plugins/Runtime/HairStrands/Source/HairStrandsCore/Private/HairStrandsInterpolation.cpp b/Engine/Plugins/Runtime/HairStrands/Source/HairStrandsCore/Private/HairStrandsInterpolation.cpp
--- a/Engine/Plugins/Runtime/HairStrands/Source/HairStrandsCore/Private/HairStrandsInterpolation.cpp
+++ b/Engine/Plugins/Runtime/HairStrands/Source/HairStrandsCore/Private/HairStrandsInterpolation.cpp
@@ -2099,7 +2099,8 @@ void ComputeHairStrandsInterpolation(
 			#if RHI_RAYTRACING
 			if (Instance->Strands.RenRaytracingResource && (Instance->Strands.ViewRayTracingMask & ViewRayTracingMask) != 0)
 			{
-				const float HairRadiusRT	= Instance->HairGroupPublicData->VFInput.Strands.Common.RaytracingRadiusScale * Instance->HairGroupPublicData->VFInput.Strands.Common.Radius;
+				const float CLODScale = Instance->HairGroupPublicData->ContinuousLODCoverageScale;
+				const float HairRadiusRT	= Instance->HairGroupPublicData->VFInput.Strands.Common.RaytracingRadiusScale * Instance->HairGroupPublicData->VFInput.Strands.Common.Radius * CLODScale;
 				const float HairRootScaleRT = Instance->HairGroupPublicData->VFInput.Strands.Common.RootScale;
 				const float HairTipScaleRT	= Instance->HairGroupPublicData->VFInput.Strands.Common.TipScale;
 				const int ProceduralSplits = GetHairRaytracingProceduralSplits();
@@ -2159,8 +2160,8 @@ void ComputeHairStrandsInterpolation(
 							FBufferRHIRef IndexBuffer(bProceduralPrimitive ? nullptr : Instance->Strands.RenRaytracingResource->IndexBuffer.Buffer->GetRHI());
 							Instance->Strands.CachedProceduralSplits = ProceduralSplits;
 							BuildHairAccelerationStructure_Strands(RHICmdList,
-								Instance->Strands.RenRaytracingResource->VertexCount, // HAIR_TODO
-								Instance->Strands.RenRaytracingResource->IndexCount,
+								Instance->Strands.RenRaytracingResource->MaxVertexCount,
+								Instance->Strands.RenRaytracingResource->MaxIndexCount,
 								PositionBuffer,
 								IndexBuffer,
 								&Instance->Strands.RenRaytracingResource->RayTracingGeometry,
diff --git a/Engine/Plugins/Runtime/HairStrands/Source/HairStrandsCore/Public/GroomResources.h b/Engine/Plugins/Runtime/HairStrands/Source/HairStrandsCore/Public/GroomResources.h
--- a/Engine/Plugins/Runtime/HairStrands/Source/HairStrandsCore/Public/GroomResources.h
+++ b/Engine/Plugins/Runtime/HairStrands/Source/HairStrandsCore/Public/GroomResources.h
@@ -566,7 +566,7 @@ struct FHairStrandsRaytracingResource : public FHairCommonResource
 	FHairStrandsRaytracingResource(const FHairMeshesBulkData& InData, const FHairResourceName& ResourceName, const FName& OwnerName);
 
 	/* Init/Release buffers */
-	virtual void InternalAllocate(FRDGBuilder& GraphBuilder) override;
+	virtual void InternalAllocate(FRDGBuilder& GraphBuilder, uint32 InCurveCount, uint32 InPointCount, int32 InLODIndex) override;
 	virtual void InternalRelease() override;
 
 	/* Get the resource name */
@@ -584,8 +584,8 @@ struct FHairStrandsRaytracingResource : public FHairCommonResource
 	FRDGExternalBuffer PositionBuffer;
 	FRDGExternalBuffer IndexBuffer;
 	FRayTracingGeometry RayTracingGeometry;
-	uint32 VertexCount = 0;
-	uint32 IndexCount = 0;
+	uint32 MaxVertexCount = 0;
+	uint32 MaxIndexCount = 0;
 	bool bOwnBuffers = false;
 	bool bIsRTGeometryInitialized = false;
 	bool bProceduralPrimitive = false;
diff --git a/Engine/Shaders/Private/HairStrands/HairStrandsDebug.usf b/Engine/Shaders/Private/HairStrands/HairStrandsDebug.usf
--- a/Engine/Shaders/Private/HairStrands/HairStrandsDebug.usf
+++ b/Engine/Shaders/Private/HairStrands/HairStrandsDebug.usf
@@ -1169,6 +1169,7 @@ struct FInstanceInfo
 	uint  BindingType;
 	bool  bHasSim;
 	bool  bHasRBF;
+	bool  bHasRTGeom;
 	float LOD;
 	float ClipValue;
 
@@ -1282,7 +1283,8 @@ FInstanceInfo ReadInstanceInfo(uint Index)
 	Out.LengthScale				= f16tof32(Data3.z >> 16u);
 	Out.CoverageScale			= f16tof32(Data3.w & 0xFFFF);
 
-	Out.bHasLODSwitched			= Data4.x > 0;
+	Out.bHasLODSwitched			= Data4.x & 0x1;
+	Out.bHasRTGeom				= Data4.x & 0x2;
 
 	return Out;
 }
@@ -1354,7 +1356,7 @@ void MainCS(uint GroupIndex : SV_GroupIndex, uint3 DispatchThreadId : SV_Dispatc
 	Newline(Context);
 
 	// Draw instance details
-	Print(Context, TEXT("   Group  LOD     Geometry Binding Cache Sim RBF Clip Screen CovSc. Strands            Points             Bind.Tri/Sect Vis. Shad. SwLOD Name"));
+	Print(Context, TEXT("   Group  LOD     Geometry Binding Cache Sim RBF Clip Screen CovSc. Strands            Points             Bind.Tri/Sect Vis. Shad. SwLOD RT  Name"));
 	Newline(Context);
 
 	for (uint InstanceIndex = 0; InstanceIndex < InstanceCount; ++InstanceIndex)
@@ -1452,6 +1454,11 @@ void MainCS(uint GroupIndex : SV_GroupIndex, uint3 DispatchThreadId : SV_Dispatc
 		// Has LOD switched
 		PrintBool(Context, Info.bHasLODSwitched);
 		Print(Context, TEXT("  "));
+
+		// Has RT geom
+		PrintBool(Context, Info.bHasRTGeom);
+		Print(Context, TEXT("  "));
+		
 		
 		// Name
 		PrintInstanceNames(Context, InstanceIndex, FontWhite);
