Rename FHairStrandsClusterCullingX to FHairStrandsClusterX to avoid confusion with FHairStrandsCullingX.

diff --git a/Engine/Plugins/Runtime/HairStrands/Source/HairStrandsCore/Private/GroomAsset.cpp b/Engine/Plugins/Runtime/HairStrands/Source/HairStrandsCore/Private/GroomAsset.cpp
--- a/Engine/Plugins/Runtime/HairStrands/Source/HairStrandsCore/Private/GroomAsset.cpp
+++ b/Engine/Plugins/Runtime/HairStrands/Source/HairStrandsCore/Private/GroomAsset.cpp
@@ -125,9 +125,9 @@ uint32 FHairGroupPlatformData::FStrands::GetDataSize() const
 	Total += FBaseWithInterpolation::GetDataSize();
 	uint32 CTotal = 0;
 	uint32 BTotal = 0;
-	BTotal += ClusterCullingBulkData.Data.CurveToClusterIds.IsBulkDataLoaded()	? ClusterCullingBulkData.Data.CurveToClusterIds.GetBulkDataSize() : 0;
-	CTotal += ClusterCullingBulkData.Data.PackedClusterInfos.IsBulkDataLoaded()	? ClusterCullingBulkData.Data.PackedClusterInfos.GetBulkDataSize() : 0;
-	CTotal += ClusterCullingBulkData.Data.PointLODs.IsBulkDataLoaded()	? ClusterCullingBulkData.Data.PointLODs.GetBulkDataSize() : 0;
+	BTotal += ClusterBulkData.Data.CurveToClusterIds.IsBulkDataLoaded() ? ClusterBulkData.Data.CurveToClusterIds.GetBulkDataSize() : 0;
+	CTotal += ClusterBulkData.Data.PackedClusterInfos.IsBulkDataLoaded()? ClusterBulkData.Data.PackedClusterInfos.GetBulkDataSize() : 0;
+	CTotal += ClusterBulkData.Data.PointLODs.IsBulkDataLoaded()	? ClusterBulkData.Data.PointLODs.GetBulkDataSize() : 0;
 
 	return BTotal + Total + CTotal;
 }
@@ -149,13 +149,13 @@ FGroomAssetMemoryStats FGroomAssetMemoryStats::Get(const FHairGroupPlatformData&
 	// Strands only
 	if (In.Strands.RestResource) 			Out.Memory.Rest 			= In.Strands.RestResource->GetResourcesSize();
 	if (In.Strands.InterpolationResource) 	Out.Memory.Interpolation 	= In.Strands.InterpolationResource->GetResourcesSize();
-	if (In.Strands.ClusterCullingResource)	Out.Memory.Cluster 			= In.Strands.ClusterCullingResource->GetResourcesSize();
+	if (In.Strands.ClusterResource)			Out.Memory.Cluster 			= In.Strands.ClusterResource->GetResourcesSize();
 #if RHI_RAYTRACING
 	if (In.Strands.RaytracingResource) 		Out.Memory.Raytracing 		= In.Strands.RaytracingResource->GetResourcesSize();
 #endif
 	if (In.Strands.RestResource) 			Out.Curves.Rest 			= In.Strands.RestResource->MaxAvailableCurveCount;
 	if (In.Strands.InterpolationResource) 	Out.Curves.Interpolation 	= In.Strands.InterpolationResource->MaxAvailableCurveCount;
-	if (In.Strands.ClusterCullingResource)	Out.Curves.Cluster 			= In.Strands.ClusterCullingResource->MaxAvailableCurveCount;
+	if (In.Strands.ClusterResource)			Out.Curves.Cluster 			= In.Strands.ClusterResource->MaxAvailableCurveCount;
 #if RHI_RAYTRACING
 	if (In.Strands.RaytracingResource)		Out.Curves.Raytracing 		= In.Strands.RaytracingResource->MaxAvailableCurveCount;
 #endif
@@ -532,7 +532,7 @@ static bool BuildHairGroup(
 			OutHairGroupsData[GroupIndex].Strands.InterpolationBulkData.Reset();
 		}
 
-		FGroomBuilder::BuildClusterBulkData(StrandsData, HairDescriptionGroups.Bounds.SphereRadius, InHairGroupsLOD[GroupIndex], OutHairGroupsData[GroupIndex].Strands.ClusterCullingBulkData);
+		FGroomBuilder::BuildClusterBulkData(StrandsData, HairDescriptionGroups.Bounds.SphereRadius, InHairGroupsLOD[GroupIndex], OutHairGroupsData[GroupIndex].Strands.ClusterBulkData);
 	}
 	return bIsValid;
 }
@@ -559,7 +559,7 @@ static bool BuildHairGroupCluster(
 		FHairStrandsDatas StrandsData;
 		FHairStrandsDatas GuidesData;
 		FGroomBuilder::BuildData(HairGroup, InHairGroupsInterpolation[GroupIndex], OutHairGroupsInfo[GroupIndex], StrandsData, GuidesData);
-		FGroomBuilder::BuildClusterBulkData(StrandsData, InHairDescriptionGroups.Bounds.SphereRadius, InHairGroupsLOD[GroupIndex], OutHairGroupsData[GroupIndex].Strands.ClusterCullingBulkData);
+		FGroomBuilder::BuildClusterBulkData(StrandsData, InHairDescriptionGroups.Bounds.SphereRadius, InHairGroupsLOD[GroupIndex], OutHairGroupsData[GroupIndex].Strands.ClusterBulkData);
 	}
 	return bIsValid;
 }
@@ -1037,11 +1037,11 @@ void UGroomAsset::UpdateResource()
 					GetHairGroupsInfo(),
 					GetHairGroupsPlatformData());
 
-				GroupData.Strands.ClusterCullingResource = new FHairStrandsClusterCullingResource(GroupData.Strands.ClusterCullingBulkData, FHairResourceName(GetFName(), GroupIndex), GetAssetPathName());
+				GroupData.Strands.ClusterResource = new FHairStrandsClusterResource(GroupData.Strands.ClusterBulkData, FHairResourceName(GetFName(), GroupIndex), GetAssetPathName());
 			}
 			else
 			{
-				InternalUpdateResource(GroupData.Strands.ClusterCullingResource);
+				InternalUpdateResource(GroupData.Strands.ClusterResource);
 			}
 
 			#if RHI_RAYTRACING
@@ -1114,7 +1114,7 @@ void UGroomAsset::ReleaseStrandsResource(uint32 GroupIndex)
 	if (GroupData.Strands.IsValid())
 	{
 		InternalReleaseResource(GroupData.Strands.RestResource);
-		InternalReleaseResource(GroupData.Strands.ClusterCullingResource);
+		InternalReleaseResource(GroupData.Strands.ClusterResource);
 		InternalReleaseResource(GroupData.Strands.InterpolationResource);
 		#if RHI_RAYTRACING
 		InternalReleaseResource(GroupData.Strands.RaytracingResource);
@@ -1842,7 +1842,7 @@ static void InternalSerializeStrand(FArchive& Ar, UObject* Owner, FHairGroupPlat
 	{
 		{ FHairStreamingRequest R; R.Request(HAIR_MAX_NUM_CURVE_PER_GROUP, HAIR_MAX_NUM_POINT_PER_GROUP, -1/*LODIndex*/, StrandData.BulkData,               true /*bWait*/, true /*bFillBulkdata*/, true /*bWarmCache*/, Owner->GetFName()); }
 		{ FHairStreamingRequest R; R.Request(HAIR_MAX_NUM_CURVE_PER_GROUP, HAIR_MAX_NUM_POINT_PER_GROUP, -1/*LODIndex*/, StrandData.InterpolationBulkData,  true /*bWait*/, true /*bFillBulkdata*/, true /*bWarmCache*/, Owner->GetFName()); }
-		{ FHairStreamingRequest R; R.Request(HAIR_MAX_NUM_CURVE_PER_GROUP, HAIR_MAX_NUM_POINT_PER_GROUP, -1/*LODIndex*/, StrandData.ClusterCullingBulkData, true /*bWait*/, true /*bFillBulkdata*/, true /*bWarmCache*/, Owner->GetFName()); }
+		{ FHairStreamingRequest R; R.Request(HAIR_MAX_NUM_CURVE_PER_GROUP, HAIR_MAX_NUM_POINT_PER_GROUP, -1/*LODIndex*/, StrandData.ClusterBulkData, true /*bWait*/, true /*bFillBulkdata*/, true /*bWarmCache*/, Owner->GetFName()); }
 	}
 
 	if (!Ar.IsCooking() || !StrandData.bIsCookedOut)
@@ -1853,8 +1853,8 @@ static void InternalSerializeStrand(FArchive& Ar, UObject* Owner, FHairGroupPlat
 		if (bData)		{ StrandData.InterpolationBulkData.SerializeData(Ar, Owner); }
 		if (Ar.CustomVer(FAnimObjectVersion::GUID) >= FAnimObjectVersion::SerializeHairClusterCullingData)
 		{
-			if (bHeader){ StrandData.ClusterCullingBulkData.SerializeHeader(Ar, Owner); }
-			if (bData)	{ StrandData.ClusterCullingBulkData.SerializeData(Ar, Owner); }
+			if (bHeader){ StrandData.ClusterBulkData.SerializeHeader(Ar, Owner); }
+			if (bData)	{ StrandData.ClusterBulkData.SerializeData(Ar, Owner); }
 		}
 
 		#if WITH_EDITORONLY_DATA
@@ -1864,7 +1864,7 @@ static void InternalSerializeStrand(FArchive& Ar, UObject* Owner, FHairGroupPlat
 		{
 			{ FHairStreamingRequest R; R.WarmCache(HAIR_MAX_NUM_CURVE_PER_GROUP, HAIR_MAX_NUM_POINT_PER_GROUP, -1/*LODIndex*/, StrandData.BulkData); }
 			{ FHairStreamingRequest R; R.WarmCache(HAIR_MAX_NUM_CURVE_PER_GROUP, HAIR_MAX_NUM_POINT_PER_GROUP, -1/*LODIndex*/, StrandData.InterpolationBulkData); }
-			{ FHairStreamingRequest R; R.WarmCache(HAIR_MAX_NUM_CURVE_PER_GROUP, HAIR_MAX_NUM_POINT_PER_GROUP, -1/*LODIndex*/, StrandData.ClusterCullingBulkData); }
+			{ FHairStreamingRequest R; R.WarmCache(HAIR_MAX_NUM_CURVE_PER_GROUP, HAIR_MAX_NUM_POINT_PER_GROUP, -1/*LODIndex*/, StrandData.ClusterBulkData); }
 		}
 		#endif
 	}
@@ -1880,8 +1880,8 @@ static void InternalSerializeStrand(FArchive& Ar, UObject* Owner, FHairGroupPlat
 		if (bData) 		{ NoStrandsData.InterpolationBulkData.SerializeData(Ar, Owner); }
 		if (Ar.CustomVer(FAnimObjectVersion::GUID) >= FAnimObjectVersion::SerializeHairClusterCullingData)
 		{
-			if (bHeader){ NoStrandsData.ClusterCullingBulkData.SerializeHeader(Ar, Owner); }
-			if (bData)	{ NoStrandsData.ClusterCullingBulkData.SerializeData(Ar, Owner); }
+			if (bHeader){ NoStrandsData.ClusterBulkData.SerializeHeader(Ar, Owner); }
+			if (bData)	{ NoStrandsData.ClusterBulkData.SerializeData(Ar, Owner); }
 		}
 	}
 }
@@ -2540,7 +2540,7 @@ bool UGroomAsset::CacheStrandsData(uint32 GroupIndex, FString& OutDerivedDataKey
 	{
 		In.Strands.BulkData.DerivedDataKey = DerivedDataKey + FString(TEXT("_RestData"));
 		In.Strands.InterpolationBulkData.DerivedDataKey = DerivedDataKey + FString(TEXT("_InterpolationData"));
-		In.Strands.ClusterCullingBulkData.DerivedDataKey = DerivedDataKey + FString(TEXT("_ClusterCullingData"));
+		In.Strands.ClusterBulkData.DerivedDataKey = DerivedDataKey + FString(TEXT("_ClusterData"));
 	};
 
 	FHairGroupPlatformData& PlatformData = GetHairGroupsPlatformData()[GroupIndex];
@@ -2628,7 +2628,7 @@ bool UGroomAsset::CacheStrandsData(uint32 GroupIndex, FString& OutDerivedDataKey
 			// Data (cluster Culling)
 			{
 				TArray<FCachePutValueRequest> Out;
-				PlatformData.Strands.ClusterCullingBulkData.Write_DDC(this, Out);
+				PlatformData.Strands.ClusterBulkData.Write_DDC(this, Out);
 
 				FRequestOwner AsyncOwner(EPriority::Normal);
 				GetCache().PutValue(Out, AsyncOwner);
@@ -3123,7 +3123,7 @@ bool UGroomAsset::BuildCardsData()
 
 		// When building cards in an empty groom, the strands resources need to be initialized once
 		FHairGroupPlatformData& GroupData = GetHairGroupsPlatformData()[0];
-		if (!GroupData.Strands.ClusterCullingResource)
+		if (!GroupData.Strands.ClusterResource)
 		{
 			InitStrandsResources();
 		}
@@ -3409,9 +3409,9 @@ void UGroomAsset::InitStrandsResources()
 			FHairResourceName ResourceName(GetFName(), GroupIndex);
 			GroupData.Strands.RestResource = new FHairStrandsRestResource(GroupData.Strands.BulkData, EHairStrandsResourcesType::Strands, ResourceName, OwnerName);
 
-			if (GroupData.Strands.ClusterCullingBulkData.IsValid())
+			if (GroupData.Strands.ClusterBulkData.IsValid())
 			{
-				GroupData.Strands.ClusterCullingResource = new FHairStrandsClusterCullingResource(GroupData.Strands.ClusterCullingBulkData, ResourceName, OwnerName);
+				GroupData.Strands.ClusterResource = new FHairStrandsClusterResource(GroupData.Strands.ClusterBulkData, ResourceName, OwnerName);
 			}
 
 			// Interpolation are lazy allocated, as these resources are only used when RBF/simulation is enabled by a groom component
@@ -3781,7 +3781,7 @@ void UGroomAsset::StripLODs(const TArray<int32>& LODsToKeep, bool bRebuildResour
 		for (int32 GroupIt = 0; GroupIt < GroupCount; ++GroupIt)
 		{
 			FHairGroupPlatformData& GroupData = GetHairGroupsPlatformData()[GroupIt];
-			InternalReleaseResource(GroupData.Strands.ClusterCullingResource);
+			InternalReleaseResource(GroupData.Strands.ClusterResource);
 
 			BuildHairGroupCluster(
 				GroupIt,
diff --git a/Engine/Plugins/Runtime/HairStrands/Source/HairStrandsCore/Private/GroomBuilder.cpp b/Engine/Plugins/Runtime/HairStrands/Source/HairStrandsCore/Private/GroomBuilder.cpp
--- a/Engine/Plugins/Runtime/HairStrands/Source/HairStrandsCore/Private/GroomBuilder.cpp
+++ b/Engine/Plugins/Runtime/HairStrands/Source/HairStrandsCore/Private/GroomBuilder.cpp
@@ -40,7 +40,7 @@ static FAutoConsoleVariableRef CVarHairGroupIndexBuilder_MaxVoxelResolution(TEXT
 
 FString FGroomBuilder::GetVersion()
 {
-	return TEXT("v8r40");
+	return TEXT("v8r41");
 }
 
 namespace FHairStrandsDecimation
@@ -2903,7 +2903,7 @@ static void BuildClusterData(
 	const FHairStrandsDatas& InRenStrandsData,
 	const float InGroomAssetRadius, 
 	const FHairGroupsLOD& InSettings, 
-	FHairStrandsClusterCullingData& Out)
+	FHairStrandsClusterData& Out)
 {
 	// 0. Rest existing culling data
 	Out.Reset();
@@ -3065,8 +3065,8 @@ static void BuildClusterData(
 }
 
 static void BuildClusterBulkData(
-	const FHairStrandsClusterCullingData& In,
-	FHairStrandsClusterCullingBulkData& Out)
+	const FHairStrandsClusterData& In,
+	FHairStrandsClusterBulkData& Out)
 {
 	Out.Reset();
 	Out.Header.ClusterCount		= In.ClusterCount;
@@ -3192,9 +3192,9 @@ void FGroomBuilder::BuildClusterBulkData(
 	const FHairStrandsDatas& InRenStrandsData,
 	const float InGroomAssetRadius,
 	const FHairGroupsLOD& InSettings,
-	FHairStrandsClusterCullingBulkData& Out)
+	FHairStrandsClusterBulkData& Out)
 {
-	FHairStrandsClusterCullingData ClusterData;
+	FHairStrandsClusterData ClusterData;
 	GroomBuilder_Cluster::BuildClusterData(InRenStrandsData, InGroomAssetRadius, InSettings, ClusterData);
 	GroomBuilder_Cluster::BuildClusterBulkData(ClusterData, Out);
 }
diff --git a/Engine/Plugins/Runtime/HairStrands/Source/HairStrandsCore/Private/GroomComponent.cpp b/Engine/Plugins/Runtime/HairStrands/Source/HairStrandsCore/Private/GroomComponent.cpp
--- a/Engine/Plugins/Runtime/HairStrands/Source/HairStrandsCore/Private/GroomComponent.cpp
+++ b/Engine/Plugins/Runtime/HairStrands/Source/HairStrandsCore/Private/GroomComponent.cpp
@@ -485,7 +485,7 @@ public:
 		TArray<FHairGroupInstance*> LocalInstances = HairGroupInstances;
 		for (FHairGroupInstance* Instance : LocalInstances)
 		{
-			if (Instance->IsValid() || Instance->Strands.ClusterCullingResource)
+			if (Instance->IsValid() || Instance->Strands.ClusterResource)
 			{
 				check(Instance->HairGroupPublicData != nullptr);
 				Instance->AddRef();
@@ -509,7 +509,7 @@ public:
 		TArray<FHairGroupInstance*> LocalInstances = HairGroupInstances;
 		for (FHairGroupInstance* Instance : LocalInstances)
 		{
-			if (Instance->IsValid() || Instance->Strands.ClusterCullingResource)
+			if (Instance->IsValid() || Instance->Strands.ClusterResource)
 			{
 				check(Instance->GetRefCount() > 0);
 				LocalScene.RemoveHairStrands(Instance);
@@ -2762,18 +2762,18 @@ void UGroomComponent::InitResources(bool bIsBindingReloading)
 				HairGroupInstance->Strands.DeformedResource = new FHairStrandsDeformedResource(GroupData.Strands.BulkData, EHairStrandsResourcesType::Strands, ResourceName, OwnerName);
 			} 
 
-			// An empty groom doesn't have a ClusterCullingResource
-			HairGroupInstance->Strands.ClusterCullingResource = GroupData.Strands.ClusterCullingResource;
-			if (HairGroupInstance->Strands.ClusterCullingResource)
+			// An empty groom doesn't have a ClusterResource
+			HairGroupInstance->Strands.ClusterResource = GroupData.Strands.ClusterResource;
+			if (HairGroupInstance->Strands.ClusterResource)
 			{
 				// This codes assumes strands LOD are contigus and the highest (i.e., 0...x). Change this code to something more robust
 				check(HairGroupInstance->HairGroupPublicData);
-				const int32 StrandsLODCount = GroupData.Strands.ClusterCullingResource->BulkData.Header.LODInfos.Num();
+				const int32 StrandsLODCount = GroupData.Strands.ClusterResource->BulkData.Header.LODInfos.Num();
 				const TArray<float>& LODScreenSizes = HairGroupInstance->HairGroupPublicData->GetLODScreenSizes();
 				const TArray<bool>& LODVisibilities = HairGroupInstance->HairGroupPublicData->GetLODVisibilities();
 				check(StrandsLODCount <= LODScreenSizes.Num());
 				check(StrandsLODCount <= LODVisibilities.Num());
-				HairGroupInstance->HairGroupPublicData->ClusterCount = HairGroupInstance->Strands.ClusterCullingResource->BulkData.Header.ClusterCount;
+				HairGroupInstance->HairGroupPublicData->ClusterCount = HairGroupInstance->Strands.ClusterResource->BulkData.Header.ClusterCount;
 			}
 
 			HairGroupInstance->Strands.CullingResource = new FHairStrandsCullingResource(
diff --git a/Engine/Plugins/Runtime/HairStrands/Source/HairStrandsCore/Private/GroomManager.cpp b/Engine/Plugins/Runtime/HairStrands/Source/HairStrandsCore/Private/GroomManager.cpp
--- a/Engine/Plugins/Runtime/HairStrands/Source/HairStrandsCore/Private/GroomManager.cpp
+++ b/Engine/Plugins/Runtime/HairStrands/Source/HairStrandsCore/Private/GroomManager.cpp
@@ -779,7 +779,7 @@ void AddHairStreamingRequest(FHairGroupInstance* Instance, int32 InLODIndex)
 		{
 			if (Instance->Strands.RestRootResource)			{ Instance->Strands.RestRootResource->StreamInData(MeshLODIndex); }
 			if (Instance->Strands.RestResource)				{ Instance->Strands.RestResource->StreamInData(); }
-			if (Instance->Strands.ClusterCullingResource)	{ Instance->Strands.ClusterCullingResource->StreamInData(); }
+			if (Instance->Strands.ClusterResource)			{ Instance->Strands.ClusterResource->StreamInData(); }
 			if (Instance->Strands.InterpolationResource)	{ Instance->Strands.InterpolationResource->StreamInData(); }
 		}
 	}
@@ -925,7 +925,7 @@ static void RunHairLODSelection(
 			Instance->HairGroupPublicData->ContinuousLODBounds = SphereBound;
 			Instance->HairGroupPublicData->ContinuousLODCoverageScale = 1.f;
 
-			if (Instance->Strands.ClusterCullingResource)
+			if (Instance->Strands.ClusterResource)
 			{
 				uint32 EffectiveCurveCount = 0;
 				if (Instance->HairGroupPublicData->bAutoLOD || IsHairStrandsForceAutoLODEnabled())
@@ -935,7 +935,7 @@ static void RunHairLODSelection(
 				}
 				else
 				{
-					EffectiveCurveCount = Instance->Strands.ClusterCullingResource->BulkData.GetCurveCount(LODIndex);
+					EffectiveCurveCount = Instance->Strands.ClusterResource->BulkData.GetCurveCount(LODIndex);
 				}
 				check(EffectiveCurveCount <= uint32(Instance->Strands.Data->Header.CurveToPointCount.Num()));
 
@@ -1071,7 +1071,7 @@ static void RunHairLODSelection(
 
 				if (Instance->Strands.RestRootResource)			{ Instance->Strands.RestRootResource->Allocate(GraphBuilder, EHairResourceLoadingType::Async, ResourceStatus, RequestedCurveCount, RequestedPointCount, MeshLODIndex); }
 				if (Instance->Strands.RestResource)				{ Instance->Strands.RestResource->Allocate(GraphBuilder, EHairResourceLoadingType::Async, ResourceStatus, RequestedCurveCount, RequestedPointCount); }
-				if (Instance->Strands.ClusterCullingResource)	{ Instance->Strands.ClusterCullingResource->Allocate(GraphBuilder, EHairResourceLoadingType::Async, ResourceStatus, RequestedCurveCount, RequestedPointCount); }
+				if (Instance->Strands.ClusterResource)			{ Instance->Strands.ClusterResource->Allocate(GraphBuilder, EHairResourceLoadingType::Async, ResourceStatus, RequestedCurveCount, RequestedPointCount); }
 				if (Instance->Strands.InterpolationResource)	{ Instance->Strands.InterpolationResource->Allocate(GraphBuilder, EHairResourceLoadingType::Async, ResourceStatus, RequestedCurveCount, RequestedPointCount); }
 				if (Instance->Strands.CullingResource)			{ Instance->Strands.CullingResource->Allocate(GraphBuilder, EHairResourceLoadingType::Async, ResourceStatus, RequestedCurveCount, RequestedPointCount); }
 
diff --git a/Engine/Plugins/Runtime/HairStrands/Source/HairStrandsCore/Private/GroomManager.h b/Engine/Plugins/Runtime/HairStrands/Source/HairStrandsCore/Private/GroomManager.h
--- a/Engine/Plugins/Runtime/HairStrands/Source/HairStrandsCore/Private/GroomManager.h
+++ b/Engine/Plugins/Runtime/HairStrands/Source/HairStrandsCore/Private/GroomManager.h
@@ -16,7 +16,7 @@
 ////////////////////////////////////////////////////////////////////////////////////////////////////
 // Misc/Helpers
 
-struct FHairStrandClusterCullingData;
+struct FHairStrandClusterData;
 class FRDGPooledBuffer;
 struct IPooledRenderTarget;
 struct FRWBuffer;
diff --git a/Engine/Plugins/Runtime/HairStrands/Source/HairStrandsCore/Private/GroomRBFDeformer.cpp b/Engine/Plugins/Runtime/HairStrands/Source/HairStrandsCore/Private/GroomRBFDeformer.cpp
--- a/Engine/Plugins/Runtime/HairStrands/Source/HairStrandsCore/Private/GroomRBFDeformer.cpp
+++ b/Engine/Plugins/Runtime/HairStrands/Source/HairStrandsCore/Private/GroomRBFDeformer.cpp
@@ -452,7 +452,7 @@ static void ApplyDeformationToGroom(const TArray<FRBFDeformedPositions>& Deforme
 				FGroomBuilder::BuildInterplationData(HairGroup.Info, StrandsData, GuidesData, InterpolationSettings, StrandsInterpolationData);
 				FGroomBuilder::BuildInterplationBulkData(GuidesData, StrandsInterpolationData, HairGroupsData.Strands.InterpolationBulkData);
 
-				FGroomBuilder::BuildClusterBulkData(StrandsData, HairDescriptionGroups.Bounds.SphereRadius, HairGroupLOD, HairGroupsData.Strands.ClusterCullingBulkData);
+				FGroomBuilder::BuildClusterBulkData(StrandsData, HairDescriptionGroups.Bounds.SphereRadius, HairGroupLOD, HairGroupsData.Strands.ClusterBulkData);
 			}
 		}
 	}
diff --git a/Engine/Plugins/Runtime/HairStrands/Source/HairStrandsCore/Private/GroomResources.cpp b/Engine/Plugins/Runtime/HairStrands/Source/HairStrandsCore/Private/GroomResources.cpp
--- a/Engine/Plugins/Runtime/HairStrands/Source/HairStrandsCore/Private/GroomResources.cpp
+++ b/Engine/Plugins/Runtime/HairStrands/Source/HairStrandsCore/Private/GroomResources.cpp
@@ -1245,14 +1245,14 @@ void FHairStrandsDeformedResource::InternalRelease()
 
 /////////////////////////////////////////////////////////////////////////////////////////
 // Cluster culling resources
-FHairStrandsClusterCullingResource::FHairStrandsClusterCullingResource(FHairStrandsClusterCullingBulkData& InBulkData, const FHairResourceName& InResourceName, const FName& InOwnerName):
+FHairStrandsClusterResource::FHairStrandsClusterResource(FHairStrandsClusterBulkData& InBulkData, const FHairResourceName& InResourceName, const FName& InOwnerName):
 	FHairCommonResource(EHairStrandsAllocationType::Deferred, InResourceName, InOwnerName),
 	BulkData(InBulkData)
 {
 	MaxAvailableCurveCount = 0;
 }
 
-bool FHairStrandsClusterCullingResource::InternalGetOrRequestData(uint32 InRequestedCurveCount, uint32 InRequestedPointCount, int32 InLODIndex)
+bool FHairStrandsClusterResource::InternalGetOrRequestData(uint32 InRequestedCurveCount, uint32 InRequestedPointCount, int32 InLODIndex)
 {
 	if (StreamingRequest.IsNone())
 	{
@@ -1261,18 +1261,18 @@ bool FHairStrandsClusterCullingResource::InternalGetOrRequestData(uint32 InReque
 	return StreamingRequest.IsCompleted();
 }
 
-void FHairStrandsClusterCullingResource::InternalAllocate(FRDGBuilder& GraphBuilder)
+void FHairStrandsClusterResource::InternalAllocate(FRDGBuilder& GraphBuilder)
 {
 	check(StreamingRequest.IsCompleted());
 
 	const uint32 EntryCount = FMath::DivideAndRoundUp(BulkData.Header.PointCount, HAIR_POINT_LOD_COUNT_PER_UINT);
 
-	InternalCreateStructuredBufferRDG_FromHairBulkData<FHairClusterInfoFormat>(GraphBuilder, BulkData.Data.PackedClusterInfos, BulkData.Header.ClusterCount, ClusterInfoBuffer, ToHairResourceDebugName(TEXT("Hair.StrandsClusterCulling_ClusterInfoBuffer"), ResourceName), OwnerName, EHairResourceUsageType::Static);
-	InternalCreateVertexBufferRDG_FromHairBulkData<FHairClusterIndexFormat>(GraphBuilder, BulkData.Data.CurveToClusterIds, BulkData.Header.CurveCount, CurveToClusterIdBuffer, ToHairResourceDebugName(TEXT("Hair.StrandsClusterCulling_CurveToClusterIds"), ResourceName), OwnerName, EHairResourceUsageType::Static);	
-	InternalCreateVertexBufferRDG_FromHairBulkData<FHairClusterIndexFormat>(GraphBuilder, BulkData.Data.PointLODs, EntryCount, PointLODBuffer, ToHairResourceDebugName(TEXT("Hair.StrandsClusterCulling_PointLOD"), ResourceName), OwnerName, EHairResourceUsageType::Static);
+	InternalCreateStructuredBufferRDG_FromHairBulkData<FHairClusterInfoFormat>(GraphBuilder, BulkData.Data.PackedClusterInfos, BulkData.Header.ClusterCount, ClusterInfoBuffer, ToHairResourceDebugName(TEXT("Hair.StrandsCluster_ClusterInfoBuffer"), ResourceName), OwnerName, EHairResourceUsageType::Static);
+	InternalCreateVertexBufferRDG_FromHairBulkData<FHairClusterIndexFormat>(GraphBuilder, BulkData.Data.CurveToClusterIds, BulkData.Header.CurveCount, CurveToClusterIdBuffer, ToHairResourceDebugName(TEXT("Hair.StrandsCluster_CurveToClusterIds"), ResourceName), OwnerName, EHairResourceUsageType::Static);	
+	InternalCreateVertexBufferRDG_FromHairBulkData<FHairClusterIndexFormat>(GraphBuilder, BulkData.Data.PointLODs, EntryCount, PointLODBuffer, ToHairResourceDebugName(TEXT("Hair.StrandsCluster_PointLOD"), ResourceName), OwnerName, EHairResourceUsageType::Static);
 }
 
-void FHairStrandsClusterCullingResource::InternalRelease()
+void FHairStrandsClusterResource::InternalRelease()
 {
 	ClusterInfoBuffer.Release();
 	CurveToClusterIdBuffer.Release();
@@ -1280,7 +1280,7 @@ void FHairStrandsClusterCullingResource::InternalRelease()
 	MaxAvailableCurveCount = 0;
 }
 
-void FHairStrandsClusterCullingResource::InternalResetLoadedSize()
+void FHairStrandsClusterResource::InternalResetLoadedSize()
 {
 	BulkData.ResetLoadedSize();
 }
diff --git a/Engine/Plugins/Runtime/HairStrands/Source/HairStrandsCore/Private/HairStrandsClusterCulling.cpp b/Engine/Plugins/Runtime/HairStrands/Source/HairStrandsCore/Private/HairStrandsClusterCulling.cpp
--- a/Engine/Plugins/Runtime/HairStrands/Source/HairStrandsCore/Private/HairStrandsClusterCulling.cpp
+++ b/Engine/Plugins/Runtime/HairStrands/Source/HairStrandsCore/Private/HairStrandsClusterCulling.cpp
@@ -279,11 +279,11 @@ void AddInstanceToClusterData(FHairGroupInstance* In, FHairStrandClusterData& Ou
 	HairGroupCluster.ClusterAABBBuffer = &In->HairGroupPublicData->GetClusterAABBBuffer();
 
 	HairGroupCluster.CurveBuffer = &In->Strands.RestResource->CurveBuffer;
-	HairGroupCluster.PointLODBuffer = &In->Strands.ClusterCullingResource->PointLODBuffer;
+	HairGroupCluster.PointLODBuffer = &In->Strands.ClusterResource->PointLODBuffer;
 
-	HairGroupCluster.ClusterInfoParameters = In->Strands.ClusterCullingResource->BulkData.Header.ClusterInfoParameters;
-	HairGroupCluster.ClusterInfoBuffer = &In->Strands.ClusterCullingResource->ClusterInfoBuffer;
-	HairGroupCluster.CurveToClusterIdBuffer = &In->Strands.ClusterCullingResource->CurveToClusterIdBuffer;
+	HairGroupCluster.ClusterInfoParameters = In->Strands.ClusterResource->BulkData.Header.ClusterInfoParameters;
+	HairGroupCluster.ClusterInfoBuffer = &In->Strands.ClusterResource->ClusterInfoBuffer;
+	HairGroupCluster.CurveToClusterIdBuffer = &In->Strands.ClusterResource->CurveToClusterIdBuffer;
 
 	HairGroupCluster.HairGroupPublicPtr = In->HairGroupPublicData;
 	HairGroupCluster.LODBias  = In->HairGroupPublicData->GetLODBias();
diff --git a/Engine/Plugins/Runtime/HairStrands/Source/HairStrandsCore/Private/HairStrandsClusterCulling.h b/Engine/Plugins/Runtime/HairStrands/Source/HairStrandsCore/Private/HairStrandsClusterCulling.h
--- a/Engine/Plugins/Runtime/HairStrands/Source/HairStrandsCore/Private/HairStrandsClusterCulling.h
+++ b/Engine/Plugins/Runtime/HairStrands/Source/HairStrandsCore/Private/HairStrandsClusterCulling.h
@@ -26,7 +26,7 @@ struct FHairStrandClusterData
 		FRDGExternalBuffer* CurveBuffer = nullptr;
 		FRDGExternalBuffer* PointLODBuffer = nullptr;
 
-		// See FHairStrandsClusterCullingResource fro details about those buffers.
+		// See FHairStrandsClusterResource fro details about those buffers.
 		FRDGExternalBuffer* GroupAABBBuffer = nullptr;
 		FRDGExternalBuffer* ClusterAABBBuffer = nullptr;
 		FRDGExternalBuffer* ClusterInfoBuffer = nullptr; // SRV
diff --git a/Engine/Plugins/Runtime/HairStrands/Source/HairStrandsCore/Private/HairStrandsDatas.cpp b/Engine/Plugins/Runtime/HairStrands/Source/HairStrandsCore/Private/HairStrandsDatas.cpp
--- a/Engine/Plugins/Runtime/HairStrands/Source/HairStrandsCore/Private/HairStrandsDatas.cpp
+++ b/Engine/Plugins/Runtime/HairStrands/Source/HairStrandsCore/Private/HairStrandsDatas.cpp
@@ -111,9 +111,9 @@ void FHairStrandsInterpolationDatas::SetNum(const uint32 NumCurves)
 	PointsSimCurvesIndex.SetNum(NumCurves);
 }
 
-void FHairStrandsClusterCullingData::Reset()
+void FHairStrandsClusterData::Reset()
 {
-	*this = FHairStrandsClusterCullingData();
+	*this = FHairStrandsClusterData();
 }
 
 /////////////////////////////////////////////////////////////////////////////////////////
@@ -690,7 +690,7 @@ void FHairStrandsInterpolationBulkData::ResetLoadedSize()
 /////////////////////////////////////////////////////////////////////////////////////////
 // Cluster culling bulk data
 
-void FHairStrandsClusterCullingBulkData::Reset()
+void FHairStrandsClusterBulkData::Reset()
 {
 	Header.ClusterCount = 0;
 	Header.PointCount = 0;
@@ -708,14 +708,14 @@ void FHairStrandsClusterCullingBulkData::Reset()
 	Data.PointLODs			= FHairBulkContainer();
 }
 
-void FHairStrandsClusterCullingBulkData::ResetLoadedSize()
+void FHairStrandsClusterBulkData::ResetLoadedSize()
 {
 	Data.CurveToClusterIds.LoadedSize 	= 0;
 	Data.PackedClusterInfos.LoadedSize	= 0;
 	Data.PointLODs.LoadedSize			= 0;
 }
 
-void FHairStrandsClusterCullingBulkData::SerializeHeader(FArchive& Ar, UObject* Owner)
+void FHairStrandsClusterBulkData::SerializeHeader(FArchive& Ar, UObject* Owner)
 {
 	Ar << Header.ClusterCount;
 	Ar << Header.PointCount;
@@ -741,14 +741,14 @@ void FHairStrandsClusterCullingBulkData::SerializeHeader(FArchive& Ar, UObject*
 	Ar << Header.Strides.PointLODStride;
 }
 
-uint32 FHairStrandsClusterCullingBulkData::GetResourceCount() const
+uint32 FHairStrandsClusterBulkData::GetResourceCount() const
 {
 	return 5;
 }
 
 bool ValidateHairBulkData();
 
-void FHairStrandsClusterCullingBulkData::GetResources(FHairStrandsBulkCommon::FQuery & Out)
+void FHairStrandsClusterBulkData::GetResources(FHairStrandsBulkCommon::FQuery & Out)
 {
 	const bool bHasClusterData = Header.ClusterCount > 0;
 	if (bHasClusterData)
@@ -770,7 +770,7 @@ void FHairStrandsClusterCullingBulkData::GetResources(FHairStrandsBulkCommon::FQ
 	}
 }
 
-uint32 FHairStrandsClusterCullingBulkData::GetCurveCount(float InLODIndex) const
+uint32 FHairStrandsClusterBulkData::GetCurveCount(float InLODIndex) const
 {
 	const TArray<FHairLODInfo>& LODInfos = Header.LODInfos;
 
diff --git a/Engine/Plugins/Runtime/HairStrands/Source/HairStrandsCore/Private/HairStrandsInterpolation.cpp b/Engine/Plugins/Runtime/HairStrands/Source/HairStrandsCore/Private/HairStrandsInterpolation.cpp
--- a/Engine/Plugins/Runtime/HairStrands/Source/HairStrandsCore/Private/HairStrandsInterpolation.cpp
+++ b/Engine/Plugins/Runtime/HairStrands/Source/HairStrandsCore/Private/HairStrandsInterpolation.cpp
@@ -2043,7 +2043,7 @@ void ComputeHairStrandsInterpolation(
 					*Instance->Strands.Data,
 					Register(GraphBuilder, Instance->Strands.RestResource->CurveAttributeBuffer, ERDGImportedBufferFlags::CreateSRV).Buffer,
 					RegisterAsSRV(GraphBuilder, Instance->Strands.RestResource->CurveBuffer),
-					RegisterAsSRV(GraphBuilder, Instance->Strands.ClusterCullingResource->CurveToClusterIdBuffer),
+					RegisterAsSRV(GraphBuilder, Instance->Strands.ClusterResource->CurveToClusterIdBuffer),
 					bValidGuide ? RegisterAsSRV(GraphBuilder, Instance->Strands.InterpolationResource->InterpolationBuffer) : nullptr,
 					OutRenCurveAttributeBuffer);
 			}
diff --git a/Engine/Plugins/Runtime/HairStrands/Source/HairStrandsCore/Public/GroomAsset.h b/Engine/Plugins/Runtime/HairStrands/Source/HairStrandsCore/Public/GroomAsset.h
--- a/Engine/Plugins/Runtime/HairStrands/Source/HairStrandsCore/Public/GroomAsset.h
+++ b/Engine/Plugins/Runtime/HairStrands/Source/HairStrandsCore/Public/GroomAsset.h
@@ -171,7 +171,7 @@ struct HAIRSTRANDSCORE_API FHairGroupPlatformData
 			uint32 Total = 0;
 			if (RestResource) Total += RestResource->GetResourcesSize();
 			if (InterpolationResource) Total += InterpolationResource->GetResourcesSize();
-			if (ClusterCullingResource) Total += ClusterCullingResource->GetResourcesSize();
+			if (ClusterResource) Total += ClusterResource->GetResourcesSize();
 			#if RHI_RAYTRACING
 			if (RaytracingResource) Total += RaytracingResource->GetResourcesSize();
 			#endif
@@ -180,8 +180,8 @@ struct HAIRSTRANDSCORE_API FHairGroupPlatformData
 
 		uint32 GetDataSize() const;
 
-		FHairStrandsClusterCullingBulkData	ClusterCullingBulkData;
-		FHairStrandsClusterCullingResource* ClusterCullingResource = nullptr;
+		FHairStrandsClusterBulkData	ClusterBulkData;
+		FHairStrandsClusterResource* ClusterResource = nullptr;
 
 		#if RHI_RAYTRACING
 		FHairStrandsRaytracingResource*		RaytracingResource = nullptr;
diff --git a/Engine/Plugins/Runtime/HairStrands/Source/HairStrandsCore/Public/GroomBuilder.h b/Engine/Plugins/Runtime/HairStrands/Source/HairStrandsCore/Public/GroomBuilder.h
--- a/Engine/Plugins/Runtime/HairStrands/Source/HairStrandsCore/Public/GroomBuilder.h
+++ b/Engine/Plugins/Runtime/HairStrands/Source/HairStrandsCore/Public/GroomBuilder.h
@@ -9,8 +9,8 @@
 #include "Math/IntVector.h"
 
 struct FHairStrandsDatas;
-struct FHairStrandsClusterCullingData;
-struct FHairStrandsClusterCullingBulkData;
+struct FHairStrandsClusterData;
+struct FHairStrandsClusterBulkData;
 struct FHairGroupsLOD;
 struct FHairGroupData;
 struct FHairDescriptionGroups;
@@ -101,7 +101,7 @@ struct HAIRSTRANDSCORE_API FGroomBuilder
 		const FHairStrandsDatas& InRenData,
 		const float InGroomAssetRadius,
 		const FHairGroupsLOD& InSettings,
-		FHairStrandsClusterCullingBulkData& OutClusterCullingData);
+		FHairStrandsClusterBulkData& OutClusterData);
 
 	// Optional: Voxelize hair group index
 	static void VoxelizeGroupIndex(
diff --git a/Engine/Plugins/Runtime/HairStrands/Source/HairStrandsCore/Public/GroomInstance.h b/Engine/Plugins/Runtime/HairStrands/Source/HairStrandsCore/Public/GroomInstance.h
--- a/Engine/Plugins/Runtime/HairStrands/Source/HairStrandsCore/Public/GroomInstance.h
+++ b/Engine/Plugins/Runtime/HairStrands/Source/HairStrandsCore/Public/GroomInstance.h
@@ -123,7 +123,7 @@ struct HAIRSTRANDSCORE_API FHairGroupInstance : public FHairStrandsInstance
 	struct FStrands : FStrandsBaseWithInterpolation
 	{
 		// Resources - Strands cluster data for culling/voxelization purpose
-		FHairStrandsClusterCullingResource* ClusterCullingResource = nullptr;
+		FHairStrandsClusterResource* ClusterResource = nullptr;
 		FHairStrandsCullingResource* CullingResource = nullptr;
 
 		// Resources - Raytracing data when enabling (expensive) raytracing method
diff --git a/Engine/Plugins/Runtime/HairStrands/Source/HairStrandsCore/Public/GroomResources.h b/Engine/Plugins/Runtime/HairStrands/Source/HairStrandsCore/Public/GroomResources.h
--- a/Engine/Plugins/Runtime/HairStrands/Source/HairStrandsCore/Public/GroomResources.h
+++ b/Engine/Plugins/Runtime/HairStrands/Source/HairStrandsCore/Public/GroomResources.h
@@ -316,9 +316,9 @@ struct FHairStrandsInterpolationResource : public FHairCommonResource
 };
 
 
-struct FHairStrandsClusterCullingResource : public FHairCommonResource
+struct FHairStrandsClusterResource : public FHairCommonResource
 {
-	FHairStrandsClusterCullingResource(FHairStrandsClusterCullingBulkData& Data, const FHairResourceName& ResourceName, const FName& OwnerName);
+	FHairStrandsClusterResource(FHairStrandsClusterBulkData& Data, const FHairResourceName& ResourceName, const FName& OwnerName);
 
 	/* Init/Release buffers */
 	virtual void InternalAllocate(FRDGBuilder& GraphBuilder) override;
@@ -349,7 +349,7 @@ struct FHairStrandsClusterCullingResource : public FHairCommonResource
 	/* Min. LOD a which a Point becomes visible */
 	FRDGExternalBuffer PointLODBuffer;
 
-	FHairStrandsClusterCullingBulkData& BulkData;
+	FHairStrandsClusterBulkData& BulkData;
 };
 
 struct FHairStrandsCullingResource : public FHairCommonResource
diff --git a/Engine/Plugins/Runtime/HairStrands/Source/HairStrandsCore/Public/HairStrandsDatas.h b/Engine/Plugins/Runtime/HairStrands/Source/HairStrandsCore/Public/HairStrandsDatas.h
--- a/Engine/Plugins/Runtime/HairStrands/Source/HairStrandsCore/Public/HairStrandsDatas.h
+++ b/Engine/Plugins/Runtime/HairStrands/Source/HairStrandsCore/Public/HairStrandsDatas.h
@@ -694,7 +694,7 @@ struct HAIRSTRANDSCORE_API FHairStrandsInterpolationBulkData : FHairStrandsBulkC
 ///////////////////////////////////////////////////////////////////////////////////////////////////
 // Cluster Data/Bulk
 
-struct HAIRSTRANDSCORE_API FHairStrandsClusterCullingData
+struct HAIRSTRANDSCORE_API FHairStrandsClusterData
 {
 	void Reset();
 	bool IsValid() const { return ClusterCount > 0 && PointCount > 0; }
@@ -710,7 +710,7 @@ struct HAIRSTRANDSCORE_API FHairStrandsClusterCullingData
 	uint32 CurveCount = 0;			// Number of curves
 };
 
-struct HAIRSTRANDSCORE_API FHairStrandsClusterCullingBulkData : FHairStrandsBulkCommon
+struct HAIRSTRANDSCORE_API FHairStrandsClusterBulkData : FHairStrandsBulkCommon
 {
 	void Reset();
 	virtual void ResetLoadedSize() override;
diff --git a/Engine/Plugins/Runtime/HairStrands/Source/HairStrandsEditor/Private/GroomAssetDetails.cpp b/Engine/Plugins/Runtime/HairStrands/Source/HairStrandsEditor/Private/GroomAssetDetails.cpp
--- a/Engine/Plugins/Runtime/HairStrands/Source/HairStrandsEditor/Private/GroomAssetDetails.cpp
+++ b/Engine/Plugins/Runtime/HairStrands/Source/HairStrandsEditor/Private/GroomAssetDetails.cpp
@@ -1328,10 +1328,10 @@ void FGroomRenderingDetails::OnGenerateElementForLODs(TSharedRef<IPropertyHandle
 	];
 
 	// LOD Stats
-	const FHairStrandsClusterCullingBulkData& ClusterCullingBulkData = GroomAsset->GetHairGroupsPlatformData()[GroupIndex].Strands.ClusterCullingBulkData;
-	if (ClusterCullingBulkData.IsValid() && LODIndex < ClusterCullingBulkData.Header.LODInfos.Num())
+	const FHairStrandsClusterBulkData& ClusterBulkData = GroomAsset->GetHairGroupsPlatformData()[GroupIndex].Strands.ClusterBulkData;
+	if (ClusterBulkData.IsValid() && LODIndex < ClusterBulkData.Header.LODInfos.Num())
 	{
-		const FHairLODInfo& LODInfo = ClusterCullingBulkData.Header.LODInfos[LODIndex];
+		const FHairLODInfo& LODInfo = ClusterBulkData.Header.LODInfos[LODIndex];
 
 		ChildrenBuilder.AddCustomRow(LOCTEXT("HairStrandsLODInfo_Array", "HairStrandsLODInfo"))
 		.ValueContent()
diff --git a/Engine/Source/Runtime/Renderer/Public/HairStrandsInterface.h b/Engine/Source/Runtime/Renderer/Public/HairStrandsInterface.h
--- a/Engine/Source/Runtime/Renderer/Public/HairStrandsInterface.h
+++ b/Engine/Source/Runtime/Renderer/Public/HairStrandsInterface.h
@@ -162,7 +162,7 @@ RENDERER_API FRDGBufferUAVRef   RegisterAsUAV(FRDGBuilder& GraphBuilder, const F
 ////////////////////////////////////////////////////////////////////////////////////////////////////
 // Misc/Helpers
 
-struct FHairStrandClusterCullingData;
+struct FHairStrandClusterData;
 struct IPooledRenderTarget;
 struct FRWBuffer;
 class  FRDGPooledBuffer;
