Fix invalid motion vector on groom when streaming data in/out.

diff --git a/Engine/Plugins/Runtime/HairStrands/Source/HairStrandsCore/Private/GroomResources.cpp b/Engine/Plugins/Runtime/HairStrands/Source/HairStrandsCore/Private/GroomResources.cpp
--- a/Engine/Plugins/Runtime/HairStrands/Source/HairStrandsCore/Private/GroomResources.cpp
+++ b/Engine/Plugins/Runtime/HairStrands/Source/HairStrandsCore/Private/GroomResources.cpp
@@ -197,7 +197,7 @@ inline void InternalSetBulkDataFlags(FByteBulkData& In)
 ///////////////////////////////////////////////////////////////////////////////////////////////////
 // Regular loading from BulkData
 
-void InternalCreateBufferRDG_FromBulkData(FRDGBuilder& GraphBuilder, FByteBulkData& InBulkData, FRDGExternalBuffer& Out, const EPixelFormat OutFormat, const FRDGBufferDesc Desc, const TCHAR* DebugName, const FName& OwnerName)
+FRDGBufferRef InternalCreateBufferRDG_FromBulkData(FRDGBuilder& GraphBuilder, FByteBulkData& InBulkData, FRDGExternalBuffer& Out, const EPixelFormat OutFormat, const FRDGBufferDesc Desc, const TCHAR* DebugName, const FName& OwnerName)
 {
 	InternalSetBulkDataFlags(InBulkData);
 
@@ -206,7 +206,7 @@ void InternalCreateBufferRDG_FromBulkData(FRDGBuilder& GraphBuilder, FByteBulkDa
 	if (DataSizeInBytes == 0)
 	{
 		Out.Buffer = nullptr;
-		return;
+		return nullptr;
 	}
 
 	const uint8* Data = (const uint8*)InBulkData.Lock(LOCK_READ_ONLY);
@@ -227,26 +227,27 @@ void InternalCreateBufferRDG_FromBulkData(FRDGBuilder& GraphBuilder, FByteBulkDa
 		}
 	}
 	ConvertToExternalBufferWithViews(GraphBuilder, Buffer, Out, OutFormat);
+	return Buffer;
 }
 
 template<typename FormatType>
-void InternalCreateVertexBufferRDG_FromBulkData(FRDGBuilder& GraphBuilder, FByteBulkData& InBulkData, uint32 InDataCount, FRDGExternalBuffer& Out, const TCHAR* DebugName, const FName& OwnerName, EHairResourceUsageType UsageType)
+FRDGBufferRef InternalCreateVertexBufferRDG_FromBulkData(FRDGBuilder& GraphBuilder, FByteBulkData& InBulkData, uint32 InDataCount, FRDGExternalBuffer& Out, const TCHAR* DebugName, const FName& OwnerName, EHairResourceUsageType UsageType)
 {
 	const FRDGBufferDesc Desc = ApplyUsage(FRDGBufferDesc::CreateBufferDesc(FormatType::SizeInByte, InDataCount), UsageType);
-	InternalCreateBufferRDG_FromBulkData(GraphBuilder, InBulkData, Out, FormatType::Format, Desc, DebugName, OwnerName);
+	return InternalCreateBufferRDG_FromBulkData(GraphBuilder, InBulkData, Out, FormatType::Format, Desc, DebugName, OwnerName);
 }
 
 template<typename FormatType>
-void InternalCreateStructuredBufferRDG_FromBulkData(FRDGBuilder& GraphBuilder, FByteBulkData& InBulkData, uint32 InDataCount, FRDGExternalBuffer& Out, const TCHAR* DebugName, const FName& OwnerName, EHairResourceUsageType UsageType)
+FRDGBufferRef InternalCreateStructuredBufferRDG_FromBulkData(FRDGBuilder& GraphBuilder, FByteBulkData& InBulkData, uint32 InDataCount, FRDGExternalBuffer& Out, const TCHAR* DebugName, const FName& OwnerName, EHairResourceUsageType UsageType)
 {
 	const FRDGBufferDesc Desc = ApplyUsage(FRDGBufferDesc::CreateStructuredDesc(FormatType::SizeInByte, InDataCount), UsageType);
-	InternalCreateBufferRDG_FromBulkData(GraphBuilder, InBulkData, Out, PF_Unknown, Desc, DebugName, OwnerName);
+	return InternalCreateBufferRDG_FromBulkData(GraphBuilder, InBulkData, Out, PF_Unknown, Desc, DebugName, OwnerName);
 }
 
-void InternalCreateByteAddressBufferRDG_FromBulkData(FRDGBuilder& GraphBuilder, FByteBulkData& InBulkData, FRDGExternalBuffer& Out, const TCHAR* DebugName, const FName& OwnerName, EHairResourceUsageType UsageType)
+FRDGBufferRef InternalCreateByteAddressBufferRDG_FromBulkData(FRDGBuilder& GraphBuilder, FByteBulkData& InBulkData, FRDGExternalBuffer& Out, const TCHAR* DebugName, const FName& OwnerName, EHairResourceUsageType UsageType)
 {
 	const FRDGBufferDesc Desc = ApplyUsage(FRDGBufferDesc::CreateByteAddressDesc(InBulkData.GetBulkDataSize()), UsageType);
-	InternalCreateBufferRDG_FromBulkData(GraphBuilder, InBulkData, Out, PF_Unknown, Desc, DebugName, OwnerName);
+	return InternalCreateBufferRDG_FromBulkData(GraphBuilder, InBulkData, Out, PF_Unknown, Desc, DebugName, OwnerName);
 }
 
 
@@ -320,65 +321,68 @@ static FRDGBufferRef InternalCreateBufferRDG_FromHairBulkData(FRDGBuilder& Graph
 }
 
 template<typename FormatType>
-void InternalCreateVertexBufferRDG_FromHairBulkData(FRDGBuilder& GraphBuilder, FHairBulkContainer& InChunk, uint32 InDataCount, FRDGExternalBuffer& Out, const TCHAR* DebugName, const FName& OwnerName, EHairResourceUsageType UsageType)
+FRDGBufferRef InternalCreateVertexBufferRDG_FromHairBulkData(FRDGBuilder& GraphBuilder, FHairBulkContainer& InChunk, uint32 InDataCount, FRDGExternalBuffer& Out, const TCHAR* DebugName, const FName& OwnerName, EHairResourceUsageType UsageType)
 {
 	// Fallback for non-streamble resources (e.g. guides)
 	if (InChunk.ChunkRequest == nullptr)
 	{
-		InternalCreateVertexBufferRDG_FromBulkData<FormatType>(GraphBuilder, InChunk.Data, InDataCount, Out, DebugName, OwnerName, UsageType);
-		return;
+		return InternalCreateVertexBufferRDG_FromBulkData<FormatType>(GraphBuilder, InChunk.Data, InDataCount, Out, DebugName, OwnerName, UsageType);
 	}
 
 	const FRDGBufferRef In = Out.Buffer ? Register(GraphBuilder, Out, ERDGImportedBufferFlags::None).Buffer : nullptr;
 	const FRDGBufferDesc BufferDesc = ApplyUsage(FRDGBufferDesc::CreateBufferDesc(FormatType::SizeInByte, FMath::DivideAndRoundUp(InChunk.ChunkRequest->TotalSize, FormatType::SizeInByte)), UsageType);
 	const FRDGBufferDesc UploadDesc = ApplyUsage(FRDGBufferDesc::CreateBufferDesc(FormatType::SizeInByte, FMath::DivideAndRoundUp(InChunk.ChunkRequest->Size, FormatType::SizeInByte)), UsageType);
-	if (BufferDesc.GetSize() == 0) 	{ Out.Buffer = nullptr; return; }
-	if (UploadDesc.GetSize()==0) 	{ return; }
+	if (BufferDesc.GetSize() == 0) 	{ Out.Buffer = nullptr; return nullptr; }
+	if (UploadDesc.GetSize()==0) 	{ return nullptr; }
 	if (FRDGBufferRef Buffer = InternalCreateBufferRDG_FromHairBulkData(GraphBuilder, InChunk, In, BufferDesc, UploadDesc, DebugName, OwnerName))
 	{
 		ConvertToExternalBufferWithViews(GraphBuilder, Buffer, Out, FormatType::Format);
+		return Buffer;
 	}
+	return nullptr;
 }
 
 template<typename FormatType>
-void InternalCreateStructuredBufferRDG_FromHairBulkData(FRDGBuilder& GraphBuilder, FHairBulkContainer& InChunk, uint32 InDataCount, FRDGExternalBuffer& Out, const TCHAR* DebugName, const FName& OwnerName, EHairResourceUsageType UsageType)
+FRDGBufferRef InternalCreateStructuredBufferRDG_FromHairBulkData(FRDGBuilder& GraphBuilder, FHairBulkContainer& InChunk, uint32 InDataCount, FRDGExternalBuffer& Out, const TCHAR* DebugName, const FName& OwnerName, EHairResourceUsageType UsageType)
 {
 	// Fallback for non-streamble resources (e.g. guides)
 	if (InChunk.ChunkRequest == nullptr)
 	{
-		InternalCreateStructuredBufferRDG_FromBulkData<FormatType>(GraphBuilder, InChunk.Data, InDataCount, Out, DebugName, OwnerName, UsageType);
-		return;
+		return InternalCreateStructuredBufferRDG_FromBulkData<FormatType>(GraphBuilder, InChunk.Data, InDataCount, Out, DebugName, OwnerName, UsageType);
 	}
 
 	const FRDGBufferRef In = Out.Buffer ? Register(GraphBuilder, Out, ERDGImportedBufferFlags::None).Buffer : nullptr;
 	const FRDGBufferDesc BufferDesc = ApplyUsage(FRDGBufferDesc::CreateStructuredDesc(FormatType::SizeInByte, FMath::DivideAndRoundUp(InChunk.ChunkRequest->TotalSize, FormatType::SizeInByte)), UsageType);
 	const FRDGBufferDesc UploadDesc = ApplyUsage(FRDGBufferDesc::CreateStructuredDesc(FormatType::SizeInByte, FMath::DivideAndRoundUp(InChunk.ChunkRequest->Size, FormatType::SizeInByte)), UsageType);
-	if (BufferDesc.GetSize() == 0) 	{ Out.Buffer = nullptr; return; }
-	if (UploadDesc.GetSize()==0) 	{ return; }
+	if (BufferDesc.GetSize() == 0) 	{ Out.Buffer = nullptr; return nullptr; }
+	if (UploadDesc.GetSize()==0) 	{ return nullptr; }
 	if (FRDGBufferRef Buffer = InternalCreateBufferRDG_FromHairBulkData(GraphBuilder, InChunk, In, BufferDesc, UploadDesc, DebugName, OwnerName))
 	{
 		ConvertToExternalBufferWithViews(GraphBuilder, Buffer, Out);
+		return Buffer;
 	}
+	return nullptr;
 }
 
-void InternalCreateByteAddressBufferRDG_FromHairBulkData(FRDGBuilder& GraphBuilder, FHairBulkContainer& InChunk, FRDGExternalBuffer& Out, const TCHAR* DebugName, const FName& OwnerName, EHairResourceUsageType UsageType)
+FRDGBufferRef InternalCreateByteAddressBufferRDG_FromHairBulkData(FRDGBuilder& GraphBuilder, FHairBulkContainer& InChunk, FRDGExternalBuffer& Out, const TCHAR* DebugName, const FName& OwnerName, EHairResourceUsageType UsageType)
 {
 	// Fallback for non-streamble resources (e.g. guides)
 	if (InChunk.ChunkRequest == nullptr)
 	{
-		InternalCreateByteAddressBufferRDG_FromBulkData(GraphBuilder, InChunk.Data, Out, DebugName, OwnerName, UsageType);
-		return;
+		return InternalCreateByteAddressBufferRDG_FromBulkData(GraphBuilder, InChunk.Data, Out, DebugName, OwnerName, UsageType);
 	}
 
 	const FRDGBufferRef In = Out.Buffer ? Register(GraphBuilder, Out, ERDGImportedBufferFlags::None).Buffer : nullptr;
 	const FRDGBufferDesc BufferDesc = ApplyUsage(FRDGBufferDesc::CreateByteAddressDesc(InChunk.ChunkRequest->TotalSize), UsageType);
 	const FRDGBufferDesc UploadDesc = ApplyUsage(FRDGBufferDesc::CreateByteAddressDesc(InChunk.ChunkRequest->Size), UsageType);
-	if (BufferDesc.GetSize() == 0) 	{ Out.Buffer = nullptr; return; }
-	if (UploadDesc.GetSize()==0) 	{ return; }
+	if (BufferDesc.GetSize() == 0) 	{ Out.Buffer = nullptr; return nullptr; }
+	if (UploadDesc.GetSize()==0) 	{ return nullptr; }
 	if (FRDGBufferRef Buffer = InternalCreateBufferRDG_FromHairBulkData(GraphBuilder, InChunk, In, BufferDesc, UploadDesc, DebugName, OwnerName))
 	{
 		ConvertToExternalBufferWithViews(GraphBuilder, Buffer, Out);
+		return Buffer;
 	}
+	return nullptr;
 }
 
 ///////////////////////////////////////////////////////////////////////////////////////////////////
@@ -414,19 +418,17 @@ static FRDGBufferRef InternalCreateVertexBuffer(
 }
 
 template<typename FormatType>
-void InternalCreateVertexBufferRDG(FRDGBuilder& GraphBuilder, const typename FormatType::Type* InData, uint32 InDataCount, FRDGExternalBuffer& Out, const TCHAR* DebugName, const FName& OwnerName, EHairResourceUsageType UsageType, ERDGInitialDataFlags InitialDataFlags)
+FRDGBufferRef InternalCreateVertexBufferRDG(FRDGBuilder& GraphBuilder, const typename FormatType::Type* InData, uint32 InDataCount, FRDGExternalBuffer& Out, const TCHAR* DebugName, const FName& OwnerName, EHairResourceUsageType UsageType, ERDGInitialDataFlags InitialDataFlags)
 {
-	FRDGBufferRef Buffer = nullptr;
-
 	const uint32 DataSizeInBytes = FormatType::SizeInByte * InDataCount;
 	if (DataSizeInBytes == 0)
 	{
 		Out.Buffer = nullptr;
-		return;
+		return nullptr;
 	}
 
 	const FRDGBufferDesc Desc = ApplyUsage(FRDGBufferDesc::CreateBufferDesc(FormatType::SizeInByte, InDataCount), UsageType);
-	Buffer = InternalCreateVertexBuffer(
+	FRDGBufferRef Buffer = InternalCreateVertexBuffer(
 		GraphBuilder,
 		DebugName,
 		Desc,
@@ -436,23 +438,22 @@ void InternalCreateVertexBufferRDG(FRDGBuilder& GraphBuilder, const typename For
 		OwnerName);
 
 	ConvertToExternalBufferWithViews(GraphBuilder, Buffer, Out, FormatType::Format);
+	return Buffer;
 }
 
 template<typename FormatType>
-void InternalCreateVertexBufferRDG(FRDGBuilder& GraphBuilder, const TArray<typename FormatType::Type>& InData, FRDGExternalBuffer& Out, const TCHAR* DebugName, const FName& OwnerName, EHairResourceUsageType UsageType, ERDGInitialDataFlags InitialDataFlags=ERDGInitialDataFlags::NoCopy)
+FRDGBufferRef InternalCreateVertexBufferRDG(FRDGBuilder& GraphBuilder, const TArray<typename FormatType::Type>& InData, FRDGExternalBuffer& Out, const TCHAR* DebugName, const FName& OwnerName, EHairResourceUsageType UsageType, ERDGInitialDataFlags InitialDataFlags=ERDGInitialDataFlags::NoCopy)
 {
-	FRDGBufferRef Buffer = nullptr;
-
 	const uint32 DataCount = InData.Num();
 	const uint32 DataSizeInBytes = FormatType::SizeInByte * DataCount;
 	if (DataSizeInBytes == 0)
 	{
 		Out.Buffer = nullptr;
-		return;
+		return nullptr;
 	}
 
 	const FRDGBufferDesc Desc = ApplyUsage(FRDGBufferDesc::CreateBufferDesc(FormatType::SizeInByte, InData.Num()), UsageType);
-	Buffer = InternalCreateVertexBuffer(
+	FRDGBufferRef Buffer = InternalCreateVertexBuffer(
 		GraphBuilder,
 		DebugName,
 		Desc,
@@ -462,24 +463,23 @@ void InternalCreateVertexBufferRDG(FRDGBuilder& GraphBuilder, const TArray<typen
 		OwnerName);
 
 	ConvertToExternalBufferWithViews(GraphBuilder, Buffer, Out, FormatType::Format);
+	return Buffer;
 }
 
 template<typename DataType>
-void InternalCreateVertexBufferRDG(FRDGBuilder& GraphBuilder, const TArray<DataType>& InData, EPixelFormat Format, FRDGExternalBuffer& Out, const TCHAR* DebugName, const FName& OwnerName, EHairResourceUsageType UsageType, ERDGInitialDataFlags InitialDataFlags=ERDGInitialDataFlags::NoCopy)
+FRDGBufferRef InternalCreateVertexBufferRDG(FRDGBuilder& GraphBuilder, const TArray<DataType>& InData, EPixelFormat Format, FRDGExternalBuffer& Out, const TCHAR* DebugName, const FName& OwnerName, EHairResourceUsageType UsageType, ERDGInitialDataFlags InitialDataFlags=ERDGInitialDataFlags::NoCopy)
 {
-	FRDGBufferRef Buffer = nullptr;
-
 	const uint32 DataCount = InData.Num();
 	const uint32 DataSizeInBytes = sizeof(DataType) * DataCount;
 	if (DataSizeInBytes == 0)
 	{
 		Out.Buffer = nullptr;
-		return;
+		return nullptr;
 	}
 
 	const FRDGBufferDesc Desc = ApplyUsage(FRDGBufferDesc::CreateBufferDesc(sizeof(DataType), InData.Num()), UsageType);
 
-	Buffer = InternalCreateVertexBuffer(
+	FRDGBufferRef Buffer = InternalCreateVertexBuffer(
 		GraphBuilder,
 		DebugName,
 		Desc,
@@ -489,10 +489,11 @@ void InternalCreateVertexBufferRDG(FRDGBuilder& GraphBuilder, const TArray<DataT
 		OwnerName);
 
 	ConvertToExternalBufferWithViews(GraphBuilder, Buffer, Out, Format);
+	return Buffer;
 }
 
 template<typename FormatType>
-void InternalCreateVertexBufferRDG(FRDGBuilder& GraphBuilder, uint32 InVertexCount, FRDGExternalBuffer& Out, const TCHAR* DebugName, const FName& OwnerName, EHairResourceUsageType Usage)
+FRDGBufferRef InternalCreateVertexBufferRDG(FRDGBuilder& GraphBuilder, uint32 InVertexCount, FRDGExternalBuffer& Out, const TCHAR* DebugName, const FName& OwnerName, EHairResourceUsageType Usage)
 {
 	// Sanity check
 	check(Usage == EHairResourceUsageType::Dynamic);
@@ -502,7 +503,7 @@ void InternalCreateVertexBufferRDG(FRDGBuilder& GraphBuilder, uint32 InVertexCou
 	if (DataSizeInBytes == 0)
 	{
 		Out.Buffer = nullptr;
-		return;
+		return nullptr;
 	}
 
 	FRDGBufferDesc Desc = FRDGBufferDesc::CreateBufferDesc(FormatType::SizeInByte, InVertexCount);
@@ -519,10 +520,11 @@ void InternalCreateVertexBufferRDG(FRDGBuilder& GraphBuilder, uint32 InVertexCou
 	}
 	
 	ConvertToExternalBufferWithViews(GraphBuilder, Buffer, Out, FormatType::Format);
+	return Buffer;
 }
 
 template<typename FormatType>
-void InternalCreateStructuredBufferRDG(FRDGBuilder& GraphBuilder, uint32 DataCount, FRDGExternalBuffer& Out, const TCHAR* DebugName, const FName& OwnerName, EHairResourceUsageType Usage)
+FRDGBufferRef InternalCreateStructuredBufferRDG(FRDGBuilder& GraphBuilder, uint32 DataCount, FRDGExternalBuffer& Out, const TCHAR* DebugName, const FName& OwnerName, EHairResourceUsageType Usage)
 {
 	// Sanity check
 	check(Usage == EHairResourceUsageType::Dynamic);
@@ -531,7 +533,7 @@ void InternalCreateStructuredBufferRDG(FRDGBuilder& GraphBuilder, uint32 DataCou
 	if (DataSizeInBytes == 0)
 	{
 		Out.Buffer = nullptr;
-		return;
+		return nullptr;
 	}
 
 	FRDGBufferDesc Desc = ApplyUsage(FRDGBufferDesc::CreateStructuredDesc(FormatType::SizeInByte, DataCount), Usage);
@@ -539,10 +541,11 @@ void InternalCreateStructuredBufferRDG(FRDGBuilder& GraphBuilder, uint32 DataCou
 	Buffer->SetOwnerName(OwnerName);
 	AddClearUAVPass(GraphBuilder, GraphBuilder.CreateUAV(Buffer, FormatType::Format), 0u);
 	ConvertToExternalBufferWithViews(GraphBuilder, Buffer, Out, FormatType::Format);
+	return Buffer;
 }
 
 template<typename DataType>
-void InternalCreateByteAddressBufferRDG(FRDGBuilder& GraphBuilder, const TArray<DataType>& InData, EPixelFormat Format, FRDGExternalBuffer& Out, const TCHAR* DebugName, const FName& OwnerName, EHairResourceUsageType UsageType)
+FRDGBufferRef InternalCreateByteAddressBufferRDG(FRDGBuilder& GraphBuilder, const TArray<DataType>& InData, EPixelFormat Format, FRDGExternalBuffer& Out, const TCHAR* DebugName, const FName& OwnerName, EHairResourceUsageType UsageType)
 {
 	const uint32 DataCount = InData.Num();
 	const uint32 DataSizeInBytes = sizeof(DataType) * DataCount;
@@ -558,10 +561,12 @@ void InternalCreateByteAddressBufferRDG(FRDGBuilder& GraphBuilder, const TArray<
 		}
 		
 		ConvertToExternalBufferWithViews(GraphBuilder, Buffer, Out, Format);
+		return Buffer;
 	}
+	return nullptr;
 }
 
-void InternalCreateByteAddressBufferRDG(FRDGBuilder& GraphBuilder, uint64 DataSizeInBytes, FRDGExternalBuffer& Out, const TCHAR* DebugName, const FName& OwnerName, EHairResourceUsageType UsageType)
+FRDGBufferRef InternalCreateByteAddressBufferRDG(FRDGBuilder& GraphBuilder, uint64 DataSizeInBytes, FRDGExternalBuffer& Out, const TCHAR* DebugName, const FName& OwnerName, EHairResourceUsageType UsageType)
 {
 	Out.Buffer = nullptr;
 	if (DataSizeInBytes != 0)
@@ -570,7 +575,9 @@ void InternalCreateByteAddressBufferRDG(FRDGBuilder& GraphBuilder, uint64 DataSi
 		FRDGBufferRef Buffer = GraphBuilder.CreateBuffer(Desc, DebugName, ERDGBufferFlags::MultiFrame);
 		Buffer->SetOwnerName(OwnerName);
 		ConvertToExternalBufferWithViews(GraphBuilder, Buffer, Out, PF_Unknown);
+		return Buffer;
 	}
+	return nullptr;
 }
 /////////////////////////////////////////////////////////////////////////////////////////
 
@@ -1173,10 +1180,20 @@ void FHairStrandsDeformedResource::InternalAllocate(FRDGBuilder& GraphBuilder, u
 	// Need to update these to by dynamic based on the active number of point
 	const uint32 PointCount = FMath::Min(InPointCount, BulkData.GetNumPoints());
 
-	InternalCreateVertexBufferRDG<FHairStrandsPositionFormat>(GraphBuilder, PointCount, DeformedPositionBuffer[0], ToHairResourceDebugName(HAIRSTRANDS_RESOUCE_NAME(CurveType, Hair.StrandsDeformed_DeformedPositionBuffer0), ResourceName), OwnerName, EHairResourceUsageType::Dynamic);
-	InternalCreateVertexBufferRDG<FHairStrandsPositionFormat>(GraphBuilder, PointCount, DeformedPositionBuffer[1], ToHairResourceDebugName(HAIRSTRANDS_RESOUCE_NAME(CurveType, Hair.StrandsDeformed_DeformedPositionBuffer1), ResourceName), OwnerName, EHairResourceUsageType::Dynamic);
+	FRDGBufferRef Prev0 = Register(GraphBuilder, DeformedPositionBuffer[0], ERDGImportedBufferFlags::None).Buffer;
+	FRDGBufferRef Prev1 = Register(GraphBuilder, DeformedPositionBuffer[1], ERDGImportedBufferFlags::None).Buffer;
+
+	FRDGBufferRef Deformed0 = InternalCreateVertexBufferRDG<FHairStrandsPositionFormat>(GraphBuilder, PointCount, DeformedPositionBuffer[0], ToHairResourceDebugName(HAIRSTRANDS_RESOUCE_NAME(CurveType, Hair.StrandsDeformed_DeformedPositionBuffer0), ResourceName), OwnerName, EHairResourceUsageType::Dynamic);
+	FRDGBufferRef Deformed1 = InternalCreateVertexBufferRDG<FHairStrandsPositionFormat>(GraphBuilder, PointCount, DeformedPositionBuffer[1], ToHairResourceDebugName(HAIRSTRANDS_RESOUCE_NAME(CurveType, Hair.StrandsDeformed_DeformedPositionBuffer1), ResourceName), OwnerName, EHairResourceUsageType::Dynamic);
 	InternalCreateVertexBufferRDG<FHairStrandsTangentFormat>(GraphBuilder,  PointCount * FHairStrandsTangentFormat::ComponentCount, TangentBuffer, ToHairResourceDebugName(HAIRSTRANDS_RESOUCE_NAME(CurveType, Hair.StrandsDeformed_TangentBuffer), ResourceName), OwnerName, EHairResourceUsageType::Dynamic);
 
+	if (Prev0 != nullptr)
+	{
+		const uint32 SizeInBytes = FMath::Min(Prev0->Desc.GetSize(), Deformed0->Desc.GetSize());
+		AddCopyBufferPass(GraphBuilder, Deformed0, 0, Prev0, 0, SizeInBytes);
+		AddCopyBufferPass(GraphBuilder, Deformed1, 0, Prev1, 0, SizeInBytes);
+	}
+	 
 	if (DeformedOffsetBuffer[0].Buffer == nullptr)
 	{
 		TArray<FVector4f> DefaultOffsets;
