* Fix incorrect culling resource (re)allocation.
* Fix crash when an unloading request is issued while a streaming request is in-flight but not completed

diff --git a/Engine/Plugins/Runtime/HairStrands/Source/HairStrandsCore/Private/GroomResources.cpp b/Engine/Plugins/Runtime/HairStrands/Source/HairStrandsCore/Private/GroomResources.cpp
--- a/Engine/Plugins/Runtime/HairStrands/Source/HairStrandsCore/Private/GroomResources.cpp
+++ b/Engine/Plugins/Runtime/HairStrands/Source/HairStrandsCore/Private/GroomResources.cpp
@@ -733,8 +733,8 @@ void FHairCommonResource::Allocate(FRDGBuilder& GraphBuilder, EHairResourceLoadi
 		// 1. If all requested curve are already loaded, nothing to do
 		if (bIsInitialized && MaxAvailableCurveCount >= InRequestedCurveCount && InternalIsLODDataLoaded(InRequestedCurveCount, InRequestedPointCount, InLODIndex)) 
 		{ 
-			// Trim/Un-stream data if needed
-			if (NeedDeallocation(InRequestedCurveCount, MaxAvailableCurveCount))
+			// Trim/Un-stream data if needed (ensure no streaming request is in-flight)
+			if (NeedDeallocation(InRequestedCurveCount, MaxAvailableCurveCount) && StreamingRequest.IsNone())
 			{
 				StreamingRequest.CurveCount = InRequestedCurveCount;
 				StreamingRequest.PointCount = InRequestedPointCount;
@@ -1275,7 +1275,11 @@ void FHairStrandsClusterResource::InternalAllocate(FRDGBuilder& GraphBuilder)
 
 	const uint32 EntryCount = FMath::DivideAndRoundUp(BulkData.Header.PointCount, HAIR_POINT_LOD_COUNT_PER_UINT);
 
-	InternalCreateStructuredBufferRDG_FromHairBulkData<FHairClusterInfoFormat>(GraphBuilder, BulkData.Data.PackedClusterInfos, BulkData.Header.ClusterCount, ClusterInfoBuffer, ToHairResourceDebugName(TEXT("Hair.StrandsCluster_ClusterInfoBuffer"), ResourceName), OwnerName, EHairResourceUsageType::Static);
+	// Allocated once with all the cluster data (i.e., no streaming)
+	if (ClusterInfoBuffer.Buffer == nullptr)
+	{
+		InternalCreateStructuredBufferRDG_FromHairBulkData<FHairClusterInfoFormat>(GraphBuilder, BulkData.Data.PackedClusterInfos, BulkData.Header.ClusterCount, ClusterInfoBuffer, ToHairResourceDebugName(TEXT("Hair.StrandsCluster_ClusterInfoBuffer"), ResourceName), OwnerName, EHairResourceUsageType::Static);
+	}
 	InternalCreateVertexBufferRDG_FromHairBulkData<FHairClusterIndexFormat>(GraphBuilder, BulkData.Data.CurveToClusterIds, BulkData.Header.CurveCount, CurveToClusterIdBuffer, ToHairResourceDebugName(TEXT("Hair.StrandsCluster_CurveToClusterIds"), ResourceName), OwnerName, EHairResourceUsageType::Static);	
 	InternalCreateVertexBufferRDG_FromHairBulkData<FHairClusterIndexFormat>(GraphBuilder, BulkData.Data.PointLODs, EntryCount, PointLODBuffer, ToHairResourceDebugName(TEXT("Hair.StrandsCluster_PointLOD"), ResourceName), OwnerName, EHairResourceUsageType::Static);
 }
