Fix incorrect group names when groom does not have group ID mapping.

diff --git a/Engine/Plugins/Runtime/HairStrands/Source/HairStrandsCore/Private/GroomBuilder.cpp b/Engine/Plugins/Runtime/HairStrands/Source/HairStrandsCore/Private/GroomBuilder.cpp
--- a/Engine/Plugins/Runtime/HairStrands/Source/HairStrandsCore/Private/GroomBuilder.cpp
+++ b/Engine/Plugins/Runtime/HairStrands/Source/HairStrandsCore/Private/GroomBuilder.cpp
@@ -1659,6 +1659,18 @@ FORCEINLINE FIntVector VectorToIntVector(const FVector& Index)
 	return FIntVector(Index.X, Index.Y, Index.Z);
 }
 
+static FName ToGroupName(const FStrandID& InStrandID, const uint32 InGroupID, const TStrandAttributesConstRef<FName>& InGroupNames)
+{
+	if (InStrandID < InGroupNames.GetNumElements())
+	{
+		return InGroupNames[InStrandID];
+	}
+	else
+	{
+		return FName(FString::Printf(TEXT("Group_%d"), InGroupID));
+	}
+}
+
 bool FGroomBuilder::BuildHairDescriptionGroups(const FHairDescription& HairDescription, FHairDescriptionGroups& Out)
 {
 	TRACE_CPUPROFILER_EVENT_SCOPE(FGroomBuilder::BuildHairDescriptionGroups);
@@ -1725,6 +1737,7 @@ bool FGroomBuilder::BuildHairDescriptionGroups(const FHairDescription& HairDescr
 	TStrandAttributesConstRef<int> GroupIDs = HairDescription.StrandAttributes().GetAttributesRef<int>(HairAttribute::Strand::GroupID);
 	TStrandAttributesConstRef<int> StrandIDs = HairDescription.StrandAttributes().GetAttributesRef<int>(HairAttribute::Strand::ID);
 	TStrandAttributesConstRef<FName> GroupNames = HairDescription.StrandAttributes().GetAttributesRef<FName>(HairAttribute::Strand::GroupName);
+	const uint32 GroupNameCount = GroupNames.GetNumElements();
 
 	bool bImportGuides = true;
 
@@ -1743,19 +1756,56 @@ bool FGroomBuilder::BuildHairDescriptionGroups(const FHairDescription& HairDescr
 	const bool bCanUseClosestGuidesAndWeights = bImportGuides && StrandIDs.IsValid() && (bPrecomputedWeight1 || bPrecomputedWeight3);
 	check(bCanUseClosestGuidesAndWeights == HairDescription.HasAttribute(EHairAttribute::PrecomputedGuideWeights)); // Sanity check
 
-	auto FindOrAdd = [&Out](int32 GroupID, FName GroupName) -> FHairDescriptionGroup&
+	auto FindOrAdd = [&Out, &GroupIDs, &GroupNames](FStrandID StrandID) -> FHairDescriptionGroup&
 	{
-		for (FHairDescriptionGroup& Group : Out.HairGroups)
+		const bool bValidGroupId = GroupIDs.IsValid();
+		const bool bValidGroupName = GroupNames.IsValid();
+
+		// 1. If GroupID are valid, then use them in priority. GroupID can optionally be mapped to GroupName
+		if (bValidGroupId)
 		{
-			if (Group.Info.GroupID == GroupID)
+			const int32 GroupID = GroupIDs[StrandID];
+			for (FHairDescriptionGroup& Group : Out.HairGroups)
 			{
-				return Group;
+				if (Group.Info.GroupID == GroupID)
+				{
+					return Group;
+				}
 			}
+			FHairDescriptionGroup& Group = Out.HairGroups.AddDefaulted_GetRef();
+			Group.Info.GroupID = GroupID;
+			Group.Info.GroupName = ToGroupName(StrandID, GroupID, GroupNames);
+			return Group;
+		}
+		// 2. If no GroupID are valid, use unique GroupName to create GroupID
+		else if (bValidGroupName)
+		{
+			const uint32 GroupID = Out.HairGroups.Num();
+			const FName GroupName = ToGroupName(StrandID, GroupID, GroupNames);
+			for (FHairDescriptionGroup& Group : Out.HairGroups)
+			{
+				if (Group.Info.GroupName == GroupName)
+				{
+					return Group;
+				}
+			}
+			FHairDescriptionGroup& Group = Out.HairGroups.AddDefaulted_GetRef();
+			Group.Info.GroupID = GroupID;
+			Group.Info.GroupName = GroupName;
+			return Group;
+		}
+		// 3. If there are no GroupID and no GroupName, add a default group are valid, use unique group name to create GroupID
+		else
+		{
+			if (Out.HairGroups.IsEmpty())
+			{
+				const uint32 GroupID = 0;
+				FHairDescriptionGroup& Group = Out.HairGroups.AddDefaulted_GetRef();
+				Group.Info.GroupID = GroupID;
+				Group.Info.GroupName = ToGroupName(StrandID, GroupID, GroupNames);
+			}
+			return Out.HairGroups[0];
 		}
-		FHairDescriptionGroup& Group = Out.HairGroups.AddDefaulted_GetRef();
-		Group.Info.GroupID = GroupID;
-		Group.Info.GroupName = (GroupName != NAME_None) ? GroupName : FName(FString::Printf(TEXT("Group_%d"), GroupID));
-		return Group;
 	};
 
 	// Track the imported max hair width among all imported CVs. This information is displayed to artsits to set/tune groom asset later.
@@ -1783,17 +1833,8 @@ bool FGroomBuilder::BuildHairDescriptionGroups(const FHairDescription& HairDescr
 			continue;
 		}
 
-		int32 GroupID = 0;
-		FName GroupName = NAME_None;
-		if (GroupIDs.IsValid())
-		{
-			GroupID = GroupIDs[StrandID];
-			GroupName = GroupNames.IsValid() && StrandID < GroupNames.GetNumElements() ? GroupNames[StrandID] : NAME_None;
-		}
-
 		FHairStrandsDatas* CurrentHairStrandsDatas = nullptr;
-		FHairDescriptionGroup& Group = FindOrAdd(GroupID, GroupName);
-		check(Group.Info.GroupID == GroupID);
+		FHairDescriptionGroup& Group = FindOrAdd(StrandID);
 		bool bNumCurveValid = false;
 		if (!bIsGuide)
 		{
