Fix crash when building groom without name attributes.

diff --git a/Engine/Plugins/Runtime/HairStrands/Source/HairStrandsCore/Private/GroomBuilder.cpp b/Engine/Plugins/Runtime/HairStrands/Source/HairStrandsCore/Private/GroomBuilder.cpp
--- a/Engine/Plugins/Runtime/HairStrands/Source/HairStrandsCore/Private/GroomBuilder.cpp
+++ b/Engine/Plugins/Runtime/HairStrands/Source/HairStrandsCore/Private/GroomBuilder.cpp
@@ -1661,7 +1661,7 @@ FORCEINLINE FIntVector VectorToIntVector(const FVector& Index)
 
 static FName ToGroupName(const FStrandID& InStrandID, const uint32 InGroupID, const TStrandAttributesConstRef<FName>& InGroupNames)
 {
-	if (InStrandID < InGroupNames.GetNumElements())
+	if (InGroupNames.IsValid() && InStrandID < InGroupNames.GetNumElements())
 	{
 		return InGroupNames[InStrandID];
 	}
@@ -1737,7 +1737,6 @@ bool FGroomBuilder::BuildHairDescriptionGroups(const FHairDescription& HairDescr
 	TStrandAttributesConstRef<int> GroupIDs = HairDescription.StrandAttributes().GetAttributesRef<int>(HairAttribute::Strand::GroupID);
 	TStrandAttributesConstRef<int> StrandIDs = HairDescription.StrandAttributes().GetAttributesRef<int>(HairAttribute::Strand::ID);
 	TStrandAttributesConstRef<FName> GroupNames = HairDescription.StrandAttributes().GetAttributesRef<FName>(HairAttribute::Strand::GroupName);
-	const uint32 GroupNameCount = GroupNames.GetNumElements();
 
 	bool bImportGuides = true;
 
