Fix incorrect RBF weight loading/streaming.

diff --git a/Engine/Plugins/Runtime/HairStrands/Source/HairStrandsCore/Private/GroomBindingBuilder.cpp b/Engine/Plugins/Runtime/HairStrands/Source/HairStrandsCore/Private/GroomBindingBuilder.cpp
--- a/Engine/Plugins/Runtime/HairStrands/Source/HairStrandsCore/Private/GroomBindingBuilder.cpp
+++ b/Engine/Plugins/Runtime/HairStrands/Source/HairStrandsCore/Private/GroomBindingBuilder.cpp
@@ -11,6 +11,7 @@
 #include "Async/ParallelFor.h"
 #include "GlobalShader.h"
 #include "Misc/ScopedSlowTask.h"
+#include "GroomRBFDeformer.h"
 
 #if WITH_EDITORONLY_DATA
 
@@ -57,7 +58,7 @@ static FAutoConsoleVariableRef CVarHairStrandsBindingBuilderWarningEnable(TEXT("
 FString FGroomBindingBuilder::GetVersion()
 {
 	// Important to update the version when groom building changes
-	return TEXT("3f");
+	return TEXT("3h");
 }
 
 ///////////////////////////////////////////////////////////////////////////////////////////////////
@@ -562,11 +563,20 @@ namespace GroomBinding_RBFWeighting
 	FWeightsBuilder::FWeightsBuilder(const uint32 NumRows, const uint32 NumColumns,
 		const FVector3f* SourcePositions, const FVector3f* TargetPositions)
 	{
-		const uint32 PolyRows = NumRows + 4;
-		const uint32 PolyColumns = NumColumns + 4;
+		const uint32 PolyRows 	 = FGroomRBFDeformer::GetEntryCount(NumRows);
+		const uint32 PolyColumns = FGroomRBFDeformer::GetEntryCount(NumColumns);
+
+		const uint32 SampleCount = NumRows;
+		const uint32 WeightCount = FGroomRBFDeformer::GetWeightCount(SampleCount);
 
 		MatrixEntries.Init(0.0, PolyRows * PolyColumns);
 		InverseEntries.Init(0.0, PolyRows * PolyColumns);
+
+		// Sanity check
+		check(NumRows == NumColumns);
+		check(WeightCount == MatrixEntries.Num());
+		check(WeightCount == InverseEntries.Num());
+
 		TArray<float>& LocalEntries = MatrixEntries;
 		ParallelFor(NumRows,
 			[
diff --git a/Engine/Plugins/Runtime/HairStrands/Source/HairStrandsCore/Private/GroomRBFDeformer.cpp b/Engine/Plugins/Runtime/HairStrands/Source/HairStrandsCore/Private/GroomRBFDeformer.cpp
--- a/Engine/Plugins/Runtime/HairStrands/Source/HairStrandsCore/Private/GroomRBFDeformer.cpp
+++ b/Engine/Plugins/Runtime/HairStrands/Source/HairStrandsCore/Private/GroomRBFDeformer.cpp
@@ -41,16 +41,16 @@ void UpdateMeshSamples(
 	TArray<FVector3f>& OutSampleDeformationsBuffer
 )
 {
-	OutSampleDeformationsBuffer.SetNum(MaxSampleCount + 4);
+	const uint32 EntryCount = FGroomRBFDeformer::GetEntryCount(MaxSampleCount);
+	OutSampleDeformationsBuffer.SetNum(EntryCount);
 
-	for (uint32 SampleIndex = 0; SampleIndex < MaxSampleCount + 4; ++SampleIndex)
+	for (uint32 SampleIndex = 0; SampleIndex < EntryCount; ++SampleIndex)
 	{
-		uint32 WeightsOffset = SampleIndex * (MaxSampleCount + 4);
+		uint32 WeightsOffset = SampleIndex * EntryCount;
 		FVector3f SampleDeformation(FVector3f::ZeroVector);
 		for (uint32 i = 0; i < MaxSampleCount; ++i, ++WeightsOffset)
 		{
-			SampleDeformation += InterpolationWeightsBuffer[WeightsOffset] *
-				(SampleDeformedPositionsBuffer[i] - SampleRestPositionsBuffer[i]); 
+			SampleDeformation += InterpolationWeightsBuffer[WeightsOffset] * (SampleDeformedPositionsBuffer[i] - SampleRestPositionsBuffer[i]);
 		}
 
 		OutSampleDeformationsBuffer[SampleIndex] = SampleDeformation;
@@ -822,3 +822,15 @@ void FGroomRBFDeformer::GetRBFDeformedGroomAsset(const UGroomAsset* InGroomAsset
 	}
 #endif // #if WITH_EDITORONLY_DATA
 }
+
+uint32 FGroomRBFDeformer::GetEntryCount(uint32 InSampleCount)
+{
+	return HAIR_RBF_ENTRY_COUNT(InSampleCount);
+}
+
+uint32 FGroomRBFDeformer::GetWeightCount(uint32 InSampleCount)
+{
+	const uint32 EntryCount = GetEntryCount(InSampleCount);
+	return EntryCount * EntryCount;
+}
+
diff --git a/Engine/Plugins/Runtime/HairStrands/Source/HairStrandsCore/Private/GroomResources.cpp b/Engine/Plugins/Runtime/HairStrands/Source/HairStrandsCore/Private/GroomResources.cpp
--- a/Engine/Plugins/Runtime/HairStrands/Source/HairStrandsCore/Private/GroomResources.cpp
+++ b/Engine/Plugins/Runtime/HairStrands/Source/HairStrandsCore/Private/GroomResources.cpp
@@ -29,6 +29,7 @@
 #include "RenderTargetPool.h"
 #include "RHIUtilities.h"
 #include "GroomBindingBuilder.h"
+#include "GroomRBFDeformer.h"
 #if WITH_EDITORONLY_DATA
 #include "Containers/StringView.h"
 #include "DerivedDataCache.h"
@@ -1456,9 +1457,13 @@ void FHairStrandsRestRootResource::InternalAllocate(FRDGBuilder& GraphBuilder, u
 			InternalCreateVertexBufferRDG_FromHairBulkData<FHairStrandsMeshTrianglePositionFormat>(GraphBuilder, CPUData.RestUniqueTrianglePositionBuffer, LODHeader.UniqueTriangleCount * 3, GPUData.RestUniqueTrianglePositionBuffer, ToHairResourceDebugName(HAIRSTRANDS_RESOUCE_NAME(CurveType, Hair.StrandsRestRoot_RestUniqueTrianglePosition0Buffer), ResourceName), OwnerName, EHairResourceUsageType::Static);
 		}
 
-		if (LODHeader.SampleCount > 0)
+		// Allocate only once
+		if (LODHeader.SampleCount > 0 && GPUData.MeshInterpolationWeightsBuffer.Buffer == nullptr)
 		{
 			const uint32 InteroplationWeightCount = CPUData.MeshInterpolationWeightsBuffer.GetBulkDataSize() / sizeof(FHairStrandsWeightFormat::Type);
+			const uint32 WeightCount = FGroomRBFDeformer::GetWeightCount(LODHeader.SampleCount);
+			check(WeightCount == InteroplationWeightCount); // Sanity check
+
 			InternalCreateVertexBufferRDG_FromHairBulkData<FHairStrandsWeightFormat>(GraphBuilder, CPUData.MeshInterpolationWeightsBuffer, InteroplationWeightCount, GPUData.MeshInterpolationWeightsBuffer, ToHairResourceDebugName(HAIRSTRANDS_RESOUCE_NAME(CurveType, Hair.StrandsRestRoot_MeshInterpolationWeightsBuffer), ResourceName), OwnerName, EHairResourceUsageType::Static);
 			InternalCreateVertexBufferRDG_FromHairBulkData<FHairStrandsIndexFormat>(GraphBuilder, CPUData.MeshSampleIndicesBuffer, LODHeader.SampleCount, GPUData.MeshSampleIndicesBuffer, ToHairResourceDebugName(HAIRSTRANDS_RESOUCE_NAME(CurveType, Hair.StrandsRestRoot_MeshSampleIndicesBuffer), ResourceName), OwnerName, EHairResourceUsageType::Static);
 			InternalCreateVertexBufferRDG_FromHairBulkData<FHairStrandsMeshTrianglePositionFormat>(GraphBuilder, CPUData.RestSamplePositionsBuffer, LODHeader.SampleCount, GPUData.RestSamplePositionsBuffer, ToHairResourceDebugName(HAIRSTRANDS_RESOUCE_NAME(CurveType, Hair.StrandsRestRoot_RestSamplePositionsBuffer), ResourceName), OwnerName, EHairResourceUsageType::Static);
@@ -1543,14 +1548,16 @@ void FHairStrandsDeformedRootResource::InternalAllocate(FRDGBuilder& GraphBuilde
 			LOD.AvailableCurveCount = FMath::Min(InCurveCount, RootCount);
 			if (LOD.SampleCount > 0)
 			{
+				const uint32 RBFEntryCount = FGroomRBFDeformer::GetEntryCount(LOD.SampleCount);
+
 				InternalCreateVertexBufferRDG<FHairStrandsMeshTrianglePositionFormat>(GraphBuilder, LOD.SampleCount, LOD.DeformedSamplePositionsBuffer[0], ToHairResourceDebugName(HAIRSTRANDS_RESOUCE_NAME(CurveType, Hair.StrandsRootDeformed_DeformedSamplePositionsBuffer0), ResourceName), OwnerName, EHairResourceUsageType::Dynamic);
-				InternalCreateVertexBufferRDG<FHairStrandsMeshTrianglePositionFormat>(GraphBuilder, LOD.SampleCount + 4, LOD.MeshSampleWeightsBuffer[0], ToHairResourceDebugName(HAIRSTRANDS_RESOUCE_NAME(CurveType, Hair.StrandsRootDeformed_MeshSampleWeightsBuffer0), ResourceName), OwnerName, EHairResourceUsageType::Dynamic);
+				InternalCreateVertexBufferRDG<FHairStrandsMeshTrianglePositionFormat>(GraphBuilder, RBFEntryCount, LOD.MeshSampleWeightsBuffer[0], ToHairResourceDebugName(HAIRSTRANDS_RESOUCE_NAME(CurveType, Hair.StrandsRootDeformed_MeshSampleWeightsBuffer0), ResourceName), OwnerName, EHairResourceUsageType::Dynamic);
 
 				// Double buffering is disabled by default unless the read-only cvar r.HairStrands.ContinuousDecimationReordering is set
 				if (IsHairStrandContinuousDecimationReorderingEnabled())
 				{
 					InternalCreateVertexBufferRDG<FHairStrandsMeshTrianglePositionFormat>(GraphBuilder, LOD.SampleCount, LOD.DeformedSamplePositionsBuffer[1], ToHairResourceDebugName(HAIRSTRANDS_RESOUCE_NAME(CurveType, Hair.StrandsRootDeformed_DeformedSamplePositionsBuffer1), ResourceName), OwnerName, EHairResourceUsageType::Dynamic);
-					InternalCreateVertexBufferRDG<FHairStrandsMeshTrianglePositionFormat>(GraphBuilder, LOD.SampleCount + 4, LOD.MeshSampleWeightsBuffer[1], ToHairResourceDebugName(HAIRSTRANDS_RESOUCE_NAME(CurveType, Hair.StrandsRootDeformed_MeshSampleWeightsBuffer1), ResourceName), OwnerName, EHairResourceUsageType::Dynamic);
+					InternalCreateVertexBufferRDG<FHairStrandsMeshTrianglePositionFormat>(GraphBuilder, RBFEntryCount, LOD.MeshSampleWeightsBuffer[1], ToHairResourceDebugName(HAIRSTRANDS_RESOUCE_NAME(CurveType, Hair.StrandsRootDeformed_MeshSampleWeightsBuffer1), ResourceName), OwnerName, EHairResourceUsageType::Dynamic);
 				}
 			}
 
diff --git a/Engine/Plugins/Runtime/HairStrands/Source/HairStrandsCore/Private/HairStrandsDatas.cpp b/Engine/Plugins/Runtime/HairStrands/Source/HairStrandsCore/Private/HairStrandsDatas.cpp
--- a/Engine/Plugins/Runtime/HairStrands/Source/HairStrandsCore/Private/HairStrandsDatas.cpp
+++ b/Engine/Plugins/Runtime/HairStrands/Source/HairStrandsCore/Private/HairStrandsDatas.cpp
@@ -5,6 +5,7 @@
 #include "UObject/UE5ReleaseStreamObjectVersion.h"
 #include "HairAttributes.h"
 #include "IO/IoDispatcher.h"
+#include "GroomRBFDeformer.h"
 
 #if WITH_EDITOR
 #include "DerivedDataCache.h"
@@ -889,14 +890,16 @@ void FHairStrandsRootBulkData::GetResources(FQuery& Out)
 
 	Out.Add(DataLOD.RootToUniqueTriangleIndexBuffer, 	TEXT("_RootToUniqueTriangleIndexBuffer"), 		DataLOD.RootToUniqueTriangleIndexBuffer.LoadedSize, 	RootCount * Header.Strides.RootToUniqueTriangleIndexBufferStride);
 	Out.Add(DataLOD.RootBarycentricBuffer, 				TEXT("_RootBarycentricBuffer"), 				DataLOD.RootBarycentricBuffer.LoadedSize, 				RootCount * Header.Strides.RootBarycentricBufferStride);
-	Out.Add(DataLOD.UniqueTriangleIndexBuffer, 			TEXT("_UniqueTriangleIndexBuffer"), 			DataLOD.UniqueTriangleIndexBuffer.LoadedSize, 			HeaderLOD.UniqueTriangleCount * Header.Strides.UniqueTriangleIndexBufferStride);
-	Out.Add(DataLOD.RestUniqueTrianglePositionBuffer, 	TEXT("_RestUniqueTrianglePositionBuffer"), 		DataLOD.RestUniqueTrianglePositionBuffer.LoadedSize, 	HeaderLOD.UniqueTriangleCount * Header.Strides.RestUniqueTrianglePositionBufferStride);
+	Out.Add(DataLOD.UniqueTriangleIndexBuffer, 			TEXT("_UniqueTriangleIndexBuffer"), 			DataLOD.UniqueTriangleIndexBuffer.LoadedSize, 			HeaderLOD.UniqueTriangleCount * Header.Strides.UniqueTriangleIndexBufferStride);		 // Load all data
+	Out.Add(DataLOD.RestUniqueTrianglePositionBuffer, 	TEXT("_RestUniqueTrianglePositionBuffer"), 		DataLOD.RestUniqueTrianglePositionBuffer.LoadedSize, 	HeaderLOD.UniqueTriangleCount * Header.Strides.RestUniqueTrianglePositionBufferStride);	 // Load all data
 
 	if (HeaderLOD.SampleCount > 0)
 	{
-		Out.Add(DataLOD.MeshInterpolationWeightsBuffer, TEXT("_MeshInterpolationWeightsBuffer"), 		DataLOD.MeshInterpolationWeightsBuffer.LoadedSize, 		HeaderLOD.SampleCount * Header.Strides.MeshInterpolationWeightsBufferStride);
-		Out.Add(DataLOD.MeshSampleIndicesBuffer, 		TEXT("_MeshSampleIndicesBuffer"), 				DataLOD.MeshSampleIndicesBuffer.LoadedSize, 			HeaderLOD.SampleCount * Header.Strides.MeshSampleIndicesBufferStride);
-		Out.Add(DataLOD.RestSamplePositionsBuffer, 		TEXT("_RestSamplePositionsBuffer"), 			DataLOD.RestSamplePositionsBuffer.LoadedSize, 			HeaderLOD.SampleCount * Header.Strides.RestSamplePositionsBufferStride);
+		const uint32 RBFWeightCount = FGroomRBFDeformer::GetWeightCount(HeaderLOD.SampleCount); 
+
+		Out.Add(DataLOD.MeshInterpolationWeightsBuffer, TEXT("_MeshInterpolationWeightsBuffer"), 		DataLOD.MeshInterpolationWeightsBuffer.LoadedSize, 		RBFWeightCount * Header.Strides.MeshInterpolationWeightsBufferStride); 	// Load all data
+		Out.Add(DataLOD.MeshSampleIndicesBuffer, 		TEXT("_MeshSampleIndicesBuffer"), 				DataLOD.MeshSampleIndicesBuffer.LoadedSize, 			HeaderLOD.SampleCount * Header.Strides.MeshSampleIndicesBufferStride); 	// Load all data
+		Out.Add(DataLOD.RestSamplePositionsBuffer, 		TEXT("_RestSamplePositionsBuffer"), 			DataLOD.RestSamplePositionsBuffer.LoadedSize, 			HeaderLOD.SampleCount * Header.Strides.RestSamplePositionsBufferStride);// Load all data
 	}
 }
 
diff --git a/Engine/Plugins/Runtime/HairStrands/Source/HairStrandsCore/Public/GroomRBFDeformer.h b/Engine/Plugins/Runtime/HairStrands/Source/HairStrandsCore/Public/GroomRBFDeformer.h
--- a/Engine/Plugins/Runtime/HairStrands/Source/HairStrandsCore/Public/GroomRBFDeformer.h
+++ b/Engine/Plugins/Runtime/HairStrands/Source/HairStrandsCore/Public/GroomRBFDeformer.h
@@ -12,4 +12,7 @@ struct HAIRSTRANDSCORE_API FGroomRBFDeformer
 {
 	// Return a new GroomAsset with the RBF deformation from the BindingAsset baked into it
 	void GetRBFDeformedGroomAsset(const UGroomAsset* InGroomAsset, const UGroomBindingAsset* BindingAsset, FTextureSource* MaskTextureSource, const float MaskScale, UGroomAsset* OutGroomAsset);
+
+	static uint32 GetEntryCount(uint32 InSampleCount);
+	static uint32 GetWeightCount(uint32 InSampleCount);
 };
\ No newline at end of file
diff --git a/Engine/Shaders/Private/HairStrands/HairStrandsMesh.usf b/Engine/Shaders/Private/HairStrands/HairStrandsMesh.usf
--- a/Engine/Shaders/Private/HairStrands/HairStrandsMesh.usf
+++ b/Engine/Shaders/Private/HairStrands/HairStrandsMesh.usf
@@ -188,17 +188,16 @@ RWBuffer<float4> OutSampleDeformationsBuffer;
 void MainCS(uint3 DispatchThreadId : SV_DispatchThreadID)
 { 
 	const uint SampleIndex = DispatchThreadId.x;
-	if (SampleIndex >= (MaxSampleCount + 4))
+	if (SampleIndex >= HAIR_RBF_ENTRY_COUNT(MaxSampleCount))
 	{
 		return;
 	}
 
-	uint WeightsOffset = SampleIndex * (MaxSampleCount+4);
+	uint WeightsOffset = SampleIndex * HAIR_RBF_ENTRY_COUNT(MaxSampleCount);
 	float3 SampleDeformation = float3(0,0,0);
 	for(uint i = 0; i < MaxSampleCount; ++i, ++WeightsOffset)
 	{
-		SampleDeformation += InterpolationWeightsBuffer[WeightsOffset] * 
-				(SampleDeformedPositionsBuffer[i] - SampleRestPositionsBuffer[i]).xyz;
+		SampleDeformation += InterpolationWeightsBuffer[WeightsOffset] * (SampleDeformedPositionsBuffer[i] - SampleRestPositionsBuffer[i]).xyz;
 	}
 
 	OutSampleDeformationsBuffer[SampleIndex] = float4(SampleDeformation,1);
diff --git a/Engine/Shaders/Shared/HairStrandsDefinitions.h b/Engine/Shaders/Shared/HairStrandsDefinitions.h
--- a/Engine/Shaders/Shared/HairStrandsDefinitions.h
+++ b/Engine/Shaders/Shared/HairStrandsDefinitions.h
@@ -94,4 +94,6 @@
 // Each curve's point encode its minium LOD at which it becomes active
 #define HAIR_POINT_LOD_BIT_COUNT 4
 #define HAIR_POINT_LOD_COUNT_PER_UINT (32u/4u)
-#define HAIR_POINT_LOD_COUNT_PER_UINT_DIV_AS_SHIFT 3
\ No newline at end of file
+#define HAIR_POINT_LOD_COUNT_PER_UINT_DIV_AS_SHIFT 3
+
+#define HAIR_RBF_ENTRY_COUNT(InSampleCount)	((InSampleCount)+4u)
\ No newline at end of file
