Fix hair cards RBF visualization.

diff --git a/Engine/Plugins/Runtime/HairStrands/Source/HairStrandsCore/Private/GroomDebug.cpp b/Engine/Plugins/Runtime/HairStrands/Source/HairStrandsCore/Private/GroomDebug.cpp
--- a/Engine/Plugins/Runtime/HairStrands/Source/HairStrandsCore/Private/GroomDebug.cpp
+++ b/Engine/Plugins/Runtime/HairStrands/Source/HairStrandsCore/Private/GroomDebug.cpp
@@ -1338,8 +1338,28 @@ void RunHairStrandsDebug(
 						continue;
 
 					const bool bRenderStrands = StrandType == EHairStrandsInterpolationType::RenderStrands;
-					FHairStrandsRestRootResource* RestRootResource = bRenderStrands ? Instance->Strands.RestRootResource : Instance->Guides.RestRootResource;
-					FHairStrandsDeformedRootResource* DeformedRootResource = bRenderStrands ? Instance->Strands.DeformedRootResource : Instance->Guides.DeformedRootResource;
+
+					FHairStrandsRestRootResource* RestRootResource = nullptr;
+					FHairStrandsDeformedRootResource* DeformedRootResource = nullptr;
+					if (!bRenderStrands)
+					{
+						RestRootResource 	 = Instance->Guides.RestRootResource;
+						DeformedRootResource = Instance->Guides.DeformedRootResource;
+					}
+					else if (Instance->GetHairGeometry() == EHairGeometryType::Strands)
+					{
+						RestRootResource 	 = Instance->Strands.RestRootResource;
+						DeformedRootResource = Instance->Strands.DeformedRootResource;
+					}
+					else if (Instance->GetHairGeometry() == EHairGeometryType::Cards)
+					{
+						const int32 LODIndex = Instance->HairGroupPublicData->GetIntLODIndex();
+						if (Instance->Cards.IsValid(LODIndex))
+						{
+							RestRootResource 	 = Instance->Cards.LODs[LODIndex].Guides.RestRootResource;
+							DeformedRootResource = Instance->Cards.LODs[LODIndex].Guides.DeformedRootResource;
+						}
+					}
 					if (RestRootResource == nullptr || DeformedRootResource == nullptr)
 						continue;
 
