Fix niagara cache hair interpolation

diff --git a/Engine/Plugins/Runtime/HairStrands/Shaders/Private/NiagaraInterpolateGroomGuides.usf b/Engine/Plugins/Runtime/HairStrands/Shaders/Private/NiagaraInterpolateGroomGuides.usf
--- a/Engine/Plugins/Runtime/HairStrands/Shaders/Private/NiagaraInterpolateGroomGuides.usf
+++ b/Engine/Plugins/Runtime/HairStrands/Shaders/Private/NiagaraInterpolateGroomGuides.usf
@@ -72,9 +72,12 @@ void MainCS(uint3 DispatchThreadId : SV_DispatchThreadID)
 	if (DispatchThreadId.x < uint(NumPoints))
 	{
 		const int StrandIndex = PointToCurveIndexBuffer[DispatchThreadId.x];
+
+		const FHairCurve Curve = UnpackHairCurve(CurvesOffsetsBuffer[StrandIndex]);
+		const int PointBegin = Curve.PointIndex;
+		const int PointCount = Curve.PointCount;
+		const int PointIndex = DispatchThreadId.x - PointBegin;
 		
-		const int PointCount = CurvesOffsetsBuffer[StrandIndex + 1] - CurvesOffsetsBuffer[StrandIndex];
-		const int PointIndex = DispatchThreadId.x - CurvesOffsetsBuffer[StrandIndex];
 		
 		const float PointCoord = (float) (PointIndex) / (PointCount - 1.0);
 		const float PointNode = PointCoord * (StrandsSize - 2);
