Fix incorrect building and serialization of RBF data for rendering strands, causing a crash at runtime.

diff --git a/Engine/Plugins/Runtime/HairStrands/Source/HairStrandsCore/Private/GroomBindingBuilder.cpp b/Engine/Plugins/Runtime/HairStrands/Source/HairStrandsCore/Private/GroomBindingBuilder.cpp
--- a/Engine/Plugins/Runtime/HairStrands/Source/HairStrandsCore/Private/GroomBindingBuilder.cpp
+++ b/Engine/Plugins/Runtime/HairStrands/Source/HairStrandsCore/Private/GroomBindingBuilder.cpp
@@ -58,7 +58,7 @@ static FAutoConsoleVariableRef CVarHairStrandsBindingBuilderWarningEnable(TEXT("
 FString FGroomBindingBuilder::GetVersion()
 {
 	// Important to update the version when groom building changes
-	return TEXT("3h");
+	return TEXT("3i");
 }
 
 ///////////////////////////////////////////////////////////////////////////////////////////////////
@@ -708,6 +708,15 @@ namespace GroomBinding_RBFWeighting
 		}
 	}
 
+	void ResetSampleData(FHairStrandsRootData& Out, uint32 LODIndex)
+	{
+		FHairStrandsRootData::FMeshProjectionLOD& LOD = Out.MeshProjectionLODs[LODIndex];
+		LOD.SampleCount = 0;	
+		LOD.MeshInterpolationWeightsBuffer.SetNum(0);
+		LOD.MeshSampleIndicesBuffer.SetNum(0);
+		LOD.RestSamplePositionsBuffer.SetNum(0);
+	}
+
 	void ComputeInterpolationWeights(
 		FHairRootGroupData& Out, 
 		const bool bNeedStrandsRoot,
@@ -746,17 +755,14 @@ namespace GroomBinding_RBFWeighting
 					FPointsSampler PointsSampler(ValidPoints, PositionsPointer, MaxSamples);
 					const uint32 SampleCount = PointsSampler.SamplePositions.Num();
 
-					FWeightsBuilder InterpolationWeights(SampleCount, SampleCount,
-						PointsSampler.SamplePositions.GetData(), PointsSampler.SamplePositions.GetData());
+					FWeightsBuilder InterpolationWeights(SampleCount, SampleCount, PointsSampler.SamplePositions.GetData(), PointsSampler.SamplePositions.GetData());
 
 					// Guides
 					UpdateInterpolationWeights(InterpolationWeights, PointsSampler, LODIndex, Out.SimRootData);
 
 					// Strands
-					if (bNeedStrandsRoot)
-					{
-						UpdateInterpolationWeights(InterpolationWeights, PointsSampler, LODIndex, Out.RenRootData);
-					}
+					// No sample data, only used/available for guides
+					ResetSampleData(Out.RenRootData, LODIndex);
 				}
 			}
 			else
@@ -768,17 +774,14 @@ namespace GroomBinding_RBFWeighting
 				FPointsSampler PointsSampler(ValidPoints, PositionsPointer, MaxSamples);
 				const uint32 SampleCount = PointsSampler.SamplePositions.Num();
 
-				FWeightsBuilder InterpolationWeights(SampleCount, SampleCount,
-					PointsSampler.SamplePositions.GetData(), PointsSampler.SamplePositions.GetData());
+				FWeightsBuilder InterpolationWeights(SampleCount, SampleCount, PointsSampler.SamplePositions.GetData(), PointsSampler.SamplePositions.GetData());
 
 				// Guides
 				UpdateInterpolationWeights(InterpolationWeights, PointsSampler, LODIndex, Out.SimRootData);
 
-				// Strands
-				if (bNeedStrandsRoot)
-				{
-					UpdateInterpolationWeights(InterpolationWeights, PointsSampler, LODIndex, Out.RenRootData);
-				}
+				// Strands 
+				// No sample data, only used/available for guides
+				ResetSampleData(Out.RenRootData, LODIndex);
 			}
 		}
 	}
