Fix incorrect curve LOD lerping.

diff --git a/Engine/Plugins/Runtime/HairStrands/Source/HairStrandsCore/Private/HairStrandsDatas.cpp b/Engine/Plugins/Runtime/HairStrands/Source/HairStrandsCore/Private/HairStrandsDatas.cpp
--- a/Engine/Plugins/Runtime/HairStrands/Source/HairStrandsCore/Private/HairStrandsDatas.cpp
+++ b/Engine/Plugins/Runtime/HairStrands/Source/HairStrandsCore/Private/HairStrandsDatas.cpp
@@ -782,9 +782,9 @@ uint32 FHairStrandsClusterBulkData::GetCurveCount(float InLODIndex) const
 	const int32 iLODIndex = InLODIndex;
 	const float S = InLODIndex - iLODIndex;
 
-	const uint32 LODIndex0 = FMath::Min(iLODIndex,   LODInfos.Num()-1);
-	const uint32 LODIndex1 = FMath::Min(iLODIndex+1, LODInfos.Num()-1);
-	return FMath::Lerp(LODInfos[LODIndex1].CurveCount, LODInfos[LODIndex0].CurveCount, 1.f-S);
+	const int32 LODIndex0 = FMath::Min(iLODIndex,   LODInfos.Num()-1);
+	const int32 LODIndex1 = FMath::Min(iLODIndex+1, LODInfos.Num()-1);
+	return FMath::LerpStable<int32>(LODInfos[LODIndex0].CurveCount, LODInfos[LODIndex1].CurveCount, S);
 }
 
 /////////////////////////////////////////////////////////////////////////////////////////
