Fix missing cached tangent invalidation causing groom to disappear when properties are edited, and incorrect rendering when when auto LOD is enabled.

diff --git a/Engine/Plugins/Runtime/HairStrands/Source/HairStrandsCore/Private/GroomResources.cpp b/Engine/Plugins/Runtime/HairStrands/Source/HairStrandsCore/Private/GroomResources.cpp
--- a/Engine/Plugins/Runtime/HairStrands/Source/HairStrandsCore/Private/GroomResources.cpp
+++ b/Engine/Plugins/Runtime/HairStrands/Source/HairStrandsCore/Private/GroomResources.cpp
@@ -1159,6 +1159,7 @@ void FHairStrandsRestResource::InternalRelease()
 	TangentBuffer.Release();
 	CurveBuffer.Release();
 	MaxAvailableCurveCount = 0;
+	CachedTangentPointCount = 0;
 }
 
 void FHairStrandsRestResource::InternalResetLoadedSize()
