[Backout] - CL26540245
Original CL Desc
-----------------------------------------------------------------
Fix incorrect building and serialization of RBF data for rendering strands, causing a crash at runtime.

diff --git a/Engine/Plugins/Runtime/HairStrands/Source/HairStrandsCore/Private/GroomBindingBuilder.cpp b/Engine/Plugins/Runtime/HairStrands/Source/HairStrandsCore/Private/GroomBindingBuilder.cpp
--- a/Engine/Plugins/Runtime/HairStrands/Source/HairStrandsCore/Private/GroomBindingBuilder.cpp
+++ b/Engine/Plugins/Runtime/HairStrands/Source/HairStrandsCore/Private/GroomBindingBuilder.cpp
@@ -58,7 +58,7 @@ static FAutoConsoleVariableRef CVarHairStrandsBindingBuilderWarningEnable(TEXT("
 FString FGroomBindingBuilder::GetVersion()
 {
 	// Important to update the version when groom building changes
-	return TEXT("3i");
+	return TEXT("3h");
 }
 
 ///////////////////////////////////////////////////////////////////////////////////////////////////
@@ -708,15 +708,6 @@ namespace GroomBinding_RBFWeighting
 		}
 	}
 
-	void ResetSampleData(FHairStrandsRootData& Out, uint32 LODIndex)
-	{
-		FHairStrandsRootData::FMeshProjectionLOD& LOD = Out.MeshProjectionLODs[LODIndex];
-		LOD.SampleCount = 0;	
-		LOD.MeshInterpolationWeightsBuffer.SetNum(0);
-		LOD.MeshSampleIndicesBuffer.SetNum(0);
-		LOD.RestSamplePositionsBuffer.SetNum(0);
-	}
-
 	void ComputeInterpolationWeights(
 		FHairRootGroupData& Out, 
 		const bool bNeedStrandsRoot,
@@ -755,14 +746,17 @@ namespace GroomBinding_RBFWeighting
 					FPointsSampler PointsSampler(ValidPoints, PositionsPointer, MaxSamples);
 					const uint32 SampleCount = PointsSampler.SamplePositions.Num();
 
-					FWeightsBuilder InterpolationWeights(SampleCount, SampleCount, PointsSampler.SamplePositions.GetData(), PointsSampler.SamplePositions.GetData());
+					FWeightsBuilder InterpolationWeights(SampleCount, SampleCount,
+						PointsSampler.SamplePositions.GetData(), PointsSampler.SamplePositions.GetData());
 
 					// Guides
 					UpdateInterpolationWeights(InterpolationWeights, PointsSampler, LODIndex, Out.SimRootData);
 
 					// Strands
-					// No sample data, only used/available for guides
-					ResetSampleData(Out.RenRootData, LODIndex);
+					if (bNeedStrandsRoot)
+					{
+						UpdateInterpolationWeights(InterpolationWeights, PointsSampler, LODIndex, Out.RenRootData);
+					}
 				}
 			}
 			else
@@ -774,14 +768,17 @@ namespace GroomBinding_RBFWeighting
 				FPointsSampler PointsSampler(ValidPoints, PositionsPointer, MaxSamples);
 				const uint32 SampleCount = PointsSampler.SamplePositions.Num();
 
-				FWeightsBuilder InterpolationWeights(SampleCount, SampleCount, PointsSampler.SamplePositions.GetData(), PointsSampler.SamplePositions.GetData());
+				FWeightsBuilder InterpolationWeights(SampleCount, SampleCount,
+					PointsSampler.SamplePositions.GetData(), PointsSampler.SamplePositions.GetData());
 
 				// Guides
 				UpdateInterpolationWeights(InterpolationWeights, PointsSampler, LODIndex, Out.SimRootData);
 
-				// Strands 
-				// No sample data, only used/available for guides
-				ResetSampleData(Out.RenRootData, LODIndex);
+				// Strands
+				if (bNeedStrandsRoot)
+				{
+					UpdateInterpolationWeights(InterpolationWeights, PointsSampler, LODIndex, Out.RenRootData);
+				}
 			}
 		}
 	}
