Fix incorrect tangent update with auto-LOD on static groom.

diff --git a/Engine/Plugins/Runtime/HairStrands/Source/HairStrandsCore/Private/GroomResources.cpp b/Engine/Plugins/Runtime/HairStrands/Source/HairStrandsCore/Private/GroomResources.cpp
--- a/Engine/Plugins/Runtime/HairStrands/Source/HairStrandsCore/Private/GroomResources.cpp
+++ b/Engine/Plugins/Runtime/HairStrands/Source/HairStrandsCore/Private/GroomResources.cpp
@@ -1123,16 +1123,18 @@ void AddHairTangentPass(
 
 FRDGExternalBuffer FHairStrandsRestResource::GetTangentBuffer(FRDGBuilder& GraphBuilder, FGlobalShaderMap* ShaderMap, uint32 ActivePointCount, uint32 ActiveCurveCount)
 {
-	// Lazy allocation and update
-	if (TangentBuffer.Buffer == nullptr)
+	const uint32 AllocatedPoint = TangentBuffer.Buffer ? TangentBuffer.Buffer->Desc.NumElements / FHairStrandsTangentFormat::ComponentCount : 0u;
+	const uint32 RequestedPoint = ActivePointCount;
+
+	// Lazy allocation
+	const bool bReallocateTangentBuffer = TangentBuffer.Buffer == nullptr || ActivePointCount > AllocatedPoint;
+	if (bReallocateTangentBuffer)
 	{
 		InternalCreateVertexBufferRDG<FHairStrandsTangentFormat>(GraphBuilder, ActivePointCount * FHairStrandsTangentFormat::ComponentCount, TangentBuffer, ToHairResourceDebugName(TEXT("Hair.StrandsRest_TangentBuffer"), ResourceName), OwnerName, EHairResourceUsageType::Dynamic);
 	}
 
-	const uint32 AllocatedPoint = TangentBuffer.Buffer->Desc.NumElements / FHairStrandsTangentFormat::ComponentCount;
-	const uint32 RequestedPoint = ActivePointCount;
-
-	const bool bRecomputeTangent = (AllocatedPoint >= RequestedPoint) && (ActivePointCount > CachedTangentPointCount);
+	// Lazy update tangent
+	const bool bRecomputeTangent = bReallocateTangentBuffer || ((AllocatedPoint >= RequestedPoint) && (ActivePointCount > CachedTangentPointCount));
 	if (bRecomputeTangent)
 	{
 		AddHairTangentPass(
