Fix mismatching shader/buffer declaration causing Vulkan crash.

diff --git a/Engine/Plugins/Runtime/HairStrands/Source/HairStrandsCore/Private/GroomResources.cpp b/Engine/Plugins/Runtime/HairStrands/Source/HairStrandsCore/Private/GroomResources.cpp
--- a/Engine/Plugins/Runtime/HairStrands/Source/HairStrandsCore/Private/GroomResources.cpp
+++ b/Engine/Plugins/Runtime/HairStrands/Source/HairStrandsCore/Private/GroomResources.cpp
@@ -1648,7 +1648,7 @@ FHairCardsInterpolationResource::FHairCardsInterpolationResource(FHairCardsInter
 
 void FHairCardsInterpolationResource::InternalAllocate(FRDGBuilder& GraphBuilder)
 {
-	InternalCreateVertexBufferRDG<FHairCardsInterpolationFormat>(GraphBuilder, BulkData.Interpolation, InterpolationBuffer, ToHairResourceDebugName(TEXT("Hair.CardsInterpolation_InterpolationBuffer"), ResourceName), OwnerName, EHairResourceUsageType::Static);
+	InternalCreateByteAddressBufferRDG(GraphBuilder, BulkData.Interpolation, FHairCardsInterpolationFormat::Format, InterpolationBuffer, ToHairResourceDebugName(TEXT("Hair.CardsInterpolation_InterpolationBuffer"), ResourceName), OwnerName, EHairResourceUsageType::Static);
 }
 
 void FHairCardsInterpolationResource::InternalRelease()
