Fix missing strands shadow when groom instance are frustum culled.

* Split strands and cards interpolation in two different passes
* Move strands interpolation later in the frame, once shadow task is completed.
* Renable out of frustum strands shadow casting

diff --git a/Engine/Plugins/Runtime/HairStrands/Source/HairStrandsCore/Private/GroomDebug.cpp b/Engine/Plugins/Runtime/HairStrands/Source/HairStrandsCore/Private/GroomDebug.cpp
--- a/Engine/Plugins/Runtime/HairStrands/Source/HairStrandsCore/Private/GroomDebug.cpp
+++ b/Engine/Plugins/Runtime/HairStrands/Source/HairStrandsCore/Private/GroomDebug.cpp
@@ -1328,7 +1328,7 @@ void RunHairStrandsDebug(
 				RenderMeshProjection(GraphBuilder, EHairStrandsProjectionMeshType::TargetMesh);
 			}
 
-			auto RenderProjectionData = [&GraphBuilder, ShaderMap, Viewport, &ViewUniformBuffer, Instances, ShaderPrintData](EHairStrandsInterpolationType StrandType, bool bRestTriangle, bool bRestFrame, bool bRestSamples, bool bDeformedTriangle, bool bDeformedFrame, bool bDeformedSamples)
+			auto RenderProjectionData = [&GraphBuilder, ShaderMap, Viewport, &ViewUniformBuffer, Instances, ShaderPrintData](bool bGuide, bool bRestTriangle, bool bRestFrame, bool bRestSamples, bool bDeformedTriangle, bool bDeformedFrame, bool bDeformedSamples)
 			{
 				TArray<int32> HairLODIndices;
 				for (FHairStrandsInstance* AbstractInstance : Instances)
@@ -1337,11 +1337,9 @@ void RunHairStrandsDebug(
 					if (!Instance->HairGroupPublicData || Instance->BindingType != EHairBindingType::Skinning)
 						continue;
 
-					const bool bRenderStrands = StrandType == EHairStrandsInterpolationType::RenderStrands;
-
 					FHairStrandsRestRootResource* RestRootResource = nullptr;
 					FHairStrandsDeformedRootResource* DeformedRootResource = nullptr;
-					if (!bRenderStrands)
+					if (bGuide)
 					{
 						RestRootResource 	 = Instance->Guides.RestRootResource;
 						DeformedRootResource = Instance->Guides.DeformedRootResource;
@@ -1382,7 +1380,7 @@ void RunHairStrandsDebug(
 				GHairDebugMeshProjection_Render_HairRestSamples > 0)
 			{
 				RenderProjectionData(
-					EHairStrandsInterpolationType::RenderStrands,
+					false,
 					GHairDebugMeshProjection_Render_HairRestTriangles > 0,
 					GHairDebugMeshProjection_Render_HairRestFrames > 0,
 					GHairDebugMeshProjection_Render_HairRestSamples > 0,
@@ -1399,7 +1397,7 @@ void RunHairStrandsDebug(
 				GHairDebugMeshProjection_Sim_HairRestSamples > 0)
 			{
 				RenderProjectionData(
-					EHairStrandsInterpolationType::SimulationStrands,
+					true,
 					GHairDebugMeshProjection_Sim_HairRestTriangles > 0,
 					GHairDebugMeshProjection_Sim_HairRestFrames > 0,
 					GHairDebugMeshProjection_Sim_HairRestSamples > 0,
diff --git a/Engine/Plugins/Runtime/HairStrands/Source/HairStrandsCore/Private/GroomManager.cpp b/Engine/Plugins/Runtime/HairStrands/Source/HairStrandsCore/Private/GroomManager.cpp
--- a/Engine/Plugins/Runtime/HairStrands/Source/HairStrandsCore/Private/GroomManager.cpp
+++ b/Engine/Plugins/Runtime/HairStrands/Source/HairStrandsCore/Private/GroomManager.cpp
@@ -215,7 +215,14 @@ static int32 GetCacheGeometryLODIndex(const FCachedGeometry& In)
 	return MeshLODIndex;
 }
 
-static void RunInternalHairStrandsInterpolation(
+enum class EHairInterpolationPassType
+{
+	Strands,
+	CardsAndMeshes,
+	Guides
+};
+
+static void RunInternalHairInterpolation(
 	FRDGBuilder& GraphBuilder,
 	FSceneInterface* Scene,
 	const FSceneView* View,
@@ -223,7 +230,7 @@ static void RunInternalHairStrandsInterpolation(
 	const FHairStrandsInstances& Instances,
 	const FShaderPrintData* ShaderPrintData,
 	FGlobalShaderMap* ShaderMap, 
-	EHairStrandsInterpolationType Type,
+	EHairInterpolationPassType PassType,
 	FHairStrandClusterData* ClusterData)
 {
 	check(IsInRenderingThread());
@@ -272,7 +279,7 @@ static void RunInternalHairStrandsInterpolation(
 			const bool bGlobalDeformationEnable = Instance->HairGroupPublicData->IsGlobalInterpolationEnable(HairLODIndex);
 			check(InstanceBindingType == BindingType);
 
-			if (EHairStrandsInterpolationType::RenderStrands == Type)
+			if (EHairInterpolationPassType::Strands == PassType)
 			{
 				if (InstanceGeometryType == EHairGeometryType::Strands)
 				{
@@ -282,12 +289,12 @@ static void RunInternalHairStrandsInterpolation(
 						check(Instance->Strands.DeformedRootResource->IsValid(MeshLODIndex));
 
 						AddHairStrandUpdateMeshTrianglesPass(
-							GraphBuilder, 
-							ShaderMap, 
-							MeshLODIndex, 
-							HairStrandsTriangleType::DeformedPose, 
-							MeshDataLOD, 
-							Instance->Strands.RestRootResource, 
+							GraphBuilder,
+							ShaderMap,
+							MeshLODIndex,
+							HairStrandsTriangleType::DeformedPose,
+							MeshDataLOD,
+							Instance->Strands.RestRootResource,
 							Instance->Strands.DeformedRootResource);
 
 						AddHairStrandUpdatePositionOffsetPass(
@@ -309,7 +316,10 @@ static void RunInternalHairStrandsInterpolation(
 							Instance->Strands.DeformedResource);
 					}
 				}
-				else if (InstanceGeometryType == EHairGeometryType::Cards)
+			}
+			else if (EHairInterpolationPassType::CardsAndMeshes == PassType)
+			{
+				if (InstanceGeometryType == EHairGeometryType::Cards)
 				{
 					if (Instance->Cards.IsValid(HairLODIndex))
 					{
@@ -353,7 +363,7 @@ static void RunInternalHairStrandsInterpolation(
 					// Nothing to do
 				}
 			}
-			else if (EHairStrandsInterpolationType::SimulationStrands == Type)
+			else if (EHairInterpolationPassType::Guides == PassType)
 			{
 				// Guide update need to run only if simulation is enabled, or if RBF is enabled (since RFB are transfer through guides)
 				if (bGlobalDeformationEnable || bSimulationEnable || bDeformationEnable)
@@ -435,7 +445,7 @@ static void RunInternalHairStrandsInterpolation(
 	}
 
 	// Reset deformation
-	if (EHairStrandsInterpolationType::SimulationStrands == Type)
+	if (EHairInterpolationPassType::Guides == PassType)
 	{
 		for (FHairStrandsInstance* AbstractInstance : Instances)
 		{
@@ -449,7 +459,7 @@ static void RunInternalHairStrandsInterpolation(
 	}
 
 	// Hair interpolation
-	if (EHairStrandsInterpolationType::RenderStrands == Type)
+	if (EHairInterpolationPassType::Strands == PassType || EHairInterpolationPassType::CardsAndMeshes == PassType)
 	{
 		check(View);
 		const FVector& TranslatedWorldOffset = View->ViewMatrices.GetPreViewTranslation();
@@ -457,20 +467,23 @@ static void RunInternalHairStrandsInterpolation(
 		{
 			FHairGroupInstance* Instance = static_cast<FHairGroupInstance*>(AbstractInstance);
 
-			if (Instance->GeometryType == EHairGeometryType::NoneGeometry)
-				continue;
-
-			ComputeHairStrandsInterpolation(
-				GraphBuilder, 
-				ShaderMap,
-				ViewUniqueID,
-				ViewRayTracingMask,
-				ViewMode,
-				TranslatedWorldOffset,
-				ShaderPrintData,
-				Instance,
-				Instance->Debug.MeshLODIndex,
-				ClusterData);
+			const bool bGeometryCompatibleWithPassType = 
+				(EHairInterpolationPassType::Strands == PassType && Instance->GeometryType == EHairGeometryType::Strands) ||
+				(EHairInterpolationPassType::CardsAndMeshes == PassType && (Instance->GeometryType == EHairGeometryType::Cards || Instance->GeometryType == EHairGeometryType::Meshes));
+			if (bGeometryCompatibleWithPassType)
+			{
+				ComputeHairStrandsInterpolation(
+					GraphBuilder, 
+					ShaderMap,
+					ViewUniqueID,
+					ViewRayTracingMask,
+					ViewMode,
+					TranslatedWorldOffset,
+					ShaderPrintData,
+					Instance,
+					Instance->Debug.MeshLODIndex,
+					ClusterData);
+			}
 		}
 	}
 }
@@ -490,7 +503,7 @@ static void RunHairStrandsInterpolation_Guide(
 	RDG_EVENT_SCOPE(GraphBuilder, "HairGuideInterpolation");
 	RDG_GPU_STAT_SCOPE(GraphBuilder, HairGuideInterpolation);
 
-	RunInternalHairStrandsInterpolation(
+	RunInternalHairInterpolation(
 		GraphBuilder,
 		Scene,
 		View,
@@ -498,7 +511,7 @@ static void RunHairStrandsInterpolation_Guide(
 		Instances,
 		ShaderPrintData,
 		ShaderMap,
-		EHairStrandsInterpolationType::SimulationStrands,
+		EHairInterpolationPassType::Guides,
 		nullptr);
 }
 
@@ -554,7 +567,7 @@ static void RunHairStrandsInterpolation_Strands(
 		}
 	}
 
-	RunInternalHairStrandsInterpolation(
+	RunInternalHairInterpolation(
 		GraphBuilder,
 		Scene,
 		View,
@@ -562,10 +575,38 @@ static void RunHairStrandsInterpolation_Strands(
 		Instances,
 		ShaderPrintData,
 		ShaderMap,
-		EHairStrandsInterpolationType::RenderStrands,
+		EHairInterpolationPassType::Strands,
 		&ClusterData);
 }
 
+static void RunHairStrandsInterpolation_Cards(
+	FRDGBuilder& GraphBuilder,
+	FSceneInterface* Scene,
+	const TArray<const FSceneView*>& Views, 
+	const FSceneView* View,
+	const uint32 ViewUniqueID,
+	const FHairStrandsInstances& Instances,
+	const FShaderPrintData* ShaderPrintData,
+	FGlobalShaderMap* ShaderMap)
+{
+	check(IsInRenderingThread());
+
+	DECLARE_GPU_STAT(HairCardsInterpolation);
+	RDG_EVENT_SCOPE(GraphBuilder, "HairCardsInterpolation");
+	RDG_GPU_STAT_SCOPE(GraphBuilder, HairCardsInterpolation);
+
+	RunInternalHairInterpolation(
+		GraphBuilder,
+		Scene,
+		View,
+		ViewUniqueID,
+		Instances,
+		ShaderPrintData,
+		ShaderMap,
+		EHairInterpolationPassType::CardsAndMeshes,
+		nullptr);
+}
+
 // Return the LOD which should be used for a given screen size and LOD bias value
 // This function is mirrored in HairStrandsClusterCommon.ush
 static float GetHairInstanceLODIndex(const TArray<float>& InLODScreenSizes, float InScreenSize, float InLODBias)
@@ -1259,7 +1300,7 @@ void ProcessHairStrandsBookmark(
 
 	const bool bCulling =
 		IsInstanceFrustumCullingEnable() && (
-			Bookmark == EHairStrandsBookmark::ProcessGatherCluster ||
+			Bookmark == EHairStrandsBookmark::ProcessCardsAndMeshesInterpolation ||
 			Bookmark == EHairStrandsBookmark::ProcessStrandsInterpolation ||
 			Bookmark == EHairStrandsBookmark::ProcessDebug);
 	FHairStrandsInstances& Instances = bCulling ? Parameters.VisibleInstances : *Parameters.Instances;
@@ -1317,6 +1358,19 @@ void ProcessHairStrandsBookmark(
 			Parameters.ShaderPrintData,
 			Parameters.ShaderMap);
 	}
+	else if (Bookmark == EHairStrandsBookmark::ProcessCardsAndMeshesInterpolation)
+	{
+		check(GraphBuilder);
+		RunHairStrandsInterpolation_Cards(
+			*GraphBuilder,
+			Parameters.Scene,
+			Parameters.AllViews,
+			Parameters.View,
+			Parameters.ViewUniqueID,
+			Instances,
+			Parameters.ShaderPrintData,
+			Parameters.ShaderMap);
+	}
 	else if (Bookmark == EHairStrandsBookmark::ProcessStrandsInterpolation)
 	{
 		check(GraphBuilder);
diff --git a/Engine/Plugins/Runtime/HairStrands/Source/HairStrandsCore/Private/GroomManager.h b/Engine/Plugins/Runtime/HairStrands/Source/HairStrandsCore/Private/GroomManager.h
--- a/Engine/Plugins/Runtime/HairStrands/Source/HairStrandsCore/Private/GroomManager.h
+++ b/Engine/Plugins/Runtime/HairStrands/Source/HairStrandsCore/Private/GroomManager.h
@@ -27,15 +27,6 @@ class  FResourceArrayInterface;
 class  FSceneView;
 class  UTexture2D;
 
-////////////////////////////////////////////////////////////////////////////////////////////////////
-// Hair/Mesh projection & interpolation
-
-enum class EHairStrandsInterpolationType
-{
-	RenderStrands,
-	SimulationStrands
-};
-
 ////////////////////////////////////////////////////////////////////////////////////////////////////
 // Mesh transfer and hair projection debug infos
 
diff --git a/Engine/Source/Runtime/Renderer/Private/DeferredShadingRenderer.cpp b/Engine/Source/Runtime/Renderer/Private/DeferredShadingRenderer.cpp
--- a/Engine/Source/Runtime/Renderer/Private/DeferredShadingRenderer.cpp
+++ b/Engine/Source/Runtime/Renderer/Private/DeferredShadingRenderer.cpp
@@ -3019,7 +3019,7 @@ void FDeferredShadingSceneRenderer::Render(FRDGBuilder& GraphBuilder)
 		const bool bRunHairStrands = HairStrandsBookmarkParameters.HasInstances() && (Views.Num() > 0);
 		if (bRunHairStrands)
 		{
-			RunHairStrandsBookmark(GraphBuilder, EHairStrandsBookmark::ProcessStrandsInterpolation, HairStrandsBookmarkParameters);
+			RunHairStrandsBookmark(GraphBuilder, EHairStrandsBookmark::ProcessCardsAndMeshesInterpolation, HairStrandsBookmarkParameters);
 		}
 		else
 		{
@@ -3477,10 +3477,14 @@ void FDeferredShadingSceneRenderer::Render(FRDGBuilder& GraphBuilder)
 			RenderShadowDepthMaps(GraphBuilder, InstanceCullingManager, ExternalAccessQueue);
 			bShadowMapsRenderedEarly = true;
 
-			if (bHairStrandsEnable && !bHasRayTracedOverlay)
+			if (bHairStrandsEnable)
 			{
-				RenderHairPrePass(GraphBuilder, Scene, Views, InstanceCullingManager);
-				RenderHairBasePass(GraphBuilder, Scene, SceneTextures, Views, InstanceCullingManager);
+				RunHairStrandsBookmark(GraphBuilder, EHairStrandsBookmark::ProcessStrandsInterpolation, HairStrandsBookmarkParameters);
+				if (!bHasRayTracedOverlay)
+				{
+					RenderHairPrePass(GraphBuilder, Scene, Views, InstanceCullingManager);
+					RenderHairBasePass(GraphBuilder, Scene, SceneTextures, Views, InstanceCullingManager);
+				}
 			}
 
 			RenderForwardShadowProjections(GraphBuilder, SceneTextures, ForwardScreenSpaceShadowMaskTexture, ForwardScreenSpaceShadowMaskHairTexture);
@@ -3644,10 +3648,14 @@ void FDeferredShadingSceneRenderer::Render(FRDGBuilder& GraphBuilder)
 		}
 
 		// Render hair
-		if (bHairStrandsEnable && !IsForwardShadingEnabled(ShaderPlatform) && !bHasRayTracedOverlay)
+		if (bHairStrandsEnable && !IsForwardShadingEnabled(ShaderPlatform))
 		{
-			RenderHairPrePass(GraphBuilder, Scene, Views, InstanceCullingManager);
-			RenderHairBasePass(GraphBuilder, Scene, SceneTextures, Views, InstanceCullingManager);
+			RunHairStrandsBookmark(GraphBuilder, EHairStrandsBookmark::ProcessStrandsInterpolation, HairStrandsBookmarkParameters);
+			if (!bHasRayTracedOverlay)
+			{
+				RenderHairPrePass(GraphBuilder, Scene, Views, InstanceCullingManager);
+				RenderHairBasePass(GraphBuilder, Scene, SceneTextures, Views, InstanceCullingManager);
+			}
 		}
 
 		// Post base pass for material classification
diff --git a/Engine/Source/Runtime/Renderer/Private/HairStrands/HairStrandsInterface.cpp b/Engine/Source/Runtime/Renderer/Private/HairStrands/HairStrandsInterface.cpp
--- a/Engine/Source/Runtime/Renderer/Private/HairStrands/HairStrandsInterface.cpp
+++ b/Engine/Source/Runtime/Renderer/Private/HairStrands/HairStrandsInterface.cpp
@@ -53,7 +53,7 @@ static TAutoConsoleVariable<int32> CVarHairStrandsSimulation(
 	ECVF_RenderThreadSafe | ECVF_Scalability);
 
 static TAutoConsoleVariable<int32> CVarHairStrandsNonVisibleShadowCasting(
-	TEXT("r.HairStrands.Shadow.CastShadowWhenNonVisible"), 0,
+	TEXT("r.HairStrands.Shadow.CastShadowWhenNonVisible"), 1,
 	TEXT("Enable shadow casting for hair strands even when culled out from the primary view"),
 	ECVF_RenderThreadSafe);
 
diff --git a/Engine/Source/Runtime/Renderer/Public/HairStrandsInterface.h b/Engine/Source/Runtime/Renderer/Public/HairStrandsInterface.h
--- a/Engine/Source/Runtime/Renderer/Public/HairStrandsInterface.h
+++ b/Engine/Source/Runtime/Renderer/Public/HairStrandsInterface.h
@@ -477,7 +477,7 @@ enum class EHairStrandsBookmark : uint8
 	ProcessTasks,
 	ProcessLODSelection,
 	ProcessGuideInterpolation,
-	ProcessGatherCluster,
+	ProcessCardsAndMeshesInterpolation,
 	ProcessStrandsInterpolation,
 	ProcessDebug,
 	ProcessEndOfFrame
