Fix crash when exporting MH.

The unique triangle count was incorrectly initialized.

diff --git a/Engine/Plugins/Runtime/HairStrands/Source/HairStrandsCore/Private/GroomRBFDeformer.cpp b/Engine/Plugins/Runtime/HairStrands/Source/HairStrandsCore/Private/GroomRBFDeformer.cpp
--- a/Engine/Plugins/Runtime/HairStrands/Source/HairStrandsCore/Private/GroomRBFDeformer.cpp
+++ b/Engine/Plugins/Runtime/HairStrands/Source/HairStrandsCore/Private/GroomRBFDeformer.cpp
@@ -157,10 +157,10 @@ void DeformStrands(
 void ExtractUniqueTrianglePositions(
 	const FHairStrandsRootData::FMeshProjectionLOD& RestLODData,
 	const uint32 MeshLODIndex,
-	const uint32 UniqueTriangleCount,
 	const FSkeletalMeshRenderData* InMeshRenderData, 
 	TArray<FHairStrandsMeshTrianglePositionFormat::Type>& OutDeformUniqueTrianglePositionBuffer)
 {
+	const uint32 UniqueTriangleCount = RestLODData.UniqueTriangleIndexBuffer.Num();
 	OutDeformUniqueTrianglePositionBuffer.SetNum(UniqueTriangleCount * 3);
 
 	const uint32 SectionCount = InMeshRenderData->LODRenderData[MeshLODIndex].RenderSections.Num();
@@ -230,7 +230,6 @@ TArray<FVector3f> GetDeformedHairStrandsPositions(
 	ExtractUniqueTrianglePositions(
 		RestLODData,
 		MeshLODIndex,
-		RestLODData.RestUniqueTrianglePositionBuffer.Num(),
 		InMeshRenderData,
 		UniqueTrianglePositionBuffer_Deformed);
 
