Fix incorrect motion vector on hair cards.

diff --git a/Engine/Plugins/Runtime/HairStrands/Source/HairStrandsCore/Private/HairStrandsInterpolation.cpp b/Engine/Plugins/Runtime/HairStrands/Source/HairStrandsCore/Private/HairStrandsInterpolation.cpp
--- a/Engine/Plugins/Runtime/HairStrands/Source/HairStrandsCore/Private/HairStrandsInterpolation.cpp
+++ b/Engine/Plugins/Runtime/HairStrands/Source/HairStrandsCore/Private/HairStrandsInterpolation.cpp
@@ -1065,7 +1065,7 @@ static void AddHairCardsDeformationPass(
 		DispatchCount);
 
 	// If LOD has switched, copy the current buffer, so that we don't get incorrect motion vector
-	//if (bHasLODSwitch)
+	if (bHasLODSwitch)
 	{
 		AddCopyBufferPass(GraphBuilder, CardsDeformedPositionBuffer_Prev.Buffer, CardsDeformedPositionBuffer_Curr.Buffer);
 	}
