Removed hair strands deferred delete logic in favor of reference counting.

diff --git a/Engine/Plugins/Runtime/HairStrands/Source/HairStrandsCore/Private/GroomComponent.cpp b/Engine/Plugins/Runtime/HairStrands/Source/HairStrandsCore/Private/GroomComponent.cpp
--- a/Engine/Plugins/Runtime/HairStrands/Source/HairStrandsCore/Private/GroomComponent.cpp
+++ b/Engine/Plugins/Runtime/HairStrands/Source/HairStrandsCore/Private/GroomComponent.cpp
@@ -79,6 +79,94 @@ static FAutoConsoleVariableRef CVarHairStrands_ViewModeClumpIndex(TEXT("r.HairSt
 
 ////////////////////////////////////////////////////////////////////////////////////////////////////////////////////////
 
+template<typename T>
+void InternalResourceRelease(T*& In)
+{
+	if (In)
+	{
+		In->ReleaseResource();
+		delete In;
+		In = nullptr;
+	}
+}
+
+FHairGroupInstance::~FHairGroupInstance()
+{
+	// Guides
+	if (Guides.IsValid())
+	{
+		InternalResourceRelease(Guides.DeformedRootResource);
+		InternalResourceRelease(Guides.DeformedResource);
+	}
+
+	// Strands
+	if (Strands.IsValid())
+	{
+		InternalResourceRelease(Strands.DeformedRootResource);
+		InternalResourceRelease(Strands.DeformedResource);
+		InternalResourceRelease(Strands.CullingResource);
+
+#if RHI_RAYTRACING
+		if (Strands.RenRaytracingResourceOwned)
+		{
+			InternalResourceRelease(Strands.RenRaytracingResource);
+		}
+#endif
+
+#if WITH_EDITOR
+		Strands.DebugCurveAttributeBuffer.Release();
+#endif
+
+		InternalResourceRelease(Strands.VertexFactory);
+	}
+
+	// Cards
+	{
+		const uint32 CardLODCount = Cards.LODs.Num();
+		for (uint32 CardLODIt = 0; CardLODIt < CardLODCount; ++CardLODIt)
+		{
+			if (Cards.IsValid(CardLODIt))
+			{
+				InternalResourceRelease(Cards.LODs[CardLODIt].Guides.DeformedRootResource);
+				InternalResourceRelease(Cards.LODs[CardLODIt].Guides.DeformedResource);
+				InternalResourceRelease(Cards.LODs[CardLODIt].DeformedResource);
+#if RHI_RAYTRACING
+				if (Cards.LODs[CardLODIt].RaytracingResourceOwned)
+				{
+					InternalResourceRelease(Cards.LODs[CardLODIt].RaytracingResource);
+				}
+#endif
+
+				InternalResourceRelease(Cards.LODs[CardLODIt].VertexFactory);
+			}
+		}
+	}
+
+	// Meshes
+	{
+		const uint32 MeshesLODCount = Meshes.LODs.Num();
+		for (uint32 MeshesLODIt = 0; MeshesLODIt < MeshesLODCount; ++MeshesLODIt)
+		{
+			if (Meshes.IsValid(MeshesLODIt))
+			{
+				InternalResourceRelease(Meshes.LODs[MeshesLODIt].DeformedResource);
+#if RHI_RAYTRACING
+				if (Meshes.LODs[MeshesLODIt].RaytracingResourceOwned)
+				{
+					InternalResourceRelease(Meshes.LODs[MeshesLODIt].RaytracingResource);
+				}
+#endif
+
+				InternalResourceRelease(Meshes.LODs[MeshesLODIt].VertexFactory);
+			}
+		}
+	}
+
+	delete HairGroupPublicData;
+}
+
+////////////////////////////////////////////////////////////////////////////////////////////////////////////////////////
+
 const FLinearColor GetHairGroupDebugColor(int32 GroupIt)
 {
 	static TArray<FLinearColor> IndexedColor;
@@ -350,19 +438,19 @@ public:
 			bAlwaysHasVelocity = true;
 		}
 
-		check(Component->HairGroupInstances.Num());
+		check(HairGroupInstances.Num());
 
 		const ERHIFeatureLevel::Type FeatureLevel = GetScene().GetFeatureLevel();
 		const EShaderPlatform ShaderPlatform = GetScene().GetShaderPlatform();
 
 		const int32 GroupCount = Component->GroomAsset->GetNumHairGroups();
-		check(Component->GroomAsset->GetHairGroupsPlatformData().Num() == Component->HairGroupInstances.Num());
+		check(Component->GroomAsset->GetHairGroupsPlatformData().Num() == HairGroupInstances.Num());
 		for (int32 GroupIt=0; GroupIt<GroupCount; GroupIt++)
 		{
 			const bool bIsVisible = Component->GroomAsset->GetHairGroupsInfo()[GroupIt].bIsVisible;
 
 			const FHairGroupPlatformData& InGroupData = Component->GroomAsset->GetHairGroupsPlatformData()[GroupIt];
-			FHairGroupInstance* HairInstance = Component->HairGroupInstances[GroupIt];
+			FHairGroupInstance* HairInstance = HairGroupInstances[GroupIt];
 			check(HairInstance->HairGroupPublicData);
 			HairInstance->bForceCards = Component->bUseCards;
 			HairInstance->bUpdatePositionOffset = Component->RegisteredMeshComponent != nullptr;
@@ -482,8 +570,7 @@ public:
 
 		// Register the data to the scene
 		FSceneInterface& LocalScene = GetScene();
-		TArray<FHairGroupInstance*> LocalInstances = HairGroupInstances;
-		for (FHairGroupInstance* Instance : LocalInstances)
+		for (TRefCountPtr<FHairGroupInstance>& Instance : HairGroupInstances)
 		{
 			if (Instance->IsValid() || Instance->Strands.ClusterResource)
 			{
@@ -506,8 +593,7 @@ public:
 
 		// Unregister the data to the scene
 		FSceneInterface& LocalScene = GetScene();
-		TArray<FHairGroupInstance*> LocalInstances = HairGroupInstances;
-		for (FHairGroupInstance* Instance : LocalInstances)
+		for (TRefCountPtr<FHairGroupInstance>& Instance : HairGroupInstances)
 		{
 			if (Instance->IsValid() || Instance->Strands.ClusterResource)
 			{
@@ -522,7 +608,7 @@ public:
 	virtual void OnTransformChanged(FRHICommandListBase& RHICmdList) override
 	{
 		const FTransform RigidLocalToWorld = FTransform(GetLocalToWorld());
-		for (FHairGroupInstance* Instance : HairGroupInstances)
+		for (TRefCountPtr<FHairGroupInstance>& Instance : HairGroupInstances)
 		{
 			Instance->Debug.RigidPreviousLocalToWorld = Instance->Debug.RigidCurrentLocalToWorld;
 			Instance->Debug.RigidCurrentLocalToWorld = RigidLocalToWorld;
@@ -628,13 +714,12 @@ public:
 			return;
 		}
 
-		TArray<FHairGroupInstance*> Instances = HairGroupInstances;
-		if (Instances.Num() == 0)
+		if (HairGroupInstances.Num() == 0)
 		{
 			return;
 		}
 
-		const uint32 GroupCount = Instances.Num();
+		const uint32 GroupCount = HairGroupInstances.Num();
 
 		QUICK_SCOPE_CYCLE_COUNTER(STAT_HairStrandsSceneProxy_GetDynamicMeshElements);
 
@@ -651,7 +736,7 @@ public:
 			{
 				for (uint32 GroupIt = 0; GroupIt < GroupCount; ++GroupIt)
 				{
-					check(Instances[GroupIt]->GetRefCount() > 0);
+					check(HairGroupInstances[GroupIt]->GetRefCount() > 0);
 
 					FMaterialRenderProxy* Debug_MaterialProxy = nullptr;
 					const EGroomViewMode ViewMode = AllowDebugViewmodes() ? GetGroomViewMode(*View) : EGroomViewMode::None;
@@ -662,7 +747,7 @@ public:
 					case EGroomViewMode::Cluster		:
 					case EGroomViewMode::ClusterAABB	:
 					case EGroomViewMode::RenderHairStrands:
-						bNeedDebugMaterial = Instances[GroupIt]->GeometryType == EHairGeometryType::Strands; break;
+						bNeedDebugMaterial = HairGroupInstances[GroupIt]->GeometryType == EHairGeometryType::Strands; break;
 					case EGroomViewMode::RootUV			:
 					case EGroomViewMode::UV				:
 					case EGroomViewMode::Seed			:
@@ -704,7 +789,7 @@ public:
 
 						// TODO: fix this as the radius is incorrect. This code run before the interpolation code, which is where HairRadius is updated.
 						float HairMaxRadius = 0;
-						for (FHairGroupInstance* Instance : Instances)
+						for (FHairGroupInstance* Instance : HairGroupInstances)
 						{
 							HairMaxRadius = FMath::Max(HairMaxRadius, Instance->Strands.Modifier.HairWidth * 0.5f);
 						}
@@ -712,7 +797,7 @@ public:
 						// Reuse the HairMaxRadius field to send the LOD index instead of adding yet another variable
 						if (ViewMode == EGroomViewMode::LODColoration)
 						{
-							HairMaxRadius = Instances[GroupIt]->HairGroupPublicData ? Instances[GroupIt]->HairGroupPublicData->LODIndex : 0;
+							HairMaxRadius = HairGroupInstances[GroupIt]->HairGroupPublicData ? HairGroupInstances[GroupIt]->HairGroupPublicData->LODIndex : 0;
 						}
 
 						// Reuse the HairMaxRadius field to send the clumpID index selection
@@ -728,7 +813,7 @@ public:
 						}
 						else if (ViewMode == EGroomViewMode::LODColoration)
 						{
-							int32 LODIndex = Instances[GroupIt]->HairGroupPublicData ? Instances[GroupIt]->HairGroupPublicData->LODIndex : 0;
+							int32 LODIndex = HairGroupInstances[GroupIt]->HairGroupPublicData ? HairGroupInstances[GroupIt]->HairGroupPublicData->LODIndex : 0;
 							LODIndex = FMath::Clamp(LODIndex, 0, GEngine->LODColorationColors.Num() - 1);
 							const FLinearColor LODColor = GEngine->LODColorationColors[LODIndex];
 							HairColor = FVector(LODColor.R, LODColor.G, LODColor.B);
@@ -738,7 +823,7 @@ public:
 						Debug_MaterialProxy = DebugMaterial;
 					}
 
-					if (FMeshBatch* MeshBatch = CreateMeshBatch(View, ViewFamily, Collector, EHairMeshBatchType::Raster, Instances[GroupIt], GroupIt, Debug_MaterialProxy))
+					if (FMeshBatch* MeshBatch = CreateMeshBatch(View, ViewFamily, Collector, EHairMeshBatchType::Raster, HairGroupInstances[GroupIt], GroupIt, Debug_MaterialProxy))
 					{
 						Collector.AddMesh(ViewIndex, *MeshBatch);
 					}
@@ -941,7 +1026,7 @@ public:
 		const bool bForceDrawRelevance = bPathtracing && (!IsShown(View) && IsShadowCast(View));
 
 		bool bUseCardsOrMesh = false;
-		for (FHairGroupInstance* Instance : HairGroupInstances)
+		for (const TRefCountPtr<FHairGroupInstance>& Instance : HairGroupInstances)
 		{
 			check(Instance->GetRefCount());
 			const EHairGeometryType GeometryType = Instance->GeometryType;
@@ -974,7 +1059,7 @@ public:
 
 	uint32 GetAllocatedSize(void) const { return(FPrimitiveSceneProxy::GetAllocatedSize()); }
 
-	TArray<FHairGroupInstance*> HairGroupInstances;
+	TArray<TRefCountPtr<FHairGroupInstance>> HairGroupInstances;
 
 	// Cache the material proxy to avoid race condition, when groom component's proxy is recreated, 
 	// while another one is currently in flight for drawing.
@@ -1613,7 +1698,7 @@ void UGroomComponent::SetForcedLOD(int32 CurrLODIndex)
 		ENQUEUE_RENDER_COMMAND(FHairComponentSendLODIndex)(
 		[GroomSceneProxy, CurrLODIndex, LocalLODSelectionType, bHasLODSwitch, bPredictionLoad](FRHICommandListImmediate& RHICmdList)
 		{
-			for (FHairGroupInstance* Instance : GroomSceneProxy->HairGroupInstances)
+			for (const TRefCountPtr<FHairGroupInstance>& Instance : GroomSceneProxy->HairGroupInstances)
 			{
 				Instance->Debug.LODForcedIndex = CurrLODIndex;
 				Instance->Debug.LODSelectionTypeForDebug = LocalLODSelectionType;
@@ -1737,7 +1822,7 @@ void UGroomComponent::UpdateHairGroupsDescAndInvalidateRenderState(bool bInvalid
 	UpdateHairGroupsDesc();
 
 	uint32 GroupIndex = 0;
-	for (FHairGroupInstance* Instance : HairGroupInstances)
+	for (const TRefCountPtr<FHairGroupInstance>& Instance : HairGroupInstances)
 	{
 		Instance->Strands.Modifier = GetGroomGroupsDesc(GroomAsset, this, GroupIndex);
 		++GroupIndex;
@@ -1750,13 +1835,11 @@ void UGroomComponent::UpdateHairGroupsDescAndInvalidateRenderState(bool bInvalid
 
 FPrimitiveSceneProxy* UGroomComponent::CreateSceneProxy()
 {
-	DeleteDeferredHairGroupInstances();
-
 	if (!GroomAsset || GroomAsset->GetNumHairGroups() == 0 || HairGroupInstances.Num() == 0)
 		return nullptr;
 
 	bool bIsValid = false;
-	for (FHairGroupInstance* Instance : HairGroupInstances)
+	for (const TRefCountPtr<FHairGroupInstance>& Instance : HairGroupInstances)
 	{
 		bIsValid |= Instance->IsValid();
 	}
@@ -2115,14 +2198,14 @@ void UGroomComponent::UpdateSimulatedGroups()
 		// For now use the force LOD to drive enabling/disabling simulation & RBF
 		const int32 LODIndex = GetForcedLOD();
 
-		TArray<FHairGroupInstance*> LocalInstances = HairGroupInstances;
+		TArray<TRefCountPtr<FHairGroupInstance>> LocalInstances = HairGroupInstances;
 		UGroomAsset* LocalGroomAsset = GroomAsset;
 		UGroomBindingAsset* LocalBindingAsset = BindingAsset;
 		ENQUEUE_RENDER_COMMAND(FHairStrandsTick_UEnableSimulatedGroups)(
 			[LocalInstances, LocalGroomAsset, LocalBindingAsset, Id, bIsStrandsEnabled, LODIndex](FRHICommandListImmediate& RHICmdList)
 		{
 			int32 GroupIt = 0;
-			for (FHairGroupInstance* Instance : LocalInstances)
+			for (const TRefCountPtr<FHairGroupInstance>& Instance : LocalInstances)
 			{
 				Instance->Strands.HairInterpolationType = EHairInterpolationType::NoneSkinning;
 				if (bIsStrandsEnabled && LocalGroomAsset)
@@ -2519,7 +2602,6 @@ void UGroomComponent::InitResources(bool bIsBindingReloading)
 	for (int32 GroupIt = 0, GroupCount = GroomAsset->GetHairGroupsPlatformData().Num(); GroupIt < GroupCount; ++GroupIt)
 	{
 		FHairGroupInstance* HairGroupInstance = new FHairGroupInstance();
-		HairGroupInstance->AddRef();
 		HairGroupInstances.Add(HairGroupInstance);
 		HairGroupInstance->Debug.GroupIndex = GroupIt;
 		HairGroupInstance->Debug.GroupCount = GroupCount;
@@ -2917,25 +2999,11 @@ void UGroomComponent::InitResources(bool bIsBindingReloading)
 	}
 }
 
-template<typename T>
-void InternalResourceRelease(T*& In)
-{
-	if (In)
-	{
-		In->ReleaseResource();
-		delete In;
-		In = nullptr;
-	}
-}
-
 void UGroomComponent::ReleaseResources()
 {
-
 	FHairStrandsSceneProxy* GroomSceneProxy = (FHairStrandsSceneProxy*)SceneProxy;
 	InitializedResources = nullptr;
 
-	// Deferring instances deletion to insure scene proxy are done with the rendering data
-	DeferredDeleteHairGroupInstances.Append(HairGroupInstances);
 	HairGroupInstances.Empty();
 
 	// Insure the ticking of the Groom component always happens after the skeletalMeshComponent.
@@ -2951,96 +3019,6 @@ void UGroomComponent::ReleaseResources()
 	MarkRenderStateDirty();
 }
 
-void UGroomComponent::DeleteDeferredHairGroupInstances()
-{
-	FHairStrandsSceneProxy* GroomSceneProxy = (FHairStrandsSceneProxy*)SceneProxy;
-	for (FHairGroupInstance* Instance : DeferredDeleteHairGroupInstances)
-	{
-		FHairGroupInstance* LocalInstance = Instance;
-		ENQUEUE_RENDER_COMMAND(FHairStrandsBuffers)(
-		[LocalInstance, GroomSceneProxy](FRHICommandListImmediate& RHICmdList)
-		{
-			// Sanity check
-			check(LocalInstance->GetRefCount() == 1);
-
-			// Guides
-			if (LocalInstance->Guides.IsValid())
-			{
-				InternalResourceRelease(LocalInstance->Guides.DeformedRootResource);
-				InternalResourceRelease(LocalInstance->Guides.DeformedResource);
-			}
-
-			// Strands
-			if (LocalInstance->Strands.IsValid())
-			{
-				InternalResourceRelease(LocalInstance->Strands.DeformedRootResource);
-				InternalResourceRelease(LocalInstance->Strands.DeformedResource);
-				InternalResourceRelease(LocalInstance->Strands.CullingResource);
-
-				#if RHI_RAYTRACING
-				if (LocalInstance->Strands.RenRaytracingResourceOwned)
-				{
-					InternalResourceRelease(LocalInstance->Strands.RenRaytracingResource);
-				}
-				#endif
-
-				#if WITH_EDITOR
-				LocalInstance->Strands.DebugCurveAttributeBuffer.Release();
-				#endif
-
-				InternalResourceRelease(LocalInstance->Strands.VertexFactory);
-			}
-
-			// Cards
-			{
-				const uint32 CardLODCount = LocalInstance->Cards.LODs.Num();
-				for (uint32 CardLODIt = 0; CardLODIt < CardLODCount; ++CardLODIt)
-				{
-					if (LocalInstance->Cards.IsValid(CardLODIt))
-					{
-						InternalResourceRelease(LocalInstance->Cards.LODs[CardLODIt].Guides.DeformedRootResource);
-						InternalResourceRelease(LocalInstance->Cards.LODs[CardLODIt].Guides.DeformedResource);
-						InternalResourceRelease(LocalInstance->Cards.LODs[CardLODIt].DeformedResource);
-						#if RHI_RAYTRACING
-						if (LocalInstance->Cards.LODs[CardLODIt].RaytracingResourceOwned)
-						{
-							InternalResourceRelease(LocalInstance->Cards.LODs[CardLODIt].RaytracingResource);
-						}
-						#endif
-
-						InternalResourceRelease(LocalInstance->Cards.LODs[CardLODIt].VertexFactory);
-					}
-				}
-			}
-
-			// Meshes
-			{
-				const uint32 MeshesLODCount = LocalInstance->Meshes.LODs.Num();
-				for (uint32 MeshesLODIt = 0; MeshesLODIt < MeshesLODCount; ++MeshesLODIt)
-				{
-					if (LocalInstance->Meshes.IsValid(MeshesLODIt))
-					{
-						InternalResourceRelease(LocalInstance->Meshes.LODs[MeshesLODIt].DeformedResource);
-						#if RHI_RAYTRACING
-						if (LocalInstance->Meshes.LODs[MeshesLODIt].RaytracingResourceOwned)
-						{
-							InternalResourceRelease(LocalInstance->Meshes.LODs[MeshesLODIt].RaytracingResource);
-						}
-						#endif
-
-						InternalResourceRelease(LocalInstance->Meshes.LODs[MeshesLODIt].VertexFactory);
-					}
-				}
-			}
-
-			delete LocalInstance->HairGroupPublicData;
-			LocalInstance->Release();
-			delete LocalInstance;
-		});
-	}
-	DeferredDeleteHairGroupInstances.Empty();
-}
-
 void UGroomComponent::PostLoad()
 {
 	LLM_SCOPE_BYTAG(Groom);
@@ -3190,14 +3168,12 @@ void UGroomComponent::BeginDestroy()
 
 void UGroomComponent::FinishDestroy()
 {
-	DeleteDeferredHairGroupInstances();
 	Super::FinishDestroy();
 }
 
 void UGroomComponent::OnComponentDestroyed(bool bDestroyingHierarchy)
 {
 	ReleaseResources();
-	DeleteDeferredHairGroupInstances();
 
 #if WITH_EDITOR
 	if (bIsGroomAssetCallbackRegistered && GroomAsset)
@@ -3475,7 +3451,7 @@ void UGroomComponent::TickComponent(float DeltaTime, enum ELevelTick TickType, F
 	if (LODSelectionType == EHairLODSelectionType::Predicted)
 	{
 		// 1. Compute the effective LOD (taking the min. LOD across all the groups)
-		for (const FHairGroupInstance* Instance : HairGroupInstances)
+		for (const TRefCountPtr<FHairGroupInstance>& Instance : HairGroupInstances)
 		{
 			/* /!\ Access rendering thread data at that point /!\ */
 			const float LODPredictedIndex_Instance = Instance->Debug.LODPredictedIndex;
@@ -3502,34 +3478,34 @@ void UGroomComponent::TickComponent(float DeltaTime, enum ELevelTick TickType, F
 	{
 		SkelLocalToTransform = RegisteredMeshComponent->GetComponentTransform();
 	}
-	
-	TArray<FHairGroupInstance*> LocalHairGroupInstances = HairGroupInstances;
-	const EHairLODSelectionType LocalLODSelectionType = LODSelectionType;
-	ENQUEUE_RENDER_COMMAND(FHairStrandsTick_TransformUpdate)(
-		[Id, FeatureLevel, LocalHairGroupInstances, EffectiveForceLOD, LocalLODSelectionType, bSwapBuffer, SkelLocalToTransform](FRHICommandListImmediate& RHICmdList)
-	{
-		if (ERHIFeatureLevel::Num == FeatureLevel)
-			return;
 
-		for (FHairGroupInstance* Instance : LocalHairGroupInstances)
+	if (ERHIFeatureLevel::Num != FeatureLevel)
+	{
+		TArray<TRefCountPtr<FHairGroupInstance>> LocalHairGroupInstances = HairGroupInstances;
+		const EHairLODSelectionType LocalLODSelectionType = LODSelectionType;
+		ENQUEUE_RENDER_COMMAND(FHairStrandsTick_TransformUpdate)(
+			[Id, FeatureLevel, LocalHairGroupInstances, EffectiveForceLOD, LocalLODSelectionType, bSwapBuffer, SkelLocalToTransform](FRHICommandListImmediate& RHICmdList)
 		{
-			Instance->Debug.LODForcedIndex = EffectiveForceLOD;
-			Instance->Debug.LODSelectionTypeForDebug = LocalLODSelectionType;
-			if (LocalLODSelectionType == EHairLODSelectionType::Predicted && EffectiveForceLOD >= 0)
+			for (const TRefCountPtr<FHairGroupInstance>& Instance : LocalHairGroupInstances)
 			{
-				AddHairStreamingRequest(Instance, EffectiveForceLOD);
-			}
+				Instance->Debug.LODForcedIndex = EffectiveForceLOD;
+				Instance->Debug.LODSelectionTypeForDebug = LocalLODSelectionType;
+				if (LocalLODSelectionType == EHairLODSelectionType::Predicted && EffectiveForceLOD >= 0)
+				{
+					AddHairStreamingRequest(Instance, EffectiveForceLOD);
+				}
 
-			if (bSwapBuffer)
-			{
-				Instance->Debug.SkinningPreviousLocalToWorld = Instance->Debug.SkinningCurrentLocalToWorld;
-				Instance->Debug.SkinningCurrentLocalToWorld  = SkelLocalToTransform;
+				if (bSwapBuffer)
+				{
+					Instance->Debug.SkinningPreviousLocalToWorld = Instance->Debug.SkinningCurrentLocalToWorld;
+					Instance->Debug.SkinningCurrentLocalToWorld = SkelLocalToTransform;
 
-				if (Instance->Guides.DeformedResource)  { Instance->Guides.DeformedResource->SwapBuffer(); }
-				if (Instance->Strands.DeformedResource) { Instance->Strands.DeformedResource->SwapBuffer(); }
+					if (Instance->Guides.DeformedResource) { Instance->Guides.DeformedResource->SwapBuffer(); }
+					if (Instance->Strands.DeformedResource) { Instance->Strands.DeformedResource->SwapBuffer(); }
+				}
 			}
-		}
-	});
+		});
+	}
 	
 	if(GetHairSwapBufferType() == EHairBufferSwapType::RenderFrame)
 	{
@@ -3556,11 +3532,11 @@ void UGroomComponent::SendRenderDynamicData_Concurrent()
 
 	if(GetHairSwapBufferType() == EHairBufferSwapType::RenderFrame)
 	{
-		TArray<FHairGroupInstance*> LocalHairGroupInstances = HairGroupInstances;
+		TArray<TRefCountPtr<FHairGroupInstance>> LocalHairGroupInstances = HairGroupInstances;
 		ENQUEUE_RENDER_COMMAND(FHairStrandsTick_TransformUpdate)(
 			[LocalHairGroupInstances](FRHICommandListImmediate& RHICmdList)
 		{
-			for (FHairGroupInstance* Instance : LocalHairGroupInstances)
+			for (const TRefCountPtr<FHairGroupInstance>& Instance : LocalHairGroupInstances)
 			{
 				if (Instance->Guides.DeformedResource)  { Instance->Guides.DeformedResource->SwapBuffer(); }
 				if (Instance->Strands.DeformedResource) { Instance->Strands.DeformedResource->SwapBuffer(); }
@@ -3694,7 +3670,7 @@ void UGroomComponent::PostEditChangeProperty(FPropertyChangedEvent& PropertyChan
 	bool bEnableLengthScaleOverrideChanged = false;
 	if (bEnableLengthScaleChanged && GroomAsset && !bAssetChanged)
 	{
-		for (const FHairGroupInstance* Instance : HairGroupInstances)
+		for (const TRefCountPtr<FHairGroupInstance>& Instance : HairGroupInstances)
 		{
 			const FHairGroupDesc GroupDesc = GetGroomGroupsDesc(GroomAsset, this, Instance->Debug.GroupIndex);
 			const bool bRecreate = Instance->Strands.DeformedResource == nullptr && GroupDesc.HairLengthScale_Override;
@@ -3716,7 +3692,7 @@ void UGroomComponent::PostEditChangeProperty(FPropertyChangedEvent& PropertyChan
 	#if RHI_RAYTRACING
 	if (GroomAsset && !bAssetChanged)
 	{
-		for (const FHairGroupInstance* Instance : HairGroupInstances)
+		for (const TRefCountPtr<FHairGroupInstance>& Instance : HairGroupInstances)
 		{
 			const FHairGroupDesc GroupDesc = GetGroomGroupsDesc(GroomAsset, this, Instance->Debug.GroupIndex);
 			const bool bRecreate =
@@ -3957,7 +3933,7 @@ void InternalAddDedicatedVideoMemoryBytes(FResourceSizeEx& CumulativeResourceSiz
 
 void UGroomComponent::GetResourceSizeEx(FResourceSizeEx& CumulativeResourceSize)
 {
-	for (const FHairGroupInstance* Instance : HairGroupInstances)
+	for (const TRefCountPtr<FHairGroupInstance>& Instance : HairGroupInstances)
 	{
 		InternalAddDedicatedVideoMemoryBytes(CumulativeResourceSize, Instance->Guides.DeformedResource);
 		InternalAddDedicatedVideoMemoryBytes(CumulativeResourceSize, Instance->Guides.DeformedRootResource);
@@ -3980,7 +3956,7 @@ void UGroomComponent::GetResourceSizeEx(FResourceSizeEx& CumulativeResourceSize)
 uint32 UGroomComponent::GetResourcesSize() const
 {
 	FGroomComponentMemoryStats Total;
-	for (const FHairGroupInstance* Instance : HairGroupInstances)
+	for (const TRefCountPtr<FHairGroupInstance>& Instance : HairGroupInstances)
 	{
 		Total.Accumulate(FGroomComponentMemoryStats::Get(Instance));
 	}
diff --git a/Engine/Plugins/Runtime/HairStrands/Source/HairStrandsCore/Public/GroomComponent.h b/Engine/Plugins/Runtime/HairStrands/Source/HairStrandsCore/Public/GroomComponent.h
--- a/Engine/Plugins/Runtime/HairStrands/Source/HairStrandsCore/Public/GroomComponent.h
+++ b/Engine/Plugins/Runtime/HairStrands/Source/HairStrandsCore/Public/GroomComponent.h
@@ -312,8 +312,7 @@ private:
 	TSharedPtr<class IGroomCacheBuffers, ESPMode::ThreadSafe> GroomCacheBuffers;
 
 private:
-	TArray<FHairGroupInstance*> HairGroupInstances;
-	TArray<FHairGroupInstance*> DeferredDeleteHairGroupInstances;
+	TArray<TRefCountPtr<FHairGroupInstance>> HairGroupInstances;
 
 #if WITH_EDITORONLY_DATA
 	UPROPERTY(Transient)
diff --git a/Engine/Plugins/Runtime/HairStrands/Source/HairStrandsCore/Public/GroomInstance.h b/Engine/Plugins/Runtime/HairStrands/Source/HairStrandsCore/Public/GroomInstance.h
--- a/Engine/Plugins/Runtime/HairStrands/Source/HairStrandsCore/Public/GroomInstance.h
+++ b/Engine/Plugins/Runtime/HairStrands/Source/HairStrandsCore/Public/GroomInstance.h
@@ -72,6 +72,8 @@ enum EHairViewRayTracingMask
 // Represent/Describe data & resources of a hair group belonging to a groom
 struct HAIRSTRANDSCORE_API FHairGroupInstance : public FHairStrandsInstance
 {
+	virtual ~FHairGroupInstance();
+
 	//////////////////////////////////////////////////////////////////////////////////////////
 	// Helper struct which aggregate strands based data/resources
 
diff --git a/Engine/Source/Runtime/Renderer/Private/HairStrands/HairStrandsInterface.cpp b/Engine/Source/Runtime/Renderer/Private/HairStrands/HairStrandsInterface.cpp
--- a/Engine/Source/Runtime/Renderer/Private/HairStrands/HairStrandsInterface.cpp
+++ b/Engine/Source/Runtime/Renderer/Private/HairStrands/HairStrandsInterface.cpp
@@ -82,25 +82,6 @@ static TAutoConsoleVariable<int32> CVarHairStrandsVisibilityComputeRaster_Contin
 	TEXT("Enable Continuos LOD when using compute rasterization. Experimental"),
 	ECVF_RenderThreadSafe);
 
-////////////////////////////////////////////////////////////////////////////////////////////////////////////////////////
-// Hair strands instance ref. counting for debug purpose only
-uint32 FHairStrandsInstance::GetRefCount() const
-{
-	return RefCount;
-}
-
-uint32 FHairStrandsInstance::AddRef() const
-{
-	return ++RefCount;
-}
-
-uint32 FHairStrandsInstance::Release() const
-{
-	check(RefCount > 0);
-	uint32 LocalRefCount = --RefCount;
-	return LocalRefCount;
-}
-
 ////////////////////////////////////////////////////////////////////////////////////////////////////////////////////////
 // Import/export utils function for hair resources
 void FRDGExternalBuffer::Release()
diff --git a/Engine/Source/Runtime/Renderer/Public/HairStrandsInterface.h b/Engine/Source/Runtime/Renderer/Public/HairStrandsInterface.h
--- a/Engine/Source/Runtime/Renderer/Public/HairStrandsInterface.h
+++ b/Engine/Source/Runtime/Renderer/Public/HairStrandsInterface.h
@@ -198,19 +198,14 @@ enum EHairInterpolationType
 ////////////////////////////////////////////////////////////////////////////////////////////////////
 // Public group data 
 
-struct FHairStrandsInstance
+struct FHairStrandsInstance : public FRefCountBase
 {
 	virtual ~FHairStrandsInstance() = default;
-	RENDERER_API uint32 GetRefCount() const;
-	RENDERER_API uint32 AddRef() const;
-	RENDERER_API uint32 Release() const;
 	int32 RegisteredIndex = -1;
 	virtual const FBoxSphereBounds& GetBounds() const = 0;
 	virtual const FBoxSphereBounds& GetLocalBounds() const = 0;
 	virtual const FHairGroupPublicData* GetHairData() const { return nullptr; }
 	virtual const EHairGeometryType GetHairGeometry() const { return EHairGeometryType::NoneGeometry; }
-protected:
-	mutable uint32 RefCount = 0;
 };
 typedef TArray<FHairStrandsInstance*> FHairStrandsInstances;
 
