Fix incorrect groom DDC key.

The ddc key didn't containt LOD geometry, causing invalid data fetch when switching a LOD from a geometry type to another.

diff --git a/Engine/Plugins/Runtime/HairStrands/Source/HairStrandsCore/Private/GroomAssetInterpolation.cpp b/Engine/Plugins/Runtime/HairStrands/Source/HairStrandsCore/Private/GroomAssetInterpolation.cpp
--- a/Engine/Plugins/Runtime/HairStrands/Source/HairStrandsCore/Private/GroomAssetInterpolation.cpp
+++ b/Engine/Plugins/Runtime/HairStrands/Source/HairStrandsCore/Private/GroomAssetInterpolation.cpp
@@ -118,6 +118,7 @@ void FHairGroupsLOD::BuildDDCKey(FArchive& Ar)
 		Ar << LOD.ThicknessScale;
 		Ar << LOD.bVisible;
 		Ar << LOD.BindingType;
+		Ar << LOD.GeometryType;
 		Ar << LOD.Simulation;
 		Ar << LOD.GlobalInterpolation;
 	}
diff --git a/Engine/Plugins/Runtime/HairStrands/Source/HairStrandsCore/Private/GroomBuilder.cpp b/Engine/Plugins/Runtime/HairStrands/Source/HairStrandsCore/Private/GroomBuilder.cpp
--- a/Engine/Plugins/Runtime/HairStrands/Source/HairStrandsCore/Private/GroomBuilder.cpp
+++ b/Engine/Plugins/Runtime/HairStrands/Source/HairStrandsCore/Private/GroomBuilder.cpp
@@ -40,7 +40,7 @@ static FAutoConsoleVariableRef CVarHairGroupIndexBuilder_MaxVoxelResolution(TEXT
 
 FString FGroomBuilder::GetVersion()
 {
-	return TEXT("v8r43");
+	return TEXT("v8r45");
 }
 
 namespace FHairStrandsDecimation
