Fix hair cards/meshes reloading when changing assets properties.

diff --git a/Engine/Plugins/Runtime/HairStrands/Source/HairStrandsCore/Private/GroomAsset.cpp b/Engine/Plugins/Runtime/HairStrands/Source/HairStrandsCore/Private/GroomAsset.cpp
--- a/Engine/Plugins/Runtime/HairStrands/Source/HairStrandsCore/Private/GroomAsset.cpp
+++ b/Engine/Plugins/Runtime/HairStrands/Source/HairStrandsCore/Private/GroomAsset.cpp
@@ -2700,12 +2700,6 @@ bool UGroomAsset::CacheCardsData(uint32 GroupIndex, const FString& StrandsKey)
 	const FString KeySuffix = GroomDerivedDataCacheUtils::BuildCardsDerivedDataKeySuffix(GroupIndex, GetHairGroupsLOD()[GroupIndex].LODs, GetHairGroupsCards());
 	const FString DerivedDataKey = StrandsKey + KeySuffix;
 
-	if (CardsDerivedDataKey[GroupIndex] == DerivedDataKey)
-	{
-		// Skip if no change
-		return false;
-	}
-
 	FHairGroupPlatformData& HairGroupData = GetHairGroupsPlatformData()[GroupIndex];
 
 	// We should never query the DDC for something that we won't be able to build. So check that first.
@@ -2821,12 +2815,6 @@ bool UGroomAsset::CacheMeshesData(uint32 GroupIndex)
 	const FString KeySuffix = GroomDerivedDataCacheUtils::BuildMeshesDerivedDataKeySuffix(GroupIndex, GetHairGroupsLOD()[GroupIndex].LODs, GetHairGroupsMeshes());
 	const FString DerivedDataKey = GroomDerivedDataCacheUtils::BuildGroomDerivedDataKey(KeySuffix);
 
-	if (MeshesDerivedDataKey[GroupIndex] == DerivedDataKey)
-	{
-		// Skip if no change
-		return false;
-	}
-
 	FHairGroupPlatformData& HairGroupData = GetHairGroupsPlatformData()[GroupIndex];
 
 	// We should never query the DDC for something that we won't be able to build. So check that first.
