Fix incorrect hair physics output.

diff --git a/Engine/Plugins/Runtime/HairStrands/Shaders/Private/NiagaraDataInterfaceHairStrandsTemplate.ush b/Engine/Plugins/Runtime/HairStrands/Shaders/Private/NiagaraDataInterfaceHairStrandsTemplate.ush
--- a/Engine/Plugins/Runtime/HairStrands/Shaders/Private/NiagaraDataInterfaceHairStrandsTemplate.ush
+++ b/Engine/Plugins/Runtime/HairStrands/Shaders/Private/NiagaraDataInterfaceHairStrandsTemplate.ush
@@ -666,20 +666,19 @@ void UpdateDeformedPositions_{ParameterName}(in float3 LocalDisplace, in int Pos
 {
 	if( {ParameterName}_InterpolationMode >= 1)
 	{
-		uint4 Packed = { ParameterName }_RestPositionBuffer[PositionIndex];
-
-		const float3 LocalTrianglePosition = TriangleLocalPosition_{ParameterName}(RestTrianglePosition, RestTriangleOrientation, UnpackHairControlPointPosition(Packed, {ParameterName}_RestPositionOffset) + LocalDisplace);
+		// TODO: convert this code back to use UnpackHairControlPointPosition & PackHairControlPointPosition
+		const float3 LocalTrianglePosition = TriangleLocalPosition_{ParameterName}(RestTrianglePosition, RestTriangleOrientation, 
+			f16tof32({ParameterName}_RestPositionBuffer[PositionIndex].xyz) + {ParameterName}_RestPositionOffset) + LocalDisplace;
 			
-		const float3 LocalComponentPosition = TriangleWorldPosition_{ParameterName}(DeformedTrianglePosition, DeformedTriangleOrientation, LocalTrianglePosition);
-		{ParameterName}_DeformedPositionBuffer[PositionIndex] = PackHairControlPointPosition(Packed, LocalComponentPosition, {ParameterName}_DeformedPositionOffset[0].xyz);
+		const float3 LocalComponentPosition = TriangleWorldPosition_{ParameterName}(DeformedTrianglePosition, DeformedTriangleOrientation, LocalTrianglePosition) - {ParameterName}_DeformedPositionOffset[0].xyz;
+		{ParameterName}_DeformedPositionBuffer[PositionIndex] = uint4(f32tof16(LocalComponentPosition), {ParameterName}_RestPositionBuffer[PositionIndex].w);
 		
 	}
 	else
 	{
-		uint4 Packed = {ParameterName}_RestPositionBuffer[PositionIndex];
-
-		const float3 LocalComponentPosition = UnpackHairControlPointPosition(Packed, {ParameterName}_RestPositionOffset) + LocalDisplace;
-		{ ParameterName }_DeformedPositionBuffer[PositionIndex] = PackHairControlPointPosition(Packed, LocalComponentPosition, {ParameterName}_DeformedPositionOffset[0].xyz);
+		// TODO: convert this code back to use UnpackHairControlPointPosition & PackHairControlPointPosition
+		const float3 LocalComponentPosition = f16tof32({ParameterName}_RestPositionBuffer[PositionIndex].xyz) + LocalDisplace + {ParameterName}_RestPositionOffset - {ParameterName}_DeformedPositionOffset[0].xyz;
+		{ParameterName}_DeformedPositionBuffer[PositionIndex] = uint4(f32tof16(LocalComponentPosition), {ParameterName}_RestPositionBuffer[PositionIndex].w);	
 	}
 }
 
diff --git a/Engine/Plugins/Runtime/HairStrands/Shaders/Private/NiagaraInterpolateGroomGuides.usf b/Engine/Plugins/Runtime/HairStrands/Shaders/Private/NiagaraInterpolateGroomGuides.usf
--- a/Engine/Plugins/Runtime/HairStrands/Shaders/Private/NiagaraInterpolateGroomGuides.usf
+++ b/Engine/Plugins/Runtime/HairStrands/Shaders/Private/NiagaraInterpolateGroomGuides.usf
@@ -132,7 +132,7 @@ void MainCS(uint3 DispatchThreadId : SV_DispatchThreadID)
 		
 		const uint4 Packed = RestPositionBuffer[DispatchThreadId.x];
 		const float3 LocalTrianglePosition = DIHairStrands_TriangleLocalPosition(RestTrianglePosition, RestTriangleOrientation, 
-					UnpackHairControlPointPosition(Packed, RestPositionOffset) + LocalPointDisplace);
+					UnpackHairControlPointPosition(Packed, RestPositionOffset)) + LocalPointDisplace;
 		
 		const float3 LocalComponentPosition = DIHairStrands_TriangleWorldPosition(DeformedTrianglePosition, DeformedTriangleOrientation, LocalTrianglePosition) - DeformedPositionOffset[0].xyz;
 		DeformedPositionBuffer[DispatchThreadId.x] = PackHairControlPointPosition(Packed, LocalComponentPosition);
