Fixed the other use

diff --git a/Engine/Plugins/Runtime/HairStrands/Source/HairStrandsCore/Private/GroomComponent.cpp b/Engine/Plugins/Runtime/HairStrands/Source/HairStrandsCore/Private/GroomComponent.cpp
--- a/Engine/Plugins/Runtime/HairStrands/Source/HairStrandsCore/Private/GroomComponent.cpp
+++ b/Engine/Plugins/Runtime/HairStrands/Source/HairStrandsCore/Private/GroomComponent.cpp
@@ -2191,7 +2191,7 @@ void UGroomComponent::UpdateSimulatedGroups()
 {
 	if (HairGroupInstances.Num()>0)
 	{
-		const uint32 Id = ComponentId.PrimIDValue;
+		const uint32 Id = GetPrimitiveSceneId().PrimIDValue;
 
 		const bool bIsStrandsEnabled = IsHairStrandsEnabled(EHairStrandsShaderType::Strands);
 
@@ -2616,7 +2616,7 @@ void UGroomComponent::InitResources(bool bIsBindingReloading)
 		HairGroupInstance->Debug.LODSelectionTypeForDebug = LODSelectionType;
 		HairGroupInstance->DeformedComponent = DeformedMeshComponent;
 		HairGroupInstance->DeformedSection = GroomAsset->GetDeformedGroupSections().IsValidIndex(GroupIt) ? GroomAsset->GetDeformedGroupSections()[GroupIt] : INDEX_NONE; 
-		HairGroupInstance->Debug.MeshComponentId = ValidatedMeshComponent ? ValidatedMeshComponent->ComponentId : FPrimitiveComponentId();
+		HairGroupInstance->Debug.MeshComponentId = ValidatedMeshComponent ? ValidatedMeshComponent->GetPrimitiveSceneId() : FPrimitiveComponentId();
 		HairGroupInstance->Debug.CachedMeshPersistentPrimitiveIndex = FPersistentPrimitiveIndex();
 
 		if (RegisteredMeshComponent)
@@ -3363,7 +3363,7 @@ void UGroomComponent::TickComponent(float DeltaTime, enum ELevelTick TickType, F
 {
 	Super::TickComponent(DeltaTime, TickType, ThisTickFunction);	
 	
-	const uint32 Id = ComponentId.PrimIDValue;
+	const uint32 Id = GetPrimitiveSceneId().PrimIDValue;
 	const ERHIFeatureLevel::Type FeatureLevel = GetWorld() ? ERHIFeatureLevel::Type(GetWorld()->GetFeatureLevel()) : ERHIFeatureLevel::Num;
 
 	// When a groom binding and simulation are disabled, and the groom component is parented with a skeletal mesh, we can optionally 
