GroomAsset DDC: Do not assume that existence of the header's DDC value implies existence of the data segment's DDC values. Verify that each DDC segment exists and rebuild if any of them do not exist.
#rnx

diff --git a/Engine/Plugins/Runtime/HairStrands/Source/HairStrandsCore/Private/GroomAsset.cpp b/Engine/Plugins/Runtime/HairStrands/Source/HairStrandsCore/Private/GroomAsset.cpp
--- a/Engine/Plugins/Runtime/HairStrands/Source/HairStrandsCore/Private/GroomAsset.cpp
+++ b/Engine/Plugins/Runtime/HairStrands/Source/HairStrandsCore/Private/GroomAsset.cpp
@@ -1830,6 +1830,34 @@ static void InternalSerializeGuide(FArchive& Ar, UObject* Owner, FHairGroupPlatf
 	GuideData.BulkData.Serialize(Ar, Owner);
 }
 
+static bool AfterVerifyHeader_DoesDataExist(UObject* Owner, FHairGroupPlatformData::FStrands& StrandData)
+{
+	bool bAllSucceeded = true;
+	bool bSucceeded;
+	{
+		FHairStreamingRequest R;
+		R.Request(HAIR_MAX_NUM_CURVE_PER_GROUP, HAIR_MAX_NUM_POINT_PER_GROUP, -1/*LODIndex*/, StrandData.BulkData,
+			true /*bWait*/, true /*bFillBulkdata*/, true /*bWarmCache*/, Owner->GetFName(),
+			&bSucceeded);
+		bAllSucceeded &= bSucceeded;
+	}
+	{
+		FHairStreamingRequest R;
+		R.Request(HAIR_MAX_NUM_CURVE_PER_GROUP, HAIR_MAX_NUM_POINT_PER_GROUP, -1/*LODIndex*/, StrandData.InterpolationBulkData,
+			true /*bWait*/, true /*bFillBulkdata*/, true /*bWarmCache*/, Owner->GetFName(),
+			&bSucceeded);
+		bAllSucceeded &= bSucceeded;
+	}
+	{
+		FHairStreamingRequest R;
+		R.Request(HAIR_MAX_NUM_CURVE_PER_GROUP, HAIR_MAX_NUM_POINT_PER_GROUP, -1/*LODIndex*/, StrandData.ClusterBulkData,
+			true /*bWait*/, true /*bFillBulkdata*/, true /*bWarmCache*/, Owner->GetFName(),
+			&bSucceeded);
+		bAllSucceeded &= bSucceeded;
+	}
+	return bAllSucceeded;
+}
+
 static void InternalSerializeStrand(FArchive& Ar, UObject* Owner, FHairGroupPlatformData::FStrands& StrandData, bool bHeader, bool bData)
 {
 	Ar.UsingCustomVersion(FAnimObjectVersion::GUID);
@@ -2545,7 +2573,7 @@ bool UGroomAsset::CacheStrandsData(uint32 GroupIndex, FString& OutDerivedDataKey
 
 	FHairGroupPlatformData& PlatformData = GetHairGroupsPlatformData()[GroupIndex];
 
-	bool bSuccess = true;
+	bool bSuccess = false;
 	if (Data)
 	{
 		UE_CLOG(IsHairStrandsDDCLogEnable(), LogHairStrands, Log, TEXT("[Groom/DDC] Strands - Found (Groom:%s Group6:%d)."), *GetName(), GroupIndex);
@@ -2558,8 +2586,9 @@ bool UGroomAsset::CacheStrandsData(uint32 GroupIndex, FString& OutDerivedDataKey
 
 		InternalSerializeGuide(Ar, this, PlatformData.Guides);
 		InternalSerializeStrand(Ar, this, PlatformData.Strands, true/*Header*/, false/*Data*/);
+		bSuccess = AfterVerifyHeader_DoesDataExist(this, PlatformData.Strands);
 	}
-	else
+	if (!bSuccess)
 	{
 		if (!IsInGameThread())
 		{
diff --git a/Engine/Plugins/Runtime/HairStrands/Source/HairStrandsCore/Private/HairStrandsDatas.cpp b/Engine/Plugins/Runtime/HairStrands/Source/HairStrandsCore/Private/HairStrandsDatas.cpp
--- a/Engine/Plugins/Runtime/HairStrands/Source/HairStrandsCore/Private/HairStrandsDatas.cpp
+++ b/Engine/Plugins/Runtime/HairStrands/Source/HairStrandsCore/Private/HairStrandsDatas.cpp
@@ -6,6 +6,7 @@
 #include "HairAttributes.h"
 #include "IO/IoDispatcher.h"
 #include "GroomRBFDeformer.h"
+#include "Misc/ScopeExit.h"
 
 #if WITH_EDITOR
 #include "DerivedDataCache.h"
@@ -400,8 +401,19 @@ void FHairStreamingRequest::WarmCache(uint32 InRequestedCurveCount, uint32 InReq
 // Request fullfil 2 use cases:
 // * Load IO/DDC data and upload them to GPU
 // * Load DDC data and store them into bulkdata for serialization
-void FHairStreamingRequest::Request(uint32 InRequestedCurveCount, uint32 InRequestedPointCount, int32 InLODIndex, FHairStrandsBulkCommon& In, bool bWait, bool bFillBulkdata, bool bWarmCache, const FName& InOwnerName)
+void FHairStreamingRequest::Request(uint32 InRequestedCurveCount, uint32 InRequestedPointCount, int32 InLODIndex, FHairStrandsBulkCommon& In,
+	bool bWait, bool bFillBulkdata, bool bWarmCache, const FName& InOwnerName, bool* bWaitResult)
 {
+	bool bLogErrors = bWaitResult == nullptr;
+	bool bLocalWaitResult = true;
+	ON_SCOPE_EXIT
+	{
+		if (bWaitResult)
+		{
+			*bWaitResult = bLocalWaitResult;
+		}
+	};
+
 	if (In.GetResourceCount() == 0 || InRequestedCurveCount == 0)
 	{
 		CurveCount = 0;
@@ -443,7 +455,7 @@ void FHairStreamingRequest::Request(uint32 InRequestedCurveCount, uint32 InReque
 
 		//FRequestBarrier Barrier(*DDCRequestOwner);	// This is a critical section on the owner. It does not constrain ordering
 		GetCache().GetChunks(Requests, *DDCRequestOwner,
-		[this, &In, bFillBulkdata, InOwnerName](FCacheGetChunkResponse && Response)
+		[this, &In, bFillBulkdata, InOwnerName, bLogErrors](FCacheGetChunkResponse && Response)
 		{
 			if (Response.Status == UE::DerivedData::EStatus::Ok)
 			{
@@ -476,7 +488,8 @@ void FHairStreamingRequest::Request(uint32 InRequestedCurveCount, uint32 InReque
 			{
 				FHairStreamingRequest::FChunk& Chunk = *(FHairStreamingRequest::FChunk*)Response.UserData;
 				Chunk.Status = FHairStreamingRequest::FChunk::Failed;
-				UE_LOG(LogHairStrands, Error, TEXT("[Groom] DDC request failed for '%s' (Key:%s) "), *InOwnerName.ToString(), *In.DerivedDataKey);
+				UE_CLOG(bLogErrors, LogHairStrands, Error,
+					TEXT("[Groom] DDC request failed for '%s' (Key:%s) "), *InOwnerName.ToString(), *In.DerivedDataKey);
 			}
 		});
 
@@ -484,7 +497,14 @@ void FHairStreamingRequest::Request(uint32 InRequestedCurveCount, uint32 InReque
 		if (bWait || bFillBulkdata)
 		{
 			DDCRequestOwner->Wait();
-			check(IsCompleted());
+			if (!IsCompleted())
+			{
+				if (bLogErrors)
+				{
+					checkf(IsCompleted(), TEXT("HairStrands fatal error: DDC request failed."));
+				}
+				bLocalWaitResult = false;
+			}
 		}
 	}
 #endif
diff --git a/Engine/Plugins/Runtime/HairStrands/Source/HairStrandsCore/Public/HairStrandsDatas.h b/Engine/Plugins/Runtime/HairStrands/Source/HairStrandsCore/Public/HairStrandsDatas.h
--- a/Engine/Plugins/Runtime/HairStrands/Source/HairStrandsCore/Public/HairStrandsDatas.h
+++ b/Engine/Plugins/Runtime/HairStrands/Source/HairStrandsCore/Public/HairStrandsDatas.h
@@ -341,7 +341,9 @@ struct FHairStreamingRequest
 		void Release();
 	};
 
-	void Request(uint32 InRequestedCurveCount, uint32 InRequestedPointCount, int32 InLODIndex, FHairStrandsBulkCommon& In, bool bWait=false, bool bFillBulkData=false, bool bWarmCache=false, const FName& InOwnerName = NAME_None);
+	void Request(uint32 InRequestedCurveCount, uint32 InRequestedPointCount, int32 InLODIndex, FHairStrandsBulkCommon& In,
+		bool bWait=false, bool bFillBulkData=false, bool bWarmCache=false, const FName& InOwnerName = NAME_None,
+		bool* bWaitResult = nullptr);
 	bool IsNone() const;
 	bool IsCompleted();
 	bool IsUnloading() const;
