Fix incorrect building of hair strands data when a groom contain only mesh/cards data.

diff --git a/Engine/Plugins/Runtime/HairStrands/Source/HairStrandsCore/Private/GroomAsset.cpp b/Engine/Plugins/Runtime/HairStrands/Source/HairStrandsCore/Private/GroomAsset.cpp
--- a/Engine/Plugins/Runtime/HairStrands/Source/HairStrandsCore/Private/GroomAsset.cpp
+++ b/Engine/Plugins/Runtime/HairStrands/Source/HairStrandsCore/Private/GroomAsset.cpp
@@ -2528,7 +2528,7 @@ bool UGroomAsset::CacheDerivedData(uint32 GroupIndex)
 
 bool UGroomAsset::CacheStrandsData(uint32 GroupIndex, FString& OutDerivedDataKey)
 {
-	if (!HairDescriptionBulkData[HairDescriptionType])
+	if (!HairDescriptionBulkData[HairDescriptionType] || HairDescriptionBulkData[HairDescriptionType]->IsEmpty())
 	{
 		// Groom having only cards/meshes (i.e., not containing imported strands) don't have hair description
 		bool bHasStrands = false;
