Fix hair strands loading.

diff --git a/Engine/Plugins/Runtime/HairStrands/Source/HairStrandsCore/Private/GroomAsset.cpp b/Engine/Plugins/Runtime/HairStrands/Source/HairStrandsCore/Private/GroomAsset.cpp
--- a/Engine/Plugins/Runtime/HairStrands/Source/HairStrandsCore/Private/GroomAsset.cpp
+++ b/Engine/Plugins/Runtime/HairStrands/Source/HairStrandsCore/Private/GroomAsset.cpp
@@ -713,7 +713,7 @@ static void InternalSerializePlatformDatas(FArchive& Ar, UObject* Owner, TArray<
 static void InternalSerializeCard(FArchive& Ar, UObject* Owner, FHairGroupPlatformData::FCards::FLOD& CardLODData);
 static void InternalSerializeMesh(FArchive& Ar, UObject* Owner, FHairGroupPlatformData::FMeshes::FLOD& MeshLODData);
 static void InternalSerializeGuide(FArchive& Ar, UObject* Owner, FHairGroupPlatformData::FGuides& GuideData);
-static void InternalSerializeStrand(FArchive& Ar, UObject* Owner, FHairGroupPlatformData::FStrands& StrandData, bool bHeader, bool bData);
+static void InternalSerializeStrand(FArchive& Ar, UObject* Owner, FHairGroupPlatformData::FStrands& StrandData, bool bHeader, bool bData, bool* bOutHasDataInCache=nullptr);
 static void InternalSerializePlatformData(FArchive& Ar, UObject* Owner, FHairGroupPlatformData& GroupData);
 
 // This dummy serialization function is only intended to support *loading* of legacy content 
@@ -1830,35 +1830,7 @@ static void InternalSerializeGuide(FArchive& Ar, UObject* Owner, FHairGroupPlatf
 	GuideData.BulkData.Serialize(Ar, Owner);
 }
 
-static bool AfterVerifyHeader_DoesDataExist(UObject* Owner, FHairGroupPlatformData::FStrands& StrandData)
-{
-	bool bAllSucceeded = true;
-	bool bSucceeded;
-	{
-		FHairStreamingRequest R;
-		R.Request(HAIR_MAX_NUM_CURVE_PER_GROUP, HAIR_MAX_NUM_POINT_PER_GROUP, -1/*LODIndex*/, StrandData.BulkData,
-			true /*bWait*/, true /*bFillBulkdata*/, true /*bWarmCache*/, Owner->GetFName(),
-			&bSucceeded);
-		bAllSucceeded &= bSucceeded;
-	}
-	{
-		FHairStreamingRequest R;
-		R.Request(HAIR_MAX_NUM_CURVE_PER_GROUP, HAIR_MAX_NUM_POINT_PER_GROUP, -1/*LODIndex*/, StrandData.InterpolationBulkData,
-			true /*bWait*/, true /*bFillBulkdata*/, true /*bWarmCache*/, Owner->GetFName(),
-			&bSucceeded);
-		bAllSucceeded &= bSucceeded;
-	}
-	{
-		FHairStreamingRequest R;
-		R.Request(HAIR_MAX_NUM_CURVE_PER_GROUP, HAIR_MAX_NUM_POINT_PER_GROUP, -1/*LODIndex*/, StrandData.ClusterBulkData,
-			true /*bWait*/, true /*bFillBulkdata*/, true /*bWarmCache*/, Owner->GetFName(),
-			&bSucceeded);
-		bAllSucceeded &= bSucceeded;
-	}
-	return bAllSucceeded;
-}
-
-static void InternalSerializeStrand(FArchive& Ar, UObject* Owner, FHairGroupPlatformData::FStrands& StrandData, bool bHeader, bool bData)
+static void InternalSerializeStrand(FArchive& Ar, UObject* Owner, FHairGroupPlatformData::FStrands& StrandData, bool bHeader, bool bData, bool* bOutHasDataInCache)
 {
 	Ar.UsingCustomVersion(FAnimObjectVersion::GUID);
 
@@ -1870,7 +1842,7 @@ static void InternalSerializeStrand(FArchive& Ar, UObject* Owner, FHairGroupPlat
 	{
 		{ FHairStreamingRequest R; R.Request(HAIR_MAX_NUM_CURVE_PER_GROUP, HAIR_MAX_NUM_POINT_PER_GROUP, -1/*LODIndex*/, StrandData.BulkData,               true /*bWait*/, true /*bFillBulkdata*/, true /*bWarmCache*/, Owner->GetFName()); }
 		{ FHairStreamingRequest R; R.Request(HAIR_MAX_NUM_CURVE_PER_GROUP, HAIR_MAX_NUM_POINT_PER_GROUP, -1/*LODIndex*/, StrandData.InterpolationBulkData,  true /*bWait*/, true /*bFillBulkdata*/, true /*bWarmCache*/, Owner->GetFName()); }
-		{ FHairStreamingRequest R; R.Request(HAIR_MAX_NUM_CURVE_PER_GROUP, HAIR_MAX_NUM_POINT_PER_GROUP, -1/*LODIndex*/, StrandData.ClusterBulkData, true /*bWait*/, true /*bFillBulkdata*/, true /*bWarmCache*/, Owner->GetFName()); }
+		{ FHairStreamingRequest R; R.Request(HAIR_MAX_NUM_CURVE_PER_GROUP, HAIR_MAX_NUM_POINT_PER_GROUP, -1/*LODIndex*/, StrandData.ClusterBulkData,        true /*bWait*/, true /*bFillBulkdata*/, true /*bWarmCache*/, Owner->GetFName()); }
 	}
 
 	if (!Ar.IsCooking() || !StrandData.bIsCookedOut)
@@ -1890,9 +1862,12 @@ static void InternalSerializeStrand(FArchive& Ar, UObject* Owner, FHairGroupPlat
 		const bool bPreWarmCache = IsLoading() && bHeader && !bData;
 		if (bPreWarmCache)
 		{
-			{ FHairStreamingRequest R; R.WarmCache(HAIR_MAX_NUM_CURVE_PER_GROUP, HAIR_MAX_NUM_POINT_PER_GROUP, -1/*LODIndex*/, StrandData.BulkData); }
-			{ FHairStreamingRequest R; R.WarmCache(HAIR_MAX_NUM_CURVE_PER_GROUP, HAIR_MAX_NUM_POINT_PER_GROUP, -1/*LODIndex*/, StrandData.InterpolationBulkData); }
-			{ FHairStreamingRequest R; R.WarmCache(HAIR_MAX_NUM_CURVE_PER_GROUP, HAIR_MAX_NUM_POINT_PER_GROUP, -1/*LODIndex*/, StrandData.ClusterBulkData); }
+			bool bHasDataInCache = true;
+			{ FHairStreamingRequest R; bHasDataInCache &= R.WarmCache(HAIR_MAX_NUM_CURVE_PER_GROUP, HAIR_MAX_NUM_POINT_PER_GROUP, -1/*LODIndex*/, StrandData.BulkData); }
+			{ FHairStreamingRequest R; bHasDataInCache &= R.WarmCache(HAIR_MAX_NUM_CURVE_PER_GROUP, HAIR_MAX_NUM_POINT_PER_GROUP, -1/*LODIndex*/, StrandData.InterpolationBulkData); }
+			{ FHairStreamingRequest R; bHasDataInCache &= R.WarmCache(HAIR_MAX_NUM_CURVE_PER_GROUP, HAIR_MAX_NUM_POINT_PER_GROUP, -1/*LODIndex*/, StrandData.ClusterBulkData); }
+
+			if (bOutHasDataInCache) { *bOutHasDataInCache = bHasDataInCache; }
 		}
 		#endif
 	}
@@ -2585,8 +2560,7 @@ bool UGroomAsset::CacheStrandsData(uint32 GroupIndex, FString& OutDerivedDataKey
 		FMemoryReaderView Ar(Data, /*bIsPersistent*/ true);
 
 		InternalSerializeGuide(Ar, this, PlatformData.Guides);
-		InternalSerializeStrand(Ar, this, PlatformData.Strands, true/*Header*/, false/*Data*/);
-		bSuccess = AfterVerifyHeader_DoesDataExist(this, PlatformData.Strands);
+		InternalSerializeStrand(Ar, this, PlatformData.Strands, true/*Header*/, false/*Data*/, &bSuccess);
 	}
 	if (!bSuccess)
 	{
diff --git a/Engine/Plugins/Runtime/HairStrands/Source/HairStrandsCore/Private/HairStrandsDatas.cpp b/Engine/Plugins/Runtime/HairStrands/Source/HairStrandsCore/Private/HairStrandsDatas.cpp
--- a/Engine/Plugins/Runtime/HairStrands/Source/HairStrandsCore/Private/HairStrandsDatas.cpp
+++ b/Engine/Plugins/Runtime/HairStrands/Source/HairStrandsCore/Private/HairStrandsDatas.cpp
@@ -358,7 +358,7 @@ bool FHairStreamingRequest::IsUnloading() const
 	return false; 
 }
 #if WITH_EDITORONLY_DATA
-static void RequestWarmCache(UE::DerivedData::FRequestOwner* RequestOwner, const TArray<UE::DerivedData::FCacheGetChunkRequest>& Requests)
+static bool RequestWarmCache(UE::DerivedData::FRequestOwner* RequestOwner, const TArray<UE::DerivedData::FCacheGetChunkRequest>& Requests)
 {
 	using namespace UE::DerivedData;
 
@@ -370,18 +370,20 @@ static void RequestWarmCache(UE::DerivedData::FRequestOwner* RequestOwner, const
 	{				
 		WarmRequest.Add({R.Name, R.Key, Policy, 0});
 	}
-	GetCache().GetValue(WarmRequest, *RequestOwner, [](FCacheGetValueResponse && Response) { /* If the data are not built the cache quert can return false. check(Response.Status == EStatus::Ok);*/ });
+	bool bHasDataInCache = true;
+	GetCache().GetValue(WarmRequest, *RequestOwner, [&bHasDataInCache](FCacheGetValueResponse && Response) { if (Response.Status != EStatus::Ok) { bHasDataInCache = false; } /* If the data are not built the cache query can return false. check(Response.Status == EStatus::Ok);*/ });
 	RequestOwner->Wait();
+	return bHasDataInCache;
 }
 
-void FHairStreamingRequest::WarmCache(uint32 InRequestedCurveCount, uint32 InRequestedPointCount, int32 InLODIndex, FHairStrandsBulkCommon& In)
+bool FHairStreamingRequest::WarmCache(uint32 InRequestedCurveCount, uint32 InRequestedPointCount, int32 InLODIndex, FHairStrandsBulkCommon& In)
 {
 	if (In.GetResourceCount() == 0 || InRequestedCurveCount == 0)
 	{
 		CurveCount = 0;
 		PointCount = 0;
 		LODIndex = -1;
-		return;
+		return true;
 	}
 	CurveCount = InRequestedCurveCount;
 	PointCount = InRequestedPointCount;
@@ -394,7 +396,7 @@ void FHairStreamingRequest::WarmCache(uint32 InRequestedCurveCount, uint32 InReq
 
 	check(DDCRequestOwner == nullptr);
 	DDCRequestOwner = MakeUnique<FRequestOwner>(UE::DerivedData::EPriority::Blocking);
-	RequestWarmCache(DDCRequestOwner.Get(), Requests);
+	return RequestWarmCache(DDCRequestOwner.Get(), Requests);
 }
 #endif
 
diff --git a/Engine/Plugins/Runtime/HairStrands/Source/HairStrandsCore/Public/HairStrandsDatas.h b/Engine/Plugins/Runtime/HairStrands/Source/HairStrandsCore/Public/HairStrandsDatas.h
--- a/Engine/Plugins/Runtime/HairStrands/Source/HairStrandsCore/Public/HairStrandsDatas.h
+++ b/Engine/Plugins/Runtime/HairStrands/Source/HairStrandsCore/Public/HairStrandsDatas.h
@@ -349,7 +349,7 @@ struct FHairStreamingRequest
 	bool IsUnloading() const;
 
 #if WITH_EDITORONLY_DATA
-	void WarmCache(uint32 InRequestedCurveCount, uint32 InRequestedPointCount, int32 InLODIndex, FHairStrandsBulkCommon& In);
+	bool WarmCache(uint32 InRequestedCurveCount, uint32 InRequestedPointCount, int32 InLODIndex, FHairStrandsBulkCommon& In);
 #endif
 
 #if !WITH_EDITORONLY_DATA
