Add coverage bias cards coverage texture sampler.

#prefligth 65040f985108f0fc0e8cdd95

diff --git a/Engine/Plugins/Runtime/HairStrands/Source/HairStrandsCore/Private/HairCardsVertexFactory.cpp b/Engine/Plugins/Runtime/HairStrands/Source/HairStrandsCore/Private/HairCardsVertexFactory.cpp
--- a/Engine/Plugins/Runtime/HairStrands/Source/HairStrandsCore/Private/HairCardsVertexFactory.cpp
+++ b/Engine/Plugins/Runtime/HairStrands/Source/HairStrandsCore/Private/HairCardsVertexFactory.cpp
@@ -19,6 +19,9 @@
 #include "GlobalRenderResources.h"
 #include "DataDrivenShaderPlatformInfo.h"
 
+static int32 GHairCardCoverageBias = 0;
+static FAutoConsoleVariableRef CVarHairCardCoverageBias(TEXT("r.HairStrands.Cards.CoverageBias"), GHairCardCoverageBias, TEXT("Apply a texture LOD bias to coverage texture"));
+
 template<typename T> inline void VFC_BindParam(FMeshDrawSingleShaderBindings& ShaderBindings, const FShaderResourceParameter& Param, T* Value) { if (Param.IsBound() && Value) ShaderBindings.Add(Param, Value); }
 template<typename T> inline void VFC_BindParam(FMeshDrawSingleShaderBindings& ShaderBindings, const FShaderParameter& Param, const T& Value) { if (Param.IsBound()) ShaderBindings.Add(Param, Value); }
 
@@ -52,6 +55,7 @@ FHairCardsUniformBuffer CreateHairCardsVFUniformBuffer(
 		UniformParameters.Flags = 0;
 		if (LOD.RestResource->bInvertUV) { UniformParameters.Flags |= uint32(EHairCardsFactoryFlags::InvertedUV);  }
 		UniformParameters.MaxVertexCount = LOD.RestResource->GetVertexCount();
+		UniformParameters.CoverageBias = FMath::Clamp(GHairCardCoverageBias, -16.f, 16.f);
 
 		// When the geometry is not-dynamic (no binding to skeletal mesh, no simulation), only a single vertex buffer is allocated. 
 		// In this case we force the buffer index to 0
@@ -111,6 +115,7 @@ FHairCardsUniformBuffer CreateHairCardsVFUniformBuffer(
 		UniformParameters.Flags |= uint32(EHairCardsFactoryFlags::TextureBaseColor);
 		UniformParameters.Flags |= uint32(EHairCardsFactoryFlags::TextureRoughness);
 
+		UniformParameters.CoverageBias = 0;
 		UniformParameters.MaxVertexCount = LOD.RestResource->GetVertexCount();
 		UniformParameters.NormalsBuffer = LOD.RestResource->NormalsBuffer.ShaderResourceViewRHI.GetReference();
 		UniformParameters.UVsBuffer = LOD.RestResource->UVsBuffer.ShaderResourceViewRHI.GetReference();
diff --git a/Engine/Plugins/Runtime/HairStrands/Source/HairStrandsCore/Public/GroomInstance.h b/Engine/Plugins/Runtime/HairStrands/Source/HairStrandsCore/Public/GroomInstance.h
--- a/Engine/Plugins/Runtime/HairStrands/Source/HairStrandsCore/Public/GroomInstance.h
+++ b/Engine/Plugins/Runtime/HairStrands/Source/HairStrandsCore/Public/GroomInstance.h
@@ -25,6 +25,7 @@ enum class EGroomViewMode : uint8;
 BEGIN_GLOBAL_SHADER_PARAMETER_STRUCT(FHairCardsVertexFactoryUniformShaderParameters, HAIRSTRANDSCORE_API)
 	SHADER_PARAMETER(uint32, Flags)
 	SHADER_PARAMETER(uint32, MaxVertexCount)
+	SHADER_PARAMETER(float, CoverageBias)
 	SHADER_PARAMETER_SRV(Buffer<float4>, PositionBuffer)
 	SHADER_PARAMETER_SRV(Buffer<float4>, PreviousPositionBuffer)
 	SHADER_PARAMETER_SRV(Buffer<float4>, NormalsBuffer)
diff --git a/Engine/Shaders/Private/HairStrands/HairCardsAttributeCommon.ush b/Engine/Shaders/Private/HairStrands/HairCardsAttributeCommon.ush
--- a/Engine/Shaders/Private/HairStrands/HairCardsAttributeCommon.ush
+++ b/Engine/Shaders/Private/HairStrands/HairCardsAttributeCommon.ush
@@ -49,7 +49,7 @@ float GetHairStrandsCoverage(float2 InAtlasUV)
 {
 	const bool bInvertUV = HairCardsVF.Flags & EHairCardsFVFlags_InvertedUV;
 	InAtlasUV.y = bInvertUV ? 1 - InAtlasUV.y : InAtlasUV.y;
-	return Texture2DSample(HairCardsVF.CoverageTexture, HairCardsVF.CoverageSampler, InAtlasUV).x;
+	return Texture2DSampleBias(HairCardsVF.CoverageTexture, HairCardsVF.CoverageSampler, InAtlasUV, HairCardsVF.CoverageBias).x;
 }
 
 FHairVertexAttributes GetHairStrandsAttributes(float2 InAtlasUV, float2 InRootUV = float2(0, 0), float InCardLength = 0.f, float InCardGroupIndex = 0.f, float3 InCardBaseColor = 0.f, float InCardRoughness = 0.f)
diff --git a/Engine/Source/Runtime/Renderer/Private/HairStrands/HairStrandsCluster.cpp b/Engine/Source/Runtime/Renderer/Private/HairStrands/HairStrandsCluster.cpp
--- a/Engine/Source/Runtime/Renderer/Private/HairStrands/HairStrandsCluster.cpp
+++ b/Engine/Source/Runtime/Renderer/Private/HairStrands/HairStrandsCluster.cpp
@@ -285,7 +285,7 @@ void CreateHairStrandsMacroGroups(
 	{
 		for (FHairStrandsInstance* Instance : Scene->HairStrandsSceneData.RegisteredProxies)
 		{
-			if (Instance && Instance->RegisteredIndex >= 0 && Instance->RegisteredIndex < ActiveInstanceCount && !InstancesVisibility[Instance->RegisteredIndex])
+			if (Instance && InstancesVisibility.IsValidIndex(Instance->RegisteredIndex) && !InstancesVisibility[Instance->RegisteredIndex])
 			{
 				if (IsHairStrandsVisibleInShadows(View, *Instance))
 				{
diff --git a/Engine/Source/Runtime/Renderer/Private/HairStrands/HairStrandsInterface.cpp b/Engine/Source/Runtime/Renderer/Private/HairStrands/HairStrandsInterface.cpp
--- a/Engine/Source/Runtime/Renderer/Private/HairStrands/HairStrandsInterface.cpp
+++ b/Engine/Source/Runtime/Renderer/Private/HairStrands/HairStrandsInterface.cpp
@@ -447,7 +447,7 @@ void UpdateHairStrandsBookmarkParameters(FScene* Scene, TArray<FViewInfo>& Views
 			const bool bCardsOrMeshes = Instance->GetHairGeometry() == EHairGeometryType::Cards || Instance->GetHairGeometry() == EHairGeometryType::Meshes;
 			const bool bCompatible = bStrands || bCardsOrMeshes;
 
-			if (Instance->RegisteredIndex >= 0 && Instance->RegisteredIndex < ActiveInstanceCount && !Out.InstancesVisibility[Instance->RegisteredIndex] && bCompatible)
+			if (Out.InstancesVisibility.IsValidIndex(Instance->RegisteredIndex) && !Out.InstancesVisibility[Instance->RegisteredIndex] && bCompatible)
 			{
 				if (IsHairStrandsVisibleInShadows(Views[0], *Instance))
 				{
