Enable anisotropy on cards texture samplers.

#prefligth skip

diff --git a/Engine/Plugins/Runtime/HairStrands/Source/HairStrandsCore/Private/GroomAsset.cpp b/Engine/Plugins/Runtime/HairStrands/Source/HairStrandsCore/Private/GroomAsset.cpp
--- a/Engine/Plugins/Runtime/HairStrands/Source/HairStrandsCore/Private/GroomAsset.cpp
+++ b/Engine/Plugins/Runtime/HairStrands/Source/HairStrandsCore/Private/GroomAsset.cpp
@@ -416,7 +416,7 @@ static void InitAtlasTexture(ResourceType* InResource, UTexture2D* InTexture, EH
 	ENQUEUE_RENDER_COMMAND(HairStrandsCardsTextureCommand)(
 	[InResource, InTexture, InType](FRHICommandListImmediate& RHICmdList)
 	{
-		FSamplerStateRHIRef DefaultSampler = TStaticSamplerState<SF_AnisotropicLinear, AM_Clamp, AM_Clamp, AM_Clamp, 0, 0>::GetRHI();
+		FSamplerStateRHIRef DefaultSampler = TStaticSamplerState<SF_AnisotropicLinear, AM_Clamp, AM_Clamp, AM_Clamp, 0, 8/*MaxAnisotropy*/>::GetRHI();
 		switch (InType)
 		{
 		case EHairAtlasTextureType::Depth:
