Fix hair strands resources not being loading correctly when creating hair strands texture in editor.

diff --git a/Engine/Plugins/Runtime/HairStrands/Source/HairStrandsCore/Private/GroomTextureBuilder.cpp b/Engine/Plugins/Runtime/HairStrands/Source/HairStrandsCore/Private/GroomTextureBuilder.cpp
--- a/Engine/Plugins/Runtime/HairStrands/Source/HairStrandsCore/Private/GroomTextureBuilder.cpp
+++ b/Engine/Plugins/Runtime/HairStrands/Source/HairStrandsCore/Private/GroomTextureBuilder.cpp
@@ -20,6 +20,10 @@
 #include "UObject/Package.h"
 #include "HairStrandsDefinitions.h"
 
+#if WITH_EDITORONLY_DATA
+#include "DerivedDataRequestOwner.h"
+#endif
+
 static int32 GHairStrandsTextureDilationPassCount = 8;
 static FAutoConsoleVariableRef CVarHairStrandsTextureDilationPassCount(TEXT("r.HairStrands.Textures.DilationCount"), GHairStrandsTextureDilationPassCount, TEXT("Number of dilation pass run onto the generated hair strands textures (Default:8)."));
 
@@ -1048,8 +1052,17 @@ static bool TraceTextures(
 			}
 
 			// Ensure the rest resources are loaded when rendering the strands textures
-			GroupData.Strands.RestResource->Allocate(GraphBuilder, EHairResourceLoadingType::Sync);
-			
+			// Use EHairResourceLoadingType::Async, as the we can't use ::Sync because the resource is expected to be loading in async
+			// If a streaming request is in-flight, ensure it is done and allocate the resource
+			GroupData.Strands.RestResource->Allocate(GraphBuilder, EHairResourceLoadingType::Async);
+		#if WITH_EDITORONLY_DATA
+			if (GroupData.Strands.RestResource->StreamingRequest.DDCRequestOwner)
+			{
+				GroupData.Strands.RestResource->StreamingRequest.DDCRequestOwner->Wait();
+				GroupData.Strands.RestResource->Allocate(GraphBuilder, EHairResourceLoadingType::Async);
+			}
+		#endif	
+
 			// Manually setting up parameters, since we don't have a real instance
 			FHairStrandsInstanceRawParameters Instance;
 			Instance.Common.GroupIndex = GroupIndex;
