Temporary fix groom resources being destroy and never recreated.

diff --git a/Engine/Plugins/Runtime/HairStrands/Source/HairStrandsCore/Private/GroomComponent.cpp b/Engine/Plugins/Runtime/HairStrands/Source/HairStrandsCore/Private/GroomComponent.cpp
--- a/Engine/Plugins/Runtime/HairStrands/Source/HairStrandsCore/Private/GroomComponent.cpp
+++ b/Engine/Plugins/Runtime/HairStrands/Source/HairStrandsCore/Private/GroomComponent.cpp
@@ -424,6 +424,9 @@ public:
 		// Forcing primitive uniform as we don't support robustly GPU scene data
 		bVFRequiresPrimitiveUniformBuffer = true;
 		bCastDeepShadow = true;
+		
+		// TEMP - Force OnTransformChanged and {Create, Destroy}RenderThreadResources onto the render thread.
+		bConstrainToRenderThread = true;
 
 		HairGroupMaterialProxies.SetNum(Component->HairGroupInstances.Num());
 		HairGroupInstances = Component->HairGroupInstances;
