Fix missing group name on groom asset once imported.

diff --git a/Engine/Plugins/Runtime/HairStrands/Source/HairStrandsCore/Private/GroomAsset.cpp b/Engine/Plugins/Runtime/HairStrands/Source/HairStrandsCore/Private/GroomAsset.cpp
--- a/Engine/Plugins/Runtime/HairStrands/Source/HairStrandsCore/Private/GroomAsset.cpp
+++ b/Engine/Plugins/Runtime/HairStrands/Source/HairStrandsCore/Private/GroomAsset.cpp
@@ -71,6 +71,9 @@ static FAutoConsoleVariableRef CVarHairStrandsLoadAsset(TEXT("r.HairStrands.Load
 static int32 GEnableGroomAsyncLoad = 0;
 static FAutoConsoleVariableRef CVarGroomAsyncLoad(TEXT("r.HairStrands.AsyncLoad"), GEnableGroomAsyncLoad, TEXT("Allow groom asset to be loaded asynchronously in the editor"));
 
+static int32 GUpdateGroomGroupsNames = 0;
+static FAutoConsoleVariableRef CVarUpdateGroomGroupsNames(TEXT("r.HairStrands.UpdateGroupNames"), GUpdateGroomGroupsNames, TEXT("Update groom asset groups' names if not already serialized"));
+
 void UpdateHairStrandsVerbosity(IConsoleVariable* InCVarVerbosity);
 static TAutoConsoleVariable<int32> GHairStrandsWarningLogVerbosity(
 	TEXT("r.HairStrands.Log"),
@@ -1178,6 +1181,32 @@ void UGroomAsset::UpdateHairGroupsInfo()
 	}
 }
 
+static bool UpdateHairGroupsName(UGroomAsset* In)
+{
+	bool bNeedResaved = false;
+#if WITH_EDITORONLY_DATA
+	check(In);
+
+	uint32 GroupIndex = 0;
+	for (FHairGroupInfoWithVisibility& Info : In->GetHairGroupsInfo())
+	{
+		// If the group name was not serialized, try to retrieve it from the hair description
+		if (Info.GroupName == NAME_None && In->CanRebuildFromDescription())
+		{
+			const FHairDescriptionGroups& LocalHairDescriptionGroups = In->GetHairDescriptionGroups();
+			if (LocalHairDescriptionGroups.HairGroups.IsValidIndex(GroupIndex))
+			{
+				Info.GroupName = LocalHairDescriptionGroups.HairGroups[GroupIndex].Info.GroupName;
+				bNeedResaved = true;
+			}
+		}
+
+		++GroupIndex;
+	}
+#endif
+	return bNeedResaved;
+}
+
 template<typename T>
 bool ConvertMaterial(TArray<T>& Groups, TArray<FHairGroupsMaterial>& InHairGroupsMaterials)
 {
@@ -1371,13 +1400,22 @@ void UGroomAsset::PostLoad()
 	}
 #endif
 
-	// Convert all material data to new format
+	// Convert legacy version
 	{
-		// Strands
 		bool bNeedSaving = false;
-		bNeedSaving = bNeedSaving || ConvertMaterial(GetHairGroupsRendering(),	GetHairGroupsMaterials());
-		bNeedSaving = bNeedSaving || ConvertMaterial(GetHairGroupsCards(),		GetHairGroupsMaterials());
-		bNeedSaving = bNeedSaving || ConvertMaterial(GetHairGroupsMeshes(),		GetHairGroupsMaterials());
+
+		// Convert material data to new format
+		{
+			bNeedSaving = bNeedSaving || ConvertMaterial(GetHairGroupsRendering(),	GetHairGroupsMaterials());
+			bNeedSaving = bNeedSaving || ConvertMaterial(GetHairGroupsCards(),		GetHairGroupsMaterials());
+			bNeedSaving = bNeedSaving || ConvertMaterial(GetHairGroupsMeshes(),		GetHairGroupsMaterials());
+		}
+
+		// Populate group name
+		if (GUpdateGroomGroupsNames)
+		{
+			bNeedSaving = UpdateHairGroupsName(this);
+		}
 
 		if (bNeedSaving)
 		{
diff --git a/Engine/Plugins/Runtime/HairStrands/Source/HairStrandsCore/Public/GroomAsset.h b/Engine/Plugins/Runtime/HairStrands/Source/HairStrandsCore/Public/GroomAsset.h
--- a/Engine/Plugins/Runtime/HairStrands/Source/HairStrandsCore/Public/GroomAsset.h
+++ b/Engine/Plugins/Runtime/HairStrands/Source/HairStrandsCore/Public/GroomAsset.h
@@ -431,7 +431,7 @@ class HAIRSTRANDSCORE_API UGroomAsset : public UObject, public IInterface_AssetU
 
 public:
 	UE_DEPRECATED(5.3, "Please do not access this member directly; use UGroomAsset accessor.")
-	UPROPERTY(EditAnywhere, Transient, EditFixedSize, Category = "HairInfo", meta = (DisplayName = "Group"))
+	UPROPERTY(EditAnywhere, EditFixedSize, EditFixedSize, Category = "HairInfo", meta = (DisplayName = "Group"))
 	TArray<FHairGroupInfoWithVisibility> HairGroupsInfo;
 
 	UE_DEPRECATED(5.3, "Please do not access this member directly; use UGroomAsset accessor.")
