Fix missing bound scale on groom component.

diff --git a/Engine/Plugins/Runtime/HairStrands/Source/HairStrandsCore/Private/GroomComponent.cpp b/Engine/Plugins/Runtime/HairStrands/Source/HairStrandsCore/Private/GroomComponent.cpp
--- a/Engine/Plugins/Runtime/HairStrands/Source/HairStrandsCore/Private/GroomComponent.cpp
+++ b/Engine/Plugins/Runtime/HairStrands/Source/HairStrandsCore/Private/GroomComponent.cpp
@@ -1853,6 +1853,8 @@ FPrimitiveSceneProxy* UGroomComponent::CreateSceneProxy()
 
 FBoxSphereBounds UGroomComponent::CalcBounds(const FTransform& InLocalToWorld) const
 {
+	FBoxSphereBounds Out;
+
 	if (GroomAsset && GroomAsset->GetNumHairGroups() > 0)
 	{
 		FBox LocalHairBound(EForceInit::ForceInitToZero);
@@ -1911,7 +1913,8 @@ FBoxSphereBounds UGroomComponent::CalcBounds(const FTransform& InLocalToWorld) c
 			FBox EffectiveBound = RegisteredMeshComponent->CalcBounds(InLocalToWorld).GetBox();
 			EffectiveBound.Min *= InvScale;
 			EffectiveBound.Max *= InvScale;
-			return EffectiveBound;
+
+			Out = FBoxSphereBounds(EffectiveBound);
 		}
 		else if (RegisteredMeshComponent && GHairStrands_BoundsMode == 2)
 		{
@@ -1925,7 +1928,7 @@ FBoxSphereBounds UGroomComponent::CalcBounds(const FTransform& InLocalToWorld) c
 			EffectiveBound.Max *= InvScale;			
 			EffectiveBound.Min -= FVector3d(BoundExtraRadius);
 			EffectiveBound.Max += FVector3d(BoundExtraRadius);
-			return FBoxSphereBounds(EffectiveBound.TransformBy(InLocalToWorld));
+			Out = FBoxSphereBounds(EffectiveBound.TransformBy(InLocalToWorld));
 		}
 		else if (RegisteredMeshComponent && AttachmentName.IsEmpty())
 		{
@@ -1949,18 +1952,24 @@ FBoxSphereBounds UGroomComponent::CalcBounds(const FTransform& InLocalToWorld) c
 			FBox LocalBound(EForceInit::ForceInitToZero);
 			LocalBound += InvScale * LocalSkeletalBound.Min - BoundExtraRadius;
 			LocalBound += InvScale * LocalSkeletalBound.Max + BoundExtraRadius;
-			return FBoxSphereBounds(LocalBound.TransformBy(InLocalToWorld));
+			Out = FBoxSphereBounds(LocalBound.TransformBy(InLocalToWorld));
 		}
 		else
 		{
-			return FBoxSphereBounds(LocalHairBound.TransformBy(InLocalToWorld));
+			Out = FBoxSphereBounds(LocalHairBound.TransformBy(InLocalToWorld));
 		}
 	}
 	else
 	{
 		FBoxSphereBounds LocalBounds(EForceInit::ForceInitToZero);
-		return FBoxSphereBounds(LocalBounds.TransformBy(InLocalToWorld));
+		Out = FBoxSphereBounds(LocalBounds.TransformBy(InLocalToWorld));
 	}
+
+	// Apply bounds scale
+	Out.BoxExtent *= BoundsScale;
+	Out.SphereRadius *= BoundsScale;
+
+	return Out;
 }
 
 /* Return the material slot index corresponding to the material name */
