Fix groom binding loading.

diff --git a/Engine/Plugins/Runtime/HairStrands/Source/HairStrandsCore/Private/GroomBindingAsset.cpp b/Engine/Plugins/Runtime/HairStrands/Source/HairStrandsCore/Private/GroomBindingAsset.cpp
--- a/Engine/Plugins/Runtime/HairStrands/Source/HairStrandsCore/Private/GroomBindingAsset.cpp
+++ b/Engine/Plugins/Runtime/HairStrands/Source/HairStrandsCore/Private/GroomBindingAsset.cpp
@@ -846,10 +846,11 @@ void UGroomBindingAsset::CacheDerivedDatas(uint32 InGroupIndex, const FString Ke
 {
 	const FString DerivedDataKey = GroomBindingDerivedDataCacheUtils::BuildGroomBindingDerivedDataKey(KeySuffix + FString(TEXT("_Group")) + FString::FromInt(InGroupIndex));
 
-	bOutValid = false;
+	bOutValid = true;
 	bOutReloadResource = false;
 	if (DerivedDataKey != CachedDerivedDataKey[InGroupIndex])
 	{
+		bOutValid = false;
 		using namespace UE::DerivedData;
 
 		const FCacheKey HeaderKey = ConvertLegacyCacheKey(DerivedDataKey + FString(TEXT("_Header")));
