Fix shader compilation issue on certain platform.

diff --git a/Engine/Plugins/Runtime/HairStrands/Source/HairStrandsCore/Private/GroomDebug.cpp b/Engine/Plugins/Runtime/HairStrands/Source/HairStrandsCore/Private/GroomDebug.cpp
--- a/Engine/Plugins/Runtime/HairStrands/Source/HairStrandsCore/Private/GroomDebug.cpp
+++ b/Engine/Plugins/Runtime/HairStrands/Source/HairStrandsCore/Private/GroomDebug.cpp
@@ -155,7 +155,7 @@ class FHairProjectionMeshDebugCS : public FGlobalShader
 
 	static bool ShouldCompilePermutation(const FGlobalShaderPermutationParameters& Parameters)
 	{
-		return IsHairStrandsSupported(EHairStrandsShaderType::All, Parameters.Platform);
+		return IsHairStrandsSupported(EHairStrandsShaderType::All, Parameters.Platform) && ShaderPrint::IsSupported(Parameters.Platform);
 	}
 
 	static void ModifyCompilationEnvironment(const FGlobalShaderPermutationParameters& Parameters, FShaderCompilerEnvironment& OutEnvironment)
@@ -169,6 +169,7 @@ IMPLEMENT_GLOBAL_SHADER(FHairProjectionMeshDebugCS, "/Engine/Private/HairStrands
 
 static void AddDebugProjectionMeshPass(
 	FRDGBuilder& GraphBuilder,
+	const FSceneView* View,
 	FGlobalShaderMap* ShaderMap,
 	const FShaderPrintData* ShaderPrintData,
 	const FIntRect Viewport,
@@ -176,6 +177,11 @@ static void AddDebugProjectionMeshPass(
 	const EHairStrandsProjectionMeshType MeshType,
 	FHairStrandsProjectionMeshData::Section& MeshSectionData)
 {
+	if (!View || !ShaderPrint::IsSupported(View->GetShaderPlatform()))
+	{
+		return;
+	}
+
 	// Force ShaderPrint on.
 	ShaderPrint::SetEnabled(true);
 
@@ -231,7 +237,7 @@ static void AddDebugProjectionMeshPass(
 class FHairProjectionHairDebugCS : public FGlobalShader
 {
 	DECLARE_GLOBAL_SHADER(FHairProjectionHairDebugCS);
-	SHADER_USE_PARAMETER_STRUCT(FHairProjectionHairDebugCS, FGlobalShader);
+	SHADER_USE_PARAMETER_STRUCT(FHairProjectionHairDebugCS, FGlobalShader);	
 
 	class FInputType : SHADER_PERMUTATION_INT("PERMUTATION_INPUT_TYPE", 3);
 	using FPermutationDomain = TShaderPermutationDomain<FInputType>;
@@ -255,7 +261,7 @@ class FHairProjectionHairDebugCS : public FGlobalShader
 
 	static bool ShouldCompilePermutation(const FGlobalShaderPermutationParameters& Parameters)
 	{
-		return IsHairStrandsSupported(EHairStrandsShaderType::All, Parameters.Platform);
+		return IsHairStrandsSupported(EHairStrandsShaderType::All, Parameters.Platform) && ShaderPrint::IsSupported(Parameters.Platform);
 	}
 
 	static void ModifyCompilationEnvironment(const FGlobalShaderPermutationParameters& Parameters, FShaderCompilerEnvironment& OutEnvironment)
@@ -276,6 +282,7 @@ enum class EDebugProjectionHairType
 
 static void AddDebugProjectionHairPass(
 	FRDGBuilder& GraphBuilder,
+	const FSceneView* View,
 	FGlobalShaderMap* ShaderMap,
 	const FShaderPrintData* ShaderPrintData,
 	FIntRect Viewport,
@@ -287,6 +294,11 @@ static void AddDebugProjectionHairPass(
 	const EDebugProjectionHairType GeometryType,
 	const HairStrandsTriangleType PoseType)
 {
+	if (!View || !ShaderPrint::IsSupported(View->GetShaderPlatform()))
+	{
+		return;
+	}
+
 	// Force ShaderPrint on.
 	ShaderPrint::SetEnabled(true);
 
@@ -397,7 +409,7 @@ class FDrawDebugClusterAABBCS : public FGlobalShader
 
 public:
 	static uint32 GetGroupSize() { return 64; }
-	static bool ShouldCompilePermutation(const FGlobalShaderPermutationParameters& Parameters) { return IsHairStrandsSupported(EHairStrandsShaderType::Strands, Parameters.Platform); }
+	static bool ShouldCompilePermutation(const FGlobalShaderPermutationParameters& Parameters) { return IsHairStrandsSupported(EHairStrandsShaderType::Strands, Parameters.Platform) && ShaderPrint::IsSupported(Parameters.Platform); }
 	static void ModifyCompilationEnvironment(const FGlobalShaderPermutationParameters& Parameters, FShaderCompilerEnvironment& OutEnvironment)
 	{
 		FGlobalShader::ModifyCompilationEnvironment(Parameters, OutEnvironment);
@@ -418,6 +430,11 @@ void AddDrawDebugClusterPass(
 	EGroomViewMode ViewMode,
 	FHairStrandClusterData& HairClusterData)
 {
+	if (!ShaderPrint::IsSupported(View.GetShaderPlatform()))
+	{
+		return;
+	}
+
 	ShaderPrint::SetEnabled(true);
 	ShaderPrint::RequestSpaceForLines(64000u);
 	ShaderPrint::RequestSpaceForCharacters(2000);
@@ -467,7 +484,7 @@ class FDrawDebugCardAtlasCS : public FGlobalShader
 	END_SHADER_PARAMETER_STRUCT()
 
 public:
-	static bool ShouldCompilePermutation(const FGlobalShaderPermutationParameters& Parameters) { return IsHairStrandsSupported(EHairStrandsShaderType::Cards, Parameters.Platform); }
+	static bool ShouldCompilePermutation(const FGlobalShaderPermutationParameters& Parameters) { return IsHairStrandsSupported(EHairStrandsShaderType::Cards, Parameters.Platform) && ShaderPrint::IsSupported(Parameters.Platform); }
 	static void ModifyCompilationEnvironment(const FGlobalShaderPermutationParameters& Parameters, FShaderCompilerEnvironment& OutEnvironment)
 	{
 		FGlobalShader::ModifyCompilationEnvironment(Parameters, OutEnvironment);
@@ -485,6 +502,11 @@ static void AddDrawDebugCardsAtlasPass(
 	const FShaderPrintData* ShaderPrintData,
 	FRDGTextureRef SceneColorTexture)
 {
+	if (!ShaderPrint::IsSupported(View.GetShaderPlatform()))
+	{
+		return;
+	}
+
 	if (Instance->HairGroupPublicData->VFInput.GeometryType != EHairGeometryType::Cards || ShaderPrintData == nullptr)
 	{
 		return;
@@ -626,7 +648,7 @@ class FDrawDebugCardGuidesCS : public FGlobalShader
 	END_SHADER_PARAMETER_STRUCT()
 
 public:
-	static bool ShouldCompilePermutation(const FGlobalShaderPermutationParameters& Parameters) { return IsHairStrandsSupported(EHairStrandsShaderType::Cards, Parameters.Platform); }
+	static bool ShouldCompilePermutation(const FGlobalShaderPermutationParameters& Parameters) { return IsHairStrandsSupported(EHairStrandsShaderType::Cards, Parameters.Platform) && ShaderPrint::IsSupported(Parameters.Platform); }
 	static void ModifyCompilationEnvironment(const FGlobalShaderPermutationParameters& Parameters, FShaderCompilerEnvironment& OutEnvironment)
 	{
 		FGlobalShader::ModifyCompilationEnvironment(Parameters, OutEnvironment);
@@ -645,6 +667,11 @@ static void AddDrawDebugCardsGuidesPass(
 	const bool bDeformed, 
 	const bool bRen)
 {
+	if (!ShaderPrint::IsSupported(View.GetShaderPlatform()))
+	{
+		return;
+	}
+
 	// Force ShaderPrint on.
 	ShaderPrint::SetEnabled(true);
 
@@ -808,7 +835,7 @@ class FHairDebugPrintInstanceCS : public FGlobalShader
 	END_SHADER_PARAMETER_STRUCT()
 
 public:
-	static bool ShouldCompilePermutation(const FGlobalShaderPermutationParameters& Parameters) { return IsHairStrandsSupported(EHairStrandsShaderType::All, Parameters.Platform); }
+	static bool ShouldCompilePermutation(const FGlobalShaderPermutationParameters& Parameters) { return IsHairStrandsSupported(EHairStrandsShaderType::All, Parameters.Platform) && ShaderPrint::IsSupported(Parameters.Platform); }
 	static void ModifyCompilationEnvironment(const FGlobalShaderPermutationParameters& Parameters, FShaderCompilerEnvironment& OutEnvironment)
 	{
 		// Skip optimization for avoiding long compilation time due to large UAV writes
@@ -833,6 +860,11 @@ static void AddHairDebugPrintInstancePass(
 	const FHairStrandsInstances& Instances,
 	const FUintVector4& InstanceCountPerType)
 {
+	if (!ShaderPrint::IsSupported(View.GetShaderPlatform()))
+	{
+		return;
+	}
+
 	const uint32 InstanceCount = Instances.Num();
 
 	// Force ShaderPrint on.
@@ -1061,7 +1093,7 @@ class FHairDebugPrintMemoryCS : public FGlobalShader
 	END_SHADER_PARAMETER_STRUCT()
 
 public:
-	static bool ShouldCompilePermutation(const FGlobalShaderPermutationParameters& Parameters) { return IsHairStrandsSupported(EHairStrandsShaderType::All, Parameters.Platform); }
+	static bool ShouldCompilePermutation(const FGlobalShaderPermutationParameters& Parameters) { return IsHairStrandsSupported(EHairStrandsShaderType::All, Parameters.Platform) && ShaderPrint::IsSupported(Parameters.Platform); }
 	static void ModifyCompilationEnvironment(const FGlobalShaderPermutationParameters& Parameters, FShaderCompilerEnvironment& OutEnvironment)
 	{
 		// Skip optimization for avoiding long compilation time due to large UAV writes
@@ -1083,6 +1115,11 @@ static void AddHairDebugPrintMemoryPass(
 	const FSceneView& View,
 	const FShaderPrintData* ShaderPrintData)
 {
+	if (!ShaderPrint::IsSupported(View.GetShaderPlatform()))
+	{
+		return;
+	}
+
 	ShaderPrint::SetEnabled(true);
 	
 	if (!ShaderPrintData) { return; }
@@ -1310,15 +1347,17 @@ void RunHairStrandsDebug(
 	if (ViewMode == EGroomViewMode::MeshProjection)
 	{
 		{
+			const FSceneView* LocalView = &View;
+
 			if (GHairDebugMeshProjection_SkinCacheMesh > 0)
 			{
-				auto RenderMeshProjection = [ShaderMap, Scene, ShaderPrintData, Viewport, &ViewUniformBuffer, Instances, &GraphBuilder](FRDGBuilder& LocalGraphBuilder, EHairStrandsProjectionMeshType MeshType)
+				auto RenderMeshProjection = [ShaderMap, LocalView, Scene, ShaderPrintData, Viewport, &ViewUniformBuffer, Instances, &GraphBuilder](FRDGBuilder& LocalGraphBuilder, EHairStrandsProjectionMeshType MeshType)
 				{
 					FHairStrandsProjectionMeshData::LOD MeshProjectionLODData;
 					GetGroomInterpolationData(GraphBuilder, ShaderMap, Scene, Instances, MeshType, MeshProjectionLODData);
 					for (FHairStrandsProjectionMeshData::Section& Section : MeshProjectionLODData.Sections)
 					{
-						AddDebugProjectionMeshPass(LocalGraphBuilder, ShaderMap, ShaderPrintData, Viewport, ViewUniformBuffer, MeshType, Section);
+						AddDebugProjectionMeshPass(LocalGraphBuilder, LocalView, ShaderMap, ShaderPrintData, Viewport, ViewUniformBuffer, MeshType, Section);
 					}
 				};
 
@@ -1328,7 +1367,7 @@ void RunHairStrandsDebug(
 				RenderMeshProjection(GraphBuilder, EHairStrandsProjectionMeshType::TargetMesh);
 			}
 
-			auto RenderProjectionData = [&GraphBuilder, ShaderMap, Viewport, &ViewUniformBuffer, Instances, ShaderPrintData](bool bGuide, bool bRestTriangle, bool bRestFrame, bool bRestSamples, bool bDeformedTriangle, bool bDeformedFrame, bool bDeformedSamples)
+			auto RenderProjectionData = [&GraphBuilder, LocalView, ShaderMap, Viewport, &ViewUniformBuffer, Instances, ShaderPrintData](bool bGuide, bool bRestTriangle, bool bRestFrame, bool bRestSamples, bool bDeformedTriangle, bool bDeformedFrame, bool bDeformedSamples)
 			{
 				TArray<int32> HairLODIndices;
 				for (FHairStrandsInstance* AbstractInstance : Instances)
@@ -1363,12 +1402,12 @@ void RunHairStrandsDebug(
 
 					const int32 MeshLODIndex = Instance->Debug.MeshLODIndex;
 
-					if (bRestTriangle)		{ AddDebugProjectionHairPass(GraphBuilder, ShaderMap, ShaderPrintData, Viewport, ViewUniformBuffer, MeshLODIndex, RestRootResource, DeformedRootResource, Instance->HairGroupPublicData->VFInput.LocalToWorldTransform, EDebugProjectionHairType::HairTriangle, HairStrandsTriangleType::RestPose); }
-					if (bRestFrame)			{ AddDebugProjectionHairPass(GraphBuilder, ShaderMap, ShaderPrintData, Viewport, ViewUniformBuffer, MeshLODIndex, RestRootResource, DeformedRootResource, Instance->HairGroupPublicData->VFInput.LocalToWorldTransform, EDebugProjectionHairType::HairFrame,    HairStrandsTriangleType::RestPose); }
-					if (bRestSamples)		{ AddDebugProjectionHairPass(GraphBuilder, ShaderMap, ShaderPrintData, Viewport, ViewUniformBuffer, MeshLODIndex, RestRootResource, DeformedRootResource, Instance->HairGroupPublicData->VFInput.LocalToWorldTransform, EDebugProjectionHairType::HairSamples,  HairStrandsTriangleType::RestPose); }
-					if (bDeformedTriangle)	{ AddDebugProjectionHairPass(GraphBuilder, ShaderMap, ShaderPrintData, Viewport, ViewUniformBuffer, MeshLODIndex, RestRootResource, DeformedRootResource, Instance->HairGroupPublicData->VFInput.LocalToWorldTransform, EDebugProjectionHairType::HairTriangle, HairStrandsTriangleType::DeformedPose); }
-					if (bDeformedFrame)		{ AddDebugProjectionHairPass(GraphBuilder, ShaderMap, ShaderPrintData, Viewport, ViewUniformBuffer, MeshLODIndex, RestRootResource, DeformedRootResource, Instance->HairGroupPublicData->VFInput.LocalToWorldTransform, EDebugProjectionHairType::HairFrame,    HairStrandsTriangleType::DeformedPose); }
-					if (bDeformedSamples)	{ AddDebugProjectionHairPass(GraphBuilder, ShaderMap, ShaderPrintData, Viewport, ViewUniformBuffer, MeshLODIndex, RestRootResource, DeformedRootResource, Instance->HairGroupPublicData->VFInput.LocalToWorldTransform, EDebugProjectionHairType::HairSamples,  HairStrandsTriangleType::DeformedPose); }
+					if (bRestTriangle)		{ AddDebugProjectionHairPass(GraphBuilder, LocalView, ShaderMap, ShaderPrintData, Viewport, ViewUniformBuffer, MeshLODIndex, RestRootResource, DeformedRootResource, Instance->HairGroupPublicData->VFInput.LocalToWorldTransform, EDebugProjectionHairType::HairTriangle, HairStrandsTriangleType::RestPose); }
+					if (bRestFrame)			{ AddDebugProjectionHairPass(GraphBuilder, LocalView, ShaderMap, ShaderPrintData, Viewport, ViewUniformBuffer, MeshLODIndex, RestRootResource, DeformedRootResource, Instance->HairGroupPublicData->VFInput.LocalToWorldTransform, EDebugProjectionHairType::HairFrame,    HairStrandsTriangleType::RestPose); }
+					if (bRestSamples)		{ AddDebugProjectionHairPass(GraphBuilder, LocalView, ShaderMap, ShaderPrintData, Viewport, ViewUniformBuffer, MeshLODIndex, RestRootResource, DeformedRootResource, Instance->HairGroupPublicData->VFInput.LocalToWorldTransform, EDebugProjectionHairType::HairSamples,  HairStrandsTriangleType::RestPose); }
+					if (bDeformedTriangle)	{ AddDebugProjectionHairPass(GraphBuilder, LocalView, ShaderMap, ShaderPrintData, Viewport, ViewUniformBuffer, MeshLODIndex, RestRootResource, DeformedRootResource, Instance->HairGroupPublicData->VFInput.LocalToWorldTransform, EDebugProjectionHairType::HairTriangle, HairStrandsTriangleType::DeformedPose); }
+					if (bDeformedFrame)		{ AddDebugProjectionHairPass(GraphBuilder, LocalView, ShaderMap, ShaderPrintData, Viewport, ViewUniformBuffer, MeshLODIndex, RestRootResource, DeformedRootResource, Instance->HairGroupPublicData->VFInput.LocalToWorldTransform, EDebugProjectionHairType::HairFrame,    HairStrandsTriangleType::DeformedPose); }
+					if (bDeformedSamples)	{ AddDebugProjectionHairPass(GraphBuilder, LocalView, ShaderMap, ShaderPrintData, Viewport, ViewUniformBuffer, MeshLODIndex, RestRootResource, DeformedRootResource, Instance->HairGroupPublicData->VFInput.LocalToWorldTransform, EDebugProjectionHairType::HairSamples,  HairStrandsTriangleType::DeformedPose); }
 				}
 			};
 
diff --git a/Engine/Source/Runtime/Renderer/Private/HairStrands/HairStrandsDebug.cpp b/Engine/Source/Runtime/Renderer/Private/HairStrands/HairStrandsDebug.cpp
--- a/Engine/Source/Runtime/Renderer/Private/HairStrands/HairStrandsDebug.cpp
+++ b/Engine/Source/Runtime/Renderer/Private/HairStrands/HairStrandsDebug.cpp
@@ -137,6 +137,11 @@ static void AddPrintLODInfoPass(
 	const FViewInfo& View,
 	const FHairGroupPublicData* Data)
 {
+	if (!ShaderPrint::IsSupported(View.GetShaderPlatform()))
+	{
+		return;
+	}
+
 	// Force ShaderPrint on.
 	ShaderPrint::SetEnabled(true);
 	ShaderPrint::RequestSpaceForCharacters(2000);
@@ -208,7 +213,7 @@ class FHairDebugPrintCS : public FGlobalShader
 	END_SHADER_PARAMETER_STRUCT()
 
 public:
-	static bool ShouldCompilePermutation(const FGlobalShaderPermutationParameters& Parameters) { return IsHairStrandsSupported(EHairStrandsShaderType::Strands, Parameters.Platform); }
+	static bool ShouldCompilePermutation(const FGlobalShaderPermutationParameters& Parameters) { return IsHairStrandsSupported(EHairStrandsShaderType::Strands, Parameters.Platform) && ShaderPrint::IsSupported(Parameters.Platform); }
 	static void ModifyCompilationEnvironment(const FGlobalShaderPermutationParameters& Parameters, FShaderCompilerEnvironment& OutEnvironment)
 	{
 		// Skip optimization for avoiding long compilation time due to large UAV writes
@@ -233,6 +238,11 @@ static void AddDebugHairPrintPass(
 {
 	if (!Scene || !View || !View->HairStrandsViewData.UniformBuffer || !InStencilTexture) return;
 
+	if (!ShaderPrint::IsSupported(View->GetShaderPlatform()))
+	{
+		return;
+	}
+
 	if (!TryEnableShaderDrawAndShaderPrint(*View, MacroGroupResources.MacroGroupCount * 32u, 2000u))
 	{
 		return;
@@ -326,7 +336,7 @@ class FHairDebugShadowCullingCS : public FGlobalShader
 	END_SHADER_PARAMETER_STRUCT()
 
 public:
-	static bool ShouldCompilePermutation(const FGlobalShaderPermutationParameters& Parameters) { return IsHairStrandsSupported(EHairStrandsShaderType::Strands, Parameters.Platform); }
+	static bool ShouldCompilePermutation(const FGlobalShaderPermutationParameters& Parameters) { return IsHairStrandsSupported(EHairStrandsShaderType::Strands, Parameters.Platform) && ShaderPrint::IsSupported(Parameters.Platform); }
 	static void ModifyCompilationEnvironment(const FGlobalShaderPermutationParameters& Parameters, FShaderCompilerEnvironment& OutEnvironment)
 	{
 		// Skip optimization for avoiding long compilation time due to large UAV writes
@@ -345,6 +355,11 @@ static void AddDebugHairShadowCullingPass(
 {
 	if (!Scene || !View || !View->HairStrandsViewData.UniformBuffer) return;
 
+	if (!ShaderPrint::IsSupported(View->GetShaderPlatform()))
+	{
+		return;
+	}
+
 	const uint32 LightCount = View->HairStrandsViewData.DebugData.CullData.DirectionalLights.Num();
 	uint32 InstanceCount = 0;
 	for (auto It : View->HairStrandsViewData.DebugData.CullData.DirectionalLights)
@@ -554,7 +569,7 @@ class FDeepShadowVisualizePS : public FGlobalShader
 
 		SHADER_PARAMETER_STRUCT_REF(FViewUniformShaderParameters, ViewUniformBuffer)
 		RENDER_TARGET_BINDING_SLOTS()
-		END_SHADER_PARAMETER_STRUCT()
+	END_SHADER_PARAMETER_STRUCT()
 
 public:
 	static bool ShouldCompilePermutation(const FGlobalShaderPermutationParameters& Parameters) { return IsHairStrandsSupported(EHairStrandsShaderType::Strands, Parameters.Platform); }
@@ -664,7 +679,7 @@ class FDeepShadowInfoCS : public FGlobalShader
 	END_SHADER_PARAMETER_STRUCT()
 
 public:
-	static bool ShouldCompilePermutation(const FGlobalShaderPermutationParameters& Parameters) { return IsHairStrandsSupported(EHairStrandsShaderType::Strands, Parameters.Platform); }
+	static bool ShouldCompilePermutation(const FGlobalShaderPermutationParameters& Parameters) { return IsHairStrandsSupported(EHairStrandsShaderType::Strands, Parameters.Platform) && ShaderPrint::IsSupported(Parameters.Platform); }
 	static void ModifyCompilationEnvironment(const FGlobalShaderPermutationParameters& Parameters, FShaderCompilerEnvironment& OutEnvironment)
 	{
 		FGlobalShader::ModifyCompilationEnvironment(Parameters, OutEnvironment);
@@ -681,6 +696,11 @@ static void AddDeepShadowInfoPass(
 	const FHairStrandsMacroGroupResources& MacroGroupResources,
 	FRDGTextureRef& OutputTexture)
 {
+	if (!ShaderPrint::IsSupported(View.GetShaderPlatform()))
+	{
+		return;
+	}
+
 	if (DeepShadowResources.TotalAtlasSlotCount == 0)
 	{
 		return;
@@ -737,7 +757,7 @@ class FVoxelVirtualRaymarchingCS : public FGlobalShader
 	END_SHADER_PARAMETER_STRUCT()
 
 public:
-	static bool ShouldCompilePermutation(const FGlobalShaderPermutationParameters& Parameters) { return IsHairStrandsSupported(EHairStrandsShaderType::Strands, Parameters.Platform); }
+	static bool ShouldCompilePermutation(const FGlobalShaderPermutationParameters& Parameters) { return IsHairStrandsSupported(EHairStrandsShaderType::Strands, Parameters.Platform) && ShaderPrint::IsSupported(Parameters.Platform); }
 	static void ModifyCompilationEnvironment(const FGlobalShaderPermutationParameters& Parameters, FShaderCompilerEnvironment& OutEnvironment)
 	{
 		// Skip optimization for avoiding long compilation time due to large UAV writes
@@ -756,6 +776,11 @@ static void AddVoxelPageRaymarchingPass(
 	const FHairStrandsVoxelResources& VoxelResources,
 	FRDGTextureRef& OutputTexture)
 {
+	if (!ShaderPrint::IsSupported(View.GetShaderPlatform()))
+	{
+		return;
+	}
+
 	if (!TryEnableShaderDrawAndShaderPrint(View, 4000, 2000))
 	{
 		return;
@@ -810,7 +835,7 @@ class FDebugHairTangentCS : public FGlobalShader
 	END_SHADER_PARAMETER_STRUCT()
 
 public:
-	static bool ShouldCompilePermutation(const FGlobalShaderPermutationParameters& Parameters) { return IsHairStrandsSupported(EHairStrandsShaderType::Strands, Parameters.Platform); }
+	static bool ShouldCompilePermutation(const FGlobalShaderPermutationParameters& Parameters) { return IsHairStrandsSupported(EHairStrandsShaderType::Strands, Parameters.Platform) && ShaderPrint::IsSupported(Parameters.Platform); }
 	static void ModifyCompilationEnvironment(const FGlobalShaderPermutationParameters& Parameters, FShaderCompilerEnvironment& OutEnvironment)
 	{
 		FGlobalShader::ModifyCompilationEnvironment(Parameters, OutEnvironment);
@@ -826,6 +851,11 @@ static void AddDebugHairTangentPass(
 	const FSceneTextures& SceneTextures,
 	FRDGTextureRef& OutputTexture)
 {
+	if (!ShaderPrint::IsSupported(View.GetShaderPlatform()))
+	{
+		return;
+	}
+
 	if (!TryEnableShaderDrawAndShaderPrint(View, 64000u, 1000u))
 	{
 		return;
@@ -1056,7 +1086,7 @@ class FHairVisibilityDebugPPLLCS : public FGlobalShader
 
 	static bool ShouldCompilePermutation(const FGlobalShaderPermutationParameters& Parameters)
 	{
-		return IsHairStrandsSupported(EHairStrandsShaderType::Tool, Parameters.Platform);
+		return IsHairStrandsSupported(EHairStrandsShaderType::Tool, Parameters.Platform) && ShaderPrint::IsSupported(Parameters.Platform);
 	}
 
 	static void ModifyCompilationEnvironment(const FGlobalShaderPermutationParameters& Parameters, FShaderCompilerEnvironment& OutEnvironment)
@@ -1078,6 +1108,11 @@ static void InternalRenderHairStrandsDebugInfo(
 	FRDGTextureRef SceneColorTexture,
 	FRDGTextureRef SceneDepthTexture)
 {
+	if (!ShaderPrint::IsSupported(View.GetShaderPlatform()))
+	{
+		return;
+	}
+
 	FHairStrandsBookmarkParameters Params;
 	CreateHairStrandsBookmarkParameters(Scene, View, Params);
 	Params.SceneColorTexture = SceneColorTexture;
