Remove old procedural cards generation.

diff --git a/Engine/Plugins/Experimental/HairModelingToolset/Source/HairModelingToolset/Private/GroomQueryUtil.cpp b/Engine/Plugins/Experimental/HairModelingToolset/Source/HairModelingToolset/Private/GroomQueryUtil.cpp
--- a/Engine/Plugins/Experimental/HairModelingToolset/Source/HairModelingToolset/Private/GroomQueryUtil.cpp
+++ b/Engine/Plugins/Experimental/HairModelingToolset/Source/HairModelingToolset/Private/GroomQueryUtil.cpp
@@ -53,14 +53,7 @@ void UE::GroomQueries::ExtractAllHairCards(AGroomActor* GroomActor,
 		{
 			if (Desc.GroupIndex == GroupIdx && Desc.LODIndex == LODIndex)
 			{
-				if (Desc.SourceType == EHairCardsSourceType::Imported)
-				{
-					StaticMesh = Desc.ImportedMesh;
-				}
-				else if (Desc.SourceType == EHairCardsSourceType::Procedural)
-				{
-					StaticMesh = Desc.ProceduralMesh;
-				}
+				StaticMesh = Desc.ImportedMesh;
 			}
 		}
 		
diff --git a/Engine/Plugins/Runtime/HairStrands/Source/HairStrandsCore/Private/GroomAsset.cpp b/Engine/Plugins/Runtime/HairStrands/Source/HairStrandsCore/Private/GroomAsset.cpp
--- a/Engine/Plugins/Runtime/HairStrands/Source/HairStrandsCore/Private/GroomAsset.cpp
+++ b/Engine/Plugins/Runtime/HairStrands/Source/HairStrandsCore/Private/GroomAsset.cpp
@@ -1294,6 +1294,19 @@ void UGroomAsset::PostLoad()
 		}
 	}
 
+	// Convert old procedural cards to import cards
+	for (FHairGroupsCardsSourceDescription& Group : GetHairGroupsCards())
+	{
+		if (Group.SourceType_DEPRECATED == EHairCardsSourceType::Procedural)
+		{
+			Group.bInvertUV = true;
+			Group.SourceType_DEPRECATED = EHairCardsSourceType::Imported;
+			Group.ImportedMesh = Group.ProceduralMesh_DEPRECATED;
+			Group.ProceduralMesh_DEPRECATED = nullptr;
+			Group.ImportedMeshKey = FString();
+		}
+	}
+
 #if WITH_EDITORONLY_DATA
 	bool bSucceed = true;
 	{
@@ -1553,27 +1566,6 @@ static bool IsCardsTextureResources(const FName PropertyName)
 }
 static void InitCardsTextureResources(UGroomAsset* GroomAsset);
 
-static bool IsCardsProceduralAttributes(const FName PropertyName)
-{	
-	return
-		   PropertyName == GET_MEMBER_NAME_CHECKED(FHairCardsClusterSettings, ClusterDecimation)
-		|| PropertyName == GET_MEMBER_NAME_CHECKED(FHairCardsClusterSettings, Type)
-		|| PropertyName == GET_MEMBER_NAME_CHECKED(FHairCardsClusterSettings, bUseGuide)
-
-		|| PropertyName == GET_MEMBER_NAME_CHECKED(FHairCardsTextureSettings, AtlasMaxResolution)
-		|| PropertyName == GET_MEMBER_NAME_CHECKED(FHairCardsTextureSettings, PixelPerCentimeters)
-		|| PropertyName == GET_MEMBER_NAME_CHECKED(FHairCardsTextureSettings, LengthTextureCount)
-		|| PropertyName == GET_MEMBER_NAME_CHECKED(FHairCardsTextureSettings, DensityTextureCount)
-
-		|| PropertyName == GET_MEMBER_NAME_CHECKED(FHairCardsGeometrySettings, GenerationType)
-		|| PropertyName == GET_MEMBER_NAME_CHECKED(FHairCardsGeometrySettings, CardsCount)
-		|| PropertyName == GET_MEMBER_NAME_CHECKED(FHairCardsGeometrySettings, ClusterType)
-		|| PropertyName == GET_MEMBER_NAME_CHECKED(FHairCardsGeometrySettings, MinSegmentLength)
-		|| PropertyName == GET_MEMBER_NAME_CHECKED(FHairCardsGeometrySettings, AngularThreshold)
-		|| PropertyName == GET_MEMBER_NAME_CHECKED(FHairCardsGeometrySettings, MinCardsLength)
-		|| PropertyName == GET_MEMBER_NAME_CHECKED(FHairCardsGeometrySettings, MaxCardsLength);
-}
-
 static bool IsStrandsInterpolationAttributes(const FName PropertyName)
 {
 	return
@@ -1661,8 +1653,7 @@ void UGroomAsset::PostEditChangeProperty(FPropertyChangedEvent& PropertyChangedE
 		{
 			GetHairGroupsCards().Init(FHairGroupsCardsSourceDescription(), GroupCount);
 			FHairGroupsCardsSourceDescription Dirty;
-			Dirty.ProceduralSettings.ClusterSettings.ClusterDecimation = 0;
-			Dirty.SourceType = HasImportedStrandsData() ? EHairCardsSourceType::Procedural : EHairCardsSourceType::Imported;
+			Dirty.ImportedMesh = nullptr;
 			CachedHairGroupsCards.Init(Dirty, GroupCount);
 		}
 
@@ -1685,22 +1676,11 @@ void UGroomAsset::PostEditChangeProperty(FPropertyChangedEvent& PropertyChangedE
 	{
 		CacheDerivedDatas();
 	}
-	
-	const bool bCardsArrayChanged = PropertyName == UGroomAsset::GetHairGroupsCardsMemberName();
-	if (bCardsArrayChanged && PropertyChangedEvent.ChangeType == EPropertyChangeType::ArrayAdd)
-	{
-		GetHairGroupsCards().Last().SourceType = HasImportedStrandsData() ? EHairCardsSourceType::Procedural : EHairCardsSourceType::Imported;
-	}
-
-	SavePendingProceduralAssets();
 
 	const bool bHairStrandsRaytracingRadiusChanged = PropertyName == GET_MEMBER_NAME_CHECKED(FHairShadowSettings, HairRaytracingRadiusScale);
 	
-	// By pass update for all procedural cards parameters, as we don't want them to invalidate the cards data. 
-	// Cards should be refresh only under user action
 	// By pass update if bStrandsInterpolationChanged has the resources have already been recreated
-	const bool bCardsToolUpdate = IsCardsProceduralAttributes(PropertyName);
-	if (!bCardsToolUpdate && !bNeedRebuildDerivedData)
+	if (!bNeedRebuildDerivedData)
 	{
 		FGroomComponentRecreateRenderStateContext Context(this);
 		UpdateResource();
@@ -1738,7 +1718,7 @@ void UGroomAsset::PostEditChangeProperty(FPropertyChangedEvent& PropertyChangedE
 		// Delegate used for notifying groom data & groom resoures invalidation
 		OnGroomAssetResourcesChanged.Broadcast();
 	}
-	else if (!bCardsToolUpdate)
+	else
 	{
 		// Delegate used for notifying groom data invalidation
 		OnGroomAssetChanged.Broadcast();
@@ -2189,28 +2169,11 @@ namespace GroomDerivedDataCacheUtils
 
 		Ar << Desc.GroupIndex;
 		Ar << Desc.LODIndex;
-		if (Desc.SourceType == EHairCardsSourceType::Imported)
-		{
-			if (Desc.ImportedMesh)
-			{
-				FString Key = Desc.GetMeshKey();
-				Ar << Key;
-			}
-		}
-		else if (Desc.SourceType == EHairCardsSourceType::Procedural)
+		if (Desc.ImportedMesh)
 		{
-			if (Desc.ProceduralMesh)
-			{
-				FString Key = Desc.GetMeshKey();
-				Ar << Key;
-			}
-			Desc.ProceduralSettings.BuildDDCKey(Ar);
-
-			if (Desc.GenerationSettings)
-			{
-				Desc.GenerationSettings->BuildDDCKey(Ar);
-			}
-		}
+			FString Key = Desc.GetMeshKey();
+			Ar << Key;
+		}		
 
 		FSHAHash Hash;
 		FSHA1::HashBuffer(TempBytes.GetData(), TempBytes.Num(), Hash.Hash);
@@ -2245,22 +2208,10 @@ namespace GroomDerivedDataCacheUtils
 			}
 
 			Ar << Desc.LODIndex;
-			if (Desc.SourceType == EHairCardsSourceType::Imported)
-			{
-				if (Desc.ImportedMesh)
-				{
-					FString Key = Desc.GetMeshKey();
-					Ar << Key;
-				}
-			}
-			else if (Desc.SourceType == EHairCardsSourceType::Procedural)
+			if (Desc.ImportedMesh)
 			{
-				if (Desc.ProceduralMesh)
-				{
-					FString Key = Desc.GetMeshKey();
-					Ar << Key;
-				}
-				Desc.ProceduralSettings.BuildDDCKey(Ar);
+				FString Key = Desc.GetMeshKey();
+				Ar << Key;
 			}
 			// Material is not included as it doesn't affect the data building
 		}
@@ -2755,19 +2706,7 @@ bool UGroomAsset::CacheCardsData(uint32 GroupIndex, const FString& StrandsKey)
 		{
 			// Note: We don't condition/filter resource allocation based on GeometryType == EGroomGeometryType::Cards, because 
 			// geometry type can be switched at runtime (between cards/strands)
-			{
-				UStaticMesh* CardsMesh = nullptr;
-				if (Desc->SourceType == EHairCardsSourceType::Procedural)
-				{
-					CardsMesh = Desc->ProceduralMesh;
-				}
-				else if (Desc->SourceType == EHairCardsSourceType::Imported)
-				{
-					CardsMesh = Desc->ImportedMesh;
-				}
-				
-				bDataCanBeBuilt |= CardsMesh != nullptr;
-			}
+			bDataCanBeBuilt |= Desc->ImportedMesh != nullptr;
 		}
 	}
 
@@ -3018,19 +2957,9 @@ bool UGroomAsset::BuildCardsData(uint32 GroupIndex)
 			InternalReleaseResource(LOD.Guides.InterpolationResource);
 
 			// 1. Load geometry data, if any
-			UStaticMesh* CardsMesh = nullptr;
-			if (Desc->SourceType == EHairCardsSourceType::Procedural)
-			{
-				CardsMesh = Desc->ProceduralMesh;
-			}
-			else if (Desc->SourceType == EHairCardsSourceType::Imported)
-			{
-				CardsMesh = Desc->ImportedMesh;
-			}
-
 			bool bInitResources = false;
 			FHairStrandsDatas LODGuidesData;
-			if (CardsMesh != nullptr)
+			if (UStaticMesh* CardsMesh = Desc->ImportedMesh)
 			{
 				CardsMesh->ConditionalPostLoad();
 
@@ -3077,7 +3006,7 @@ bool UGroomAsset::BuildCardsData(uint32 GroupIndex)
 				InitAtlasTexture(LOD.RestResource, Desc->Textures.CoverageTexture, EHairAtlasTextureType::Coverage);
 				InitAtlasTexture(LOD.RestResource, Desc->Textures.AuxilaryDataTexture, EHairAtlasTextureType::AuxilaryData);
 				InitAtlasTexture(LOD.RestResource, Desc->Textures.MaterialTexture, EHairAtlasTextureType::Material);
-				LOD.RestResource->bInvertUV = Desc->SourceType == EHairCardsSourceType::Procedural;
+				LOD.RestResource->bInvertUV = Desc->bInvertUV;
 				
 				// 2.2 Load interoplatino resources
 				LOD.InterpolationResource = new FHairCardsInterpolationResource(LOD.InterpolationBulkData, ResourceName, OwnerName);
@@ -3486,7 +3415,7 @@ static void InitCardsTextureResources(UGroomAsset* GroomAsset)
 					InitAtlasTexture(LOD.RestResource, Desc->Textures.MaterialTexture, EHairAtlasTextureType::Material);
 					if (LOD.RestResource)
 					{
-						LOD.RestResource->bInvertUV = Desc->SourceType == EHairCardsSourceType::Procedural; // Should fix procedural texture so that this does not happen
+						LOD.RestResource->bInvertUV = Desc->bInvertUV; // Should fix procedural texture so that this does not happen
 					}
 				}
 			}
@@ -3544,7 +3473,7 @@ void UGroomAsset::InitCardsResources()
 					InitAtlasTexture(LOD.RestResource, Desc->Textures.CoverageTexture, EHairAtlasTextureType::Coverage);
 					InitAtlasTexture(LOD.RestResource, Desc->Textures.AuxilaryDataTexture, EHairAtlasTextureType::AuxilaryData);
 					InitAtlasTexture(LOD.RestResource, Desc->Textures.MaterialTexture, EHairAtlasTextureType::Material);
-					LOD.RestResource->bInvertUV = Desc->SourceType == EHairCardsSourceType::Procedural; // Should fix procedural texture so that this does not happen
+					LOD.RestResource->bInvertUV = Desc->bInvertUV; // Should fix procedural texture so that this does not happen
 				}
 			}
 
@@ -3930,162 +3859,6 @@ bool UGroomAsset::IsMaterialUsed(int32 MaterialIndex) const
 		InternalIsMaterialUsed(GetHairGroupsMeshes(), MaterialSlotName);
 }
 
-#if WITH_EDITOR
-struct FHairProceduralCardsQuery
-{
-	FHairCardsInterpolationBulkData InterpolationBulkData;
-	FHairCardsProceduralDatas ProceduralData;
-	FHairStrandsDatas GuideData;
-	FHairCardsDatas Data;
-	FHairCardsBulkData BulkData;
-
-	UGroomAsset* Asset = nullptr;
-	FHairCardsRestResource* Resources = nullptr;
-	FHairCardsProceduralResource* ProceduralResources = nullptr;
-	FHairGroupCardsTextures* Textures = nullptr;
-};
-
-// Hair_TODO: move this into the groom asset class
-static TQueue<FHairProceduralCardsQuery*> HairCardsQueuries;
-
-// hair_TODO: Rename into GenerateProceduralCards
-// Generate geometry and textures for hair cards
-void UGroomAsset::SaveProceduralCards(uint32 DescIndex)
-{
-	if (!IsHairStrandsEnabled(EHairStrandsShaderType::Cards) || !CanRebuildFromDescription())
-	{
-		return;
-	}
-
-	LLM_SCOPE_BYTAG(Groom);
-
-	if (DescIndex >= uint32(GetHairGroupsCards().Num()))
-		return;
-
-	FHairGroupsCardsSourceDescription* Desc = &GetHairGroupsCards()[DescIndex];
-
-	const int32 GroupIndex = Desc->GroupIndex;
-	const int32 LODIndex = Desc->LODIndex;
-	if (GroupIndex >= GetHairGroupsPlatformData().Num())
-		return;
-
-	// 1. Convert old parameters (ClusterDecimation, bUseCards) to new parameters (GenerationType & CardsCount)
-	{
-		FHairCardsClusterSettings& ClusterSettings = Desc->ProceduralSettings.ClusterSettings;
-		FHairCardsGeometrySettings& GeometrySettings = Desc->ProceduralSettings.GeometrySettings;
-		const bool bNeedConversion = ClusterSettings.ClusterDecimation > 0;
-		if (bNeedConversion)
-		{
-			const int32 MaxCardCount = GetHairGroupsPlatformData()[GroupIndex].Strands.BulkData.GetNumCurves();
-
-			GeometrySettings.GenerationType = ClusterSettings.bUseGuide ? EHairCardsGenerationType::UseGuides : EHairCardsGenerationType::CardsCount;
-			GeometrySettings.CardsCount = FMath::Clamp(FMath::CeilToInt(ClusterSettings.ClusterDecimation * MaxCardCount), 1, MaxCardCount);
-			GeometrySettings.ClusterType = ClusterSettings.Type;
-
-			// Mark the asset as updated.
-			ClusterSettings.ClusterDecimation = 0;
-		}
-	}
-
-	// 2. Generate geometry (CPU)
-	FHairProceduralCardsQuery* QP = new FHairProceduralCardsQuery();
-	HairCardsQueuries.Enqueue(QP);
-
-	FHairProceduralCardsQuery& Q = *QP;
-	Q.Asset = this;
-	{
-		// 2.1 Generate transient strands data
-		const FHairDescriptionGroups& LocalHairDescriptionGroups = GetHairDescriptionGroups();
-		FHairGroupInfo DummyInfo;
-		FHairStrandsDatas StrandsData;
-		FHairStrandsDatas GuidesData;
-		FGroomBuilder::BuildData(LocalHairDescriptionGroups.HairGroups[GroupIndex], GetHairGroupsInterpolation()[GroupIndex], DummyInfo, StrandsData, GuidesData);
-
-		// 2.2 Build cards geometry
-		FHairCardsBuilder::BuildGeometry(
-			GetLODName(this, Desc->LODIndex),
-			StrandsData,
-			GuidesData,
-			Desc->ProceduralSettings,
-			Q.ProceduralData,
-			Q.BulkData,
-			Q.GuideData,
-			Q.InterpolationBulkData,
-			Desc->Textures);
-		Q.Textures = &Desc->Textures;
-	}
-
-	const FName OwnerName = GetAssetPathName(LODIndex);
-	// 3. Create resources and enqueue texture generation (GPU, kicked by the render thread) 
-	Q.Resources = new FHairCardsRestResource(Q.BulkData, FHairResourceName(GetFName(), GroupIndex, LODIndex), OwnerName);
-	BeginInitResource(Q.Resources); // Immediate allocation, as needed for the vertex factory, input stream building
-	Q.ProceduralResources = new FHairCardsProceduralResource(Q.ProceduralData.RenderData, Q.ProceduralData.Atlas.Resolution, Q.ProceduralData.Voxels, OwnerName);
-	BeginInitResource(Q.ProceduralResources); // Immediate allocation, as needed for the vertex factory, input stream building
-
-	FHairCardsBuilder::BuildTextureAtlas(&Q.ProceduralData, Q.Resources, Q.ProceduralResources, Q.Textures);
-
-	// 4. Save output asset (geometry, and enqueue texture saving)
-	{
-		// Create a static meshes with the vertex data
-		if (GetHairGroupsCards()[DescIndex].ProceduralMesh == nullptr)
-		{
-			const FString PackageName = GetOutermost()->GetName();
-			const FString SuffixName = FText::Format(LOCTEXT("CardsStatisMesh", "_CardsMesh_Group{0}_LOD{1}"), FText::AsNumber(GroupIndex), FText::AsNumber(LODIndex)).ToString();
-			GetHairGroupsCards()[DescIndex].ProceduralMesh = FHairStrandsCore::CreateStaticMesh(PackageName, SuffixName);
-		}
-
-		// Convert procedural cards data to cards data prior to export
-		FHairCardsDatas CardData;
-		{
-			CardData.Cards = Q.ProceduralData.Cards;
-		}
-
-		FHairCardsBuilder::ExportGeometry(CardData, GetHairGroupsCards()[DescIndex].ProceduralMesh);
-		FHairStrandsCore::SaveAsset(GetHairGroupsCards()[DescIndex].ProceduralMesh);
-		GetHairGroupsCards()[DescIndex].ProceduralMeshKey = GroomDerivedDataCacheUtils::BuildCardsDerivedDataKeySuffix(this, *Desc);
-	}
-}
-
-// Save geometry and textures for hair cards
-// Save out a static mesh based on generated cards
-void UGroomAsset::SavePendingProceduralAssets()
-{
-	// Proceed procedural asset which needs to be saved
-	if (!HairCardsQueuries.IsEmpty())
-	{
-		TQueue<FHairProceduralCardsQuery*> NotReady;
-		FHairProceduralCardsQuery* Q = nullptr;
-		while (HairCardsQueuries.Dequeue(Q))
-		{
-			if (Q)
-			{
-				if (Q->Asset == this && Q->Textures->bNeedToBeSaved)
-				{
-					if (Q->Textures->DepthTexture)			FHairStrandsCore::SaveAsset(Q->Textures->DepthTexture);
-					if (Q->Textures->AttributeTexture)		FHairStrandsCore::SaveAsset(Q->Textures->AttributeTexture);
-					if (Q->Textures->AuxilaryDataTexture)	FHairStrandsCore::SaveAsset(Q->Textures->AuxilaryDataTexture);
-					if (Q->Textures->CoverageTexture)		FHairStrandsCore::SaveAsset(Q->Textures->CoverageTexture);
-					if (Q->Textures->TangentTexture)		FHairStrandsCore::SaveAsset(Q->Textures->TangentTexture);
-					if (Q->Textures->MaterialTexture)		FHairStrandsCore::SaveAsset(Q->Textures->MaterialTexture);
-					Q->Textures->bNeedToBeSaved = false;
-
-					InternalReleaseResource(Q->Resources);
-					InternalReleaseResource(Q->ProceduralResources);
-				}
-				else
-				{
-					NotReady.Enqueue(Q);
-				}
-			}
-		}
-		while (NotReady.Dequeue(Q))
-		{
-			HairCardsQueuries.Enqueue(Q);
-		}
-	}
-}
-#endif // WITH_EDITOR
-
 #if WITH_EDITORONLY_DATA
 bool UGroomAsset::GetHairStrandsDatas(
 	const int32 GroupIndex,
@@ -4126,17 +3899,7 @@ bool UGroomAsset::GetHairCardsGuidesDatas(
 	if (FHairGroupsCardsSourceDescription* Desc = GetSourceDescription(GetHairGroupsCards(), GroupIndex, LODIndex, SourceIt))
 	{
 		// 1. Load geometry data, if any
-		UStaticMesh* CardsMesh = nullptr;
-		if (Desc->SourceType == EHairCardsSourceType::Procedural)
-		{
-			CardsMesh = Desc->ProceduralMesh;
-		}
-		else if (Desc->SourceType == EHairCardsSourceType::Imported)
-		{
-			CardsMesh = Desc->ImportedMesh;
-		}
-
-		if (CardsMesh)
+		if (UStaticMesh* CardsMesh = Desc->ImportedMesh)
 		{
 			CardsMesh->ConditionalPostLoad();
 
diff --git a/Engine/Plugins/Runtime/HairStrands/Source/HairStrandsCore/Private/GroomAssetCards.cpp b/Engine/Plugins/Runtime/HairStrands/Source/HairStrandsCore/Private/GroomAssetCards.cpp
--- a/Engine/Plugins/Runtime/HairStrands/Source/HairStrandsCore/Private/GroomAssetCards.cpp
+++ b/Engine/Plugins/Runtime/HairStrands/Source/HairStrandsCore/Private/GroomAssetCards.cpp
@@ -7,88 +7,14 @@
 
 #include UE_INLINE_GENERATED_CPP_BY_NAME(GroomAssetCards)
 
-FHairCardsClusterSettings::FHairCardsClusterSettings()
-{
-	ClusterDecimation = 0.1f;
-	Type = EHairCardsClusterType::High;
-	bUseGuide = true;
-}
-
-FHairCardsGeometrySettings::FHairCardsGeometrySettings()
-{
-	GenerationType = EHairCardsGenerationType::UseGuides;
-	CardsCount = 5000;
-	ClusterType = EHairCardsClusterType::High;
-	MinSegmentLength = 1;
-	AngularThreshold = 5;
-	MinCardsLength = 0;
-	MaxCardsLength = 0;
-};
-
-FHairCardsTextureSettings::FHairCardsTextureSettings()
-{
-	AtlasMaxResolution = 2048;
-	PixelPerCentimeters = 60;
-	LengthTextureCount = 10;
-	DensityTextureCount = 1;
-};
-
-FHairGroupsProceduralCards::FHairGroupsProceduralCards()
-{
-	ClusterSettings	= FHairCardsClusterSettings();
-	GeometrySettings= FHairCardsGeometrySettings();
-	TextureSettings	= FHairCardsTextureSettings();
-	Version = 0;
-}
-
 FHairGroupsCardsSourceDescription::FHairGroupsCardsSourceDescription()
 {
 	MaterialSlotName = NAME_None;
-	SourceType = EHairCardsSourceType::Imported;
-	ProceduralMeshKey.Empty();
-	ProceduralSettings = FHairGroupsProceduralCards();
+	SourceType_DEPRECATED = EHairCardsSourceType::Imported;
 	GroupIndex = 0;
 	LODIndex = -1;
 }
 
-bool FHairCardsClusterSettings::operator==(const FHairCardsClusterSettings& A) const
-{
-	return
-		Type == A.Type &&
-		ClusterDecimation == A.ClusterDecimation &&
-		bUseGuide == A.bUseGuide;
-}
-
-bool FHairCardsGeometrySettings::operator==(const FHairCardsGeometrySettings& A) const
-{
-	return
-		MinSegmentLength == A.MinSegmentLength &&
-		AngularThreshold == A.AngularThreshold &&
-		MinCardsLength == A.MinCardsLength &&
-		MaxCardsLength == A.MaxCardsLength &&
-		GenerationType == A.GenerationType &&
-		CardsCount == A.CardsCount &&
-		ClusterType == A.ClusterType;
-}
-
-bool FHairCardsTextureSettings::operator==(const FHairCardsTextureSettings& A) const
-{
-	return
-		AtlasMaxResolution == A.AtlasMaxResolution &&
-		PixelPerCentimeters == A.PixelPerCentimeters &&
-		LengthTextureCount == A.LengthTextureCount &&
-		DensityTextureCount == A.DensityTextureCount;
-}
-
-bool FHairGroupsProceduralCards::operator==(const FHairGroupsProceduralCards& A) const
-{
-	return
-		ClusterSettings == A.ClusterSettings &&
-		GeometrySettings == A.GeometrySettings &&
-		TextureSettings == A.TextureSettings &&
-		Version == A.Version;
-}
-
 void FHairGroupCardsTextures::SetTexture(EHairAtlasTextureType SlotID, UTexture2D* Texture)
 {
 	switch (SlotID)
@@ -115,31 +41,10 @@ void FHairGroupCardsTextures::SetTexture(EHairAtlasTextureType SlotID, UTexture2
 	};
 }
 
-void FHairGroupsProceduralCards::BuildDDCKey(FArchive& Ar)
-{
-	Ar << GeometrySettings.GenerationType;
-	if (GeometrySettings.GenerationType == EHairCardsGenerationType::CardsCount)
-	{
-		Ar << GeometrySettings.CardsCount;
-	}
-	Ar << GeometrySettings.ClusterType;
-	Ar << GeometrySettings.MinSegmentLength;
-	Ar << GeometrySettings.AngularThreshold;
-	Ar << GeometrySettings.MinCardsLength;
-	Ar << GeometrySettings.MaxCardsLength;
-
-	Ar << TextureSettings.AtlasMaxResolution;
-	Ar << TextureSettings.PixelPerCentimeters;
-	Ar << TextureSettings.LengthTextureCount;
-	Ar << Version;
-}
-
 bool FHairGroupsCardsSourceDescription::operator==(const FHairGroupsCardsSourceDescription& A) const
 {
 	return
 		MaterialSlotName == A.MaterialSlotName &&
-		SourceType == A.SourceType &&
-		ProceduralSettings == A.ProceduralSettings &&
 		GroupIndex == A.GroupIndex &&
 		LODIndex == A.LODIndex &&
 		ImportedMesh == A.ImportedMesh;
@@ -164,15 +69,10 @@ FString FHairGroupsCardsSourceDescription::GetMeshKey() const
 bool FHairGroupsCardsSourceDescription::HasMeshChanged() const
 {
 #if WITH_EDITORONLY_DATA
-	if (SourceType == EHairCardsSourceType::Imported && ImportedMesh)
+	if (ImportedMesh)
 	{
 		return ImportedMeshKey == GetMeshKey();
 	}
-	else if (SourceType == EHairCardsSourceType::Procedural && ProceduralMesh)
-	{
-		ProceduralMesh->ConditionalPostLoad();
-		return ProceduralMeshKey == GetMeshKey();
-	}
 #endif
 	return false;
 }
@@ -180,31 +80,20 @@ bool FHairGroupsCardsSourceDescription::HasMeshChanged() const
 void FHairGroupsCardsSourceDescription::UpdateMeshKey()
 {
 #if WITH_EDITORONLY_DATA
-	if (SourceType == EHairCardsSourceType::Imported && ImportedMesh)
+	if (ImportedMesh)
 	{
 		ImportedMeshKey = GetMeshKey();
 	}
-	else if (SourceType == EHairCardsSourceType::Procedural && ProceduralMesh)
-	{
-		ProceduralMeshKey = GetMeshKey();
-	}
 #endif
 }
 
 UStaticMesh* FHairGroupsCardsSourceDescription::GetMesh() const
 {
 #if WITH_EDITORONLY_DATA
-	if (SourceType == EHairCardsSourceType::Imported)
-	{
-		return ImportedMesh;
-
-	}
-	else if (SourceType == EHairCardsSourceType::Procedural)
-	{
-		return ProceduralMesh;
-	}
-#endif
+	return ImportedMesh;
+#else
 	return nullptr;
+#endif
 }
 
 
diff --git a/Engine/Plugins/Runtime/HairStrands/Source/HairStrandsCore/Private/GroomManager.cpp b/Engine/Plugins/Runtime/HairStrands/Source/HairStrandsCore/Private/GroomManager.cpp
--- a/Engine/Plugins/Runtime/HairStrands/Source/HairStrandsCore/Private/GroomManager.cpp
+++ b/Engine/Plugins/Runtime/HairStrands/Source/HairStrandsCore/Private/GroomManager.cpp
@@ -1245,9 +1245,6 @@ void RunHairStrandsTexturesQueries(FRDGBuilder& GraphBuilder, FGlobalShaderMap*
 
 bool HasHairStrandsPositionQueries();
 void RunHairStrandsPositionQueries(FRDGBuilder& GraphBuilder, FGlobalShaderMap* ShaderMap, const struct FShaderPrintData* DebugShaderData);
-
-bool HasHairCardsAtlasQueries();
-void RunHairCardsAtlasQueries(FRDGBuilder& GraphBuilder, FGlobalShaderMap* ShaderMap, const struct FShaderPrintData* ShaderPrintData);
 #endif
 
 static void RunHairStrandsProcess(FRDGBuilder& GraphBuilder, FGlobalShaderMap* ShaderMap, const struct FShaderPrintData* ShaderPrintData)
@@ -1268,13 +1265,6 @@ static void RunHairStrandsProcess(FRDGBuilder& GraphBuilder, FGlobalShaderMap* S
 	{
 		RunHairStrandsFolliculeMaskQueries(GraphBuilder, ShaderMap);
 	}
-
-#if WITH_EDITOR
-	if (HasHairCardsAtlasQueries())
-	{
-		RunHairCardsAtlasQueries(GraphBuilder, ShaderMap, ShaderPrintData);
-	}
-#endif
 }
 
 //////////////////////////////////////////////////////////////////////////////////////////////////////////////////////
@@ -1335,7 +1325,6 @@ void ProcessHairStrandsBookmark(
 	{
 		const bool bHasHairStardsnProcess =
 		#if WITH_EDITOR
-			HasHairCardsAtlasQueries() ||
 			HasHairStrandsTexturesQueries() ||
 			HasHairStrandsPositionQueries() ||
 		#endif
diff --git a/Engine/Plugins/Runtime/HairStrands/Source/HairStrandsCore/Private/GroomRBFDeformer.cpp b/Engine/Plugins/Runtime/HairStrands/Source/HairStrandsCore/Private/GroomRBFDeformer.cpp
--- a/Engine/Plugins/Runtime/HairStrands/Source/HairStrandsCore/Private/GroomRBFDeformer.cpp
+++ b/Engine/Plugins/Runtime/HairStrands/Source/HairStrandsCore/Private/GroomRBFDeformer.cpp
@@ -772,28 +772,17 @@ void FGroomRBFDeformer::GetRBFDeformedGroomAsset(const UGroomAsset* InGroomAsset
 		// Apply changes onto cards and meshes (OutGroomASset already contain duplicated mesh asset
 		for (FHairGroupsCardsSourceDescription& Desc : OutGroomAsset->GetHairGroupsCards())
 		{
-			UStaticMesh* Mesh = nullptr;
-			if (Desc.SourceType == EHairCardsSourceType::Procedural)
-			{
-				Mesh = Desc.ProceduralMesh;
-			}
-			else if (Desc.SourceType == EHairCardsSourceType::Imported)
-			{
-				Mesh = Desc.ImportedMesh;
-			}
-			if (!Mesh)
-			{
-				continue;
-			}
-
-			if (Desc.GroupIndex >= 0)
+			if (UStaticMesh* Mesh = Desc.ImportedMesh)
 			{
-				Mesh->ConditionalPostLoad();
-
-				FHairStrandsRootData RenRootData;
-				FGroomBindingBuilder::GetRootData(RenRootData, BindingAsset->GetHairGroupsPlatformData()[Desc.GroupIndex].RenRootBulkData);
-
-				DeformStaticMeshPositions(Mesh, MeshVertexPositionsBuffer_Target, RenRootData.MeshProjectionLODs[MeshLODIndex]);
+				if (Desc.GroupIndex >= 0)
+				{
+					Mesh->ConditionalPostLoad();
+	
+					FHairStrandsRootData RenRootData;
+					FGroomBindingBuilder::GetRootData(RenRootData, BindingAsset->GetHairGroupsPlatformData()[Desc.GroupIndex].RenRootBulkData);
+	
+					DeformStaticMeshPositions(Mesh, MeshVertexPositionsBuffer_Target, RenRootData.MeshProjectionLODs[MeshLODIndex]);
+				}
 			}
 		} 
 
diff --git a/Engine/Plugins/Runtime/HairStrands/Source/HairStrandsCore/Private/GroomResources.cpp b/Engine/Plugins/Runtime/HairStrands/Source/HairStrandsCore/Private/GroomResources.cpp
--- a/Engine/Plugins/Runtime/HairStrands/Source/HairStrandsCore/Private/GroomResources.cpp
+++ b/Engine/Plugins/Runtime/HairStrands/Source/HairStrandsCore/Private/GroomResources.cpp
@@ -877,56 +877,6 @@ void FHairCardsRestResource::ReleaseResource()
 	MaterialsBuffer.ReleaseResource();
 }
 
-/////////////////////////////////////////////////////////////////////////////////////////
-FHairCardsProceduralResource::FHairCardsProceduralResource(const FHairCardsProceduralDatas::FRenderData& InRenderData, const FIntPoint& InAtlasResolution, const FHairCardsVoxel& InVoxel, const FName& InOwnerName):
-	FHairCommonResource(EHairStrandsAllocationType::Immediate, FHairResourceName(), InOwnerName),
-	CardBoundCount(InRenderData.ClusterBounds.Num()),
-	AtlasResolution(InAtlasResolution),
-	AtlasRectBuffer(),
-	LengthBuffer(),
-	CardItToClusterBuffer(),
-	ClusterIdToVerticesBuffer(),
-	ClusterBoundBuffer(),
-	CardsStrandsPositions(),
-	CardsStrandsAttributes(),
-	RenderData(InRenderData)
-{
-	CardVoxel = InVoxel;
-}
-
-void FHairCardsProceduralResource::InternalAllocate(FRDGBuilder& GraphBuilder)
-{
-	InternalCreateVertexBufferRDG<FHairCardsAtlasRectFormat>(GraphBuilder, RenderData.CardsRect, AtlasRectBuffer, ToHairResourceDebugName(TEXT("Hair.CardsProcedural_AtlasRectBuffer"), ResourceName), OwnerName, EHairResourceUsageType::Static);
-	InternalCreateVertexBufferRDG<FHairCardsDimensionFormat>(GraphBuilder, RenderData.CardsLengths, LengthBuffer, ToHairResourceDebugName(TEXT("Hair.CardsProcedural_LengthBuffer"), ResourceName), OwnerName, EHairResourceUsageType::Static);
-
-	InternalCreateVertexBufferRDG<FHairCardsOffsetAndCount>(GraphBuilder, RenderData.CardItToCluster, CardItToClusterBuffer, ToHairResourceDebugName(TEXT("Hair.CardsProcedural_CardItToClusterBuffer"), ResourceName), OwnerName, EHairResourceUsageType::Static);
-	InternalCreateVertexBufferRDG<FHairCardsOffsetAndCount>(GraphBuilder, RenderData.ClusterIdToVertices, ClusterIdToVerticesBuffer, ToHairResourceDebugName(TEXT("Hair.CardsProcedural_ClusterIdToVerticesBuffer"), ResourceName), OwnerName, EHairResourceUsageType::Static);
-	InternalCreateVertexBufferRDG<FHairCardsBoundsFormat>(GraphBuilder, RenderData.ClusterBounds, ClusterBoundBuffer, ToHairResourceDebugName(TEXT("Hair.CardsProcedural_ClusterBoundBuffer"), ResourceName), OwnerName, EHairResourceUsageType::Static);
-
-	InternalCreateVertexBufferRDG<FHairCardsVoxelDensityFormat>(GraphBuilder, RenderData.VoxelDensity, CardVoxel.DensityBuffer, ToHairResourceDebugName(TEXT("Hair.CardsProcedural_VoxelDensityBuffer"), ResourceName), OwnerName, EHairResourceUsageType::Static);
-	InternalCreateVertexBufferRDG<FHairCardsVoxelTangentFormat>(GraphBuilder, RenderData.VoxelTangent, CardVoxel.TangentBuffer, ToHairResourceDebugName(TEXT("Hair.CardsProcedural_VoxelTangentBuffer"), ResourceName), OwnerName, EHairResourceUsageType::Static);
-	InternalCreateVertexBufferRDG<FHairCardsVoxelTangentFormat>(GraphBuilder, RenderData.VoxelNormal, CardVoxel.NormalBuffer, ToHairResourceDebugName(TEXT("Hair.CardsProcedural_VoxelNormalBuffer"), ResourceName), OwnerName, EHairResourceUsageType::Static);
-
-	InternalCreateVertexBufferRDG<FHairCardsStrandsPositionFormat>(GraphBuilder, RenderData.CardsStrandsPositions, CardsStrandsPositions, ToHairResourceDebugName(TEXT("Hair.CardsProcedural_CardsStrandsPositions"), ResourceName), OwnerName, EHairResourceUsageType::Static);
-	InternalCreateVertexBufferRDG<FHairCardsStrandsAttributeFormat>(GraphBuilder, RenderData.CardsStrandsAttributes, CardsStrandsAttributes, ToHairResourceDebugName(TEXT("Hair.CardsProcedural_CardsStrandsAttributes"), ResourceName), OwnerName, EHairResourceUsageType::Static);
-}
-
-void FHairCardsProceduralResource::InternalRelease()
-{
-	AtlasRectBuffer.Release();
-	LengthBuffer.Release();
-
-	CardItToClusterBuffer.Release();
-	ClusterIdToVerticesBuffer.Release();
-	ClusterBoundBuffer.Release();
-	CardsStrandsPositions.Release();
-	CardsStrandsAttributes.Release();
-
-	CardVoxel.DensityBuffer.Release();
-	CardVoxel.TangentBuffer.Release();
-	CardVoxel.NormalBuffer.Release();
-}
-
 /////////////////////////////////////////////////////////////////////////////////////////
 
 FHairCardsDeformedResource::FHairCardsDeformedResource(const FHairCardsBulkData& InBulkData, bool bInInitializedData, const FHairResourceName& InResourceName, const FName& InOwnerName) :
diff --git a/Engine/Plugins/Runtime/HairStrands/Source/HairStrandsCore/Private/HairCardsBuilder.cpp b/Engine/Plugins/Runtime/HairStrands/Source/HairStrandsCore/Private/HairCardsBuilder.cpp
--- a/Engine/Plugins/Runtime/HairStrands/Source/HairStrandsCore/Private/HairCardsBuilder.cpp
+++ b/Engine/Plugins/Runtime/HairStrands/Source/HairStrandsCore/Private/HairCardsBuilder.cpp
@@ -48,27 +48,6 @@ PRAGMA_DEFAULT_VISIBILITY_END
 #pragma warning(pop)
 #endif
 
-static int32 GHairCardsMaxClusterCount = 9999999;
-static FAutoConsoleVariableRef CVarHairCardsMaxClusterCount(TEXT("r.HairStrands.Cards.MaxClusterCount"), GHairCardsMaxClusterCount, TEXT("Max number of cluster for debug purpose"));
-
-static int32 GHairCardsDebugIndex = 0;
-static FAutoConsoleVariableRef CVarHairCardsDebugIndex(TEXT("r.HairStrands.Cards.DebugIndex"), GHairCardsDebugIndex, TEXT("ID of the hair card to debug"));
-
-static float GHairCardsAtlasWidthScale = 1;
-static FAutoConsoleVariableRef CVarHairCardsAtlasWidthScale(TEXT("r.HairStrands.Cards.AtlasWidthScale"), GHairCardsAtlasWidthScale, TEXT("Scale the cards resolution along the width"));
-
-static float GHairCardsWidthScale = 0.1f;
-static FAutoConsoleVariableRef CVarHairCardsWidthScale(TEXT("r.HairStrands.Cards.WidthScale"), GHairCardsWidthScale, TEXT("Scale the cards resolution along the width"));
-
-static int32 GHairCardsDynamicAtlasRefresh = 0;
-static FAutoConsoleVariableRef CVarHairCardsDynamicAtlasRefresh(TEXT("r.HairStrands.Cards.DynamicAtlasRefresh"), GHairCardsDynamicAtlasRefresh, TEXT("Enable dynamic refresh of hair cards texture atlas"));
-
-static int32 GHairCardsMaxStrandsSegmentPerCards = 1000000;
-static FAutoConsoleVariableRef CVarHairCardsMaxStrandsSegmentPerCards(TEXT("r.HairStrands.Cards.MaxHairStrandsSegmentPerCards"), GHairCardsMaxStrandsSegmentPerCards, TEXT("Limit the number of segment which are raytraced during the cards generation"));
-
-static int32 GHairCardsAtlasMaxSample = 32;
-static FAutoConsoleVariableRef CVarHairCardsAtlasMaxSample(TEXT("r.HairStrands.Cards.MaxAtlasSample"), GHairCardsAtlasMaxSample, TEXT("Max super sampling count when generating cards atlas texture"));
-
 ///////////////////////////////////////////////////////////////////////////////////////////////////
 
 static FBox ToFBox3d(const FBox3f& In)
@@ -76,2828 +55,46 @@ static FBox ToFBox3d(const FBox3f& In)
 	return FBox(FVector(In.Min), FVector(In.Max));
 }
 
-///////////////////////////////////////////////////////////////////////////////////////////////////
-
-class FHairCardAtlasTextureRectVS : public FGlobalShader
-{
-	DECLARE_GLOBAL_SHADER(FHairCardAtlasTextureRectVS);
-	SHADER_USE_PARAMETER_STRUCT(FHairCardAtlasTextureRectVS, FGlobalShader)
-
-	BEGIN_SHADER_PARAMETER_STRUCT(FParameters, )
-		SHADER_PARAMETER(FVector3f, Raster_MinBound)
-		SHADER_PARAMETER(FVector3f, Raster_MaxBound)
-
-		SHADER_PARAMETER(uint32, SampleCount)
-
-		SHADER_PARAMETER(FIntPoint, Atlas_Resolution)
-		SHADER_PARAMETER(FIntPoint, Atlas_RectOffset)
-		SHADER_PARAMETER(FIntPoint, Atlas_RectResolution)
-
-		SHADER_PARAMETER(FVector3f, Raster_AxisX)
-		SHADER_PARAMETER(FVector3f, Raster_AxisY)
-		SHADER_PARAMETER(FVector3f, Raster_AxisZ)
-
-		SHADER_PARAMETER(uint32, Curve_VertexOffset)
-		SHADER_PARAMETER(uint32, Curve_VertexCount)
-		SHADER_PARAMETER(uint32, Curve_TotalVertexCount)
-		SHADER_PARAMETER_SRV(Buffer, Curve_PositionBuffer)
-		SHADER_PARAMETER_SRV(Buffer, Curve_AttributeBuffer)
-	END_SHADER_PARAMETER_STRUCT()
-
-	static bool ShouldCompilePermutation(const FGlobalShaderPermutationParameters& Parameters) { return IsHairStrandsSupported(EHairStrandsShaderType::Tool, Parameters.Platform); }
-	static void ModifyCompilationEnvironment(const FGlobalShaderPermutationParameters& Parameters, FShaderCompilerEnvironment& OutEnvironment)
-	{
-		FGlobalShader::ModifyCompilationEnvironment(Parameters, OutEnvironment);
-		OutEnvironment.SetDefine(TEXT("SHADER_VERTEX_RECT"), 1);
-	}
-};
-
-class FHairCardAtlasTextureRectPS : public FGlobalShader
-{
-	DECLARE_GLOBAL_SHADER(FHairCardAtlasTextureRectPS);
-	SHADER_USE_PARAMETER_STRUCT(FHairCardAtlasTextureRectPS, FGlobalShader)
-
-	class FOutput : SHADER_PERMUTATION_INT("PERMUTATION_OUTPUT", 2);
-	using FPermutationDomain = TShaderPermutationDomain<FOutput>;
-
-	BEGIN_SHADER_PARAMETER_STRUCT(FParameters, )
-		SHADER_PARAMETER_STRUCT_INCLUDE(ShaderPrint::FShaderParameters, ShaderDrawParameters)
-		SHADER_PARAMETER(FVector3f, Raster_MinBound)
-		SHADER_PARAMETER(FVector3f, Raster_MaxBound)
-		
-		SHADER_PARAMETER(uint32, SampleCount)
-
-		SHADER_PARAMETER(FIntPoint, Atlas_Resolution)
-		SHADER_PARAMETER(FIntPoint, Atlas_RectOffset)
-		SHADER_PARAMETER(FIntPoint, Atlas_RectResolution)
-
-		SHADER_PARAMETER(FVector3f, Raster_AxisX)
-		SHADER_PARAMETER(FVector3f, Raster_AxisY)
-		SHADER_PARAMETER(FVector3f, Raster_AxisZ)
-
-		SHADER_PARAMETER(uint32,	Curve_TotalVertexCount)
-		SHADER_PARAMETER(uint32,	Curve_VertexOffset)
-		SHADER_PARAMETER(uint32,	Curve_VertexCount)
-		SHADER_PARAMETER_SRV(Buffer,Curve_PositionBuffer)
-		SHADER_PARAMETER_SRV(Buffer, Curve_AttributeBuffer)
-
-		RENDER_TARGET_BINDING_SLOTS()
-	END_SHADER_PARAMETER_STRUCT()
-
-	static bool ShouldCompilePermutation(const FGlobalShaderPermutationParameters& Parameters) { return IsHairStrandsSupported(EHairStrandsShaderType::Tool, Parameters.Platform); }
-	static void ModifyCompilationEnvironment(const FGlobalShaderPermutationParameters& Parameters, FShaderCompilerEnvironment& OutEnvironment)
-	{
-		FGlobalShader::ModifyCompilationEnvironment(Parameters, OutEnvironment);
-		OutEnvironment.SetDefine(TEXT("SHADER_PIXEL_RECT"), 1);
-	}
-};
-
-IMPLEMENT_GLOBAL_SHADER(FHairCardAtlasTextureRectVS, "/Engine/Private/HairStrands/HairCardsGeneration.usf", "MainVS", SF_Vertex);
-IMPLEMENT_GLOBAL_SHADER(FHairCardAtlasTextureRectPS, "/Engine/Private/HairStrands/HairCardsGeneration.usf", "MainPS", SF_Pixel);
-
-static void AddCardsTracingPass(
-	FRDGBuilder& GraphBuilder,
-	FGlobalShaderMap* ShaderMap,
-	const bool bOutputCoverageOnly,
-	const bool bClear,
-	const FHairCardsProceduralAtlas& InAtlas,
-	const FHairCardsProceduralAtlas::Rect& Rect,
-	const FShaderPrintData* ShaderPrintData,
-	const uint32 TotalVertexCount,
-	FRHIShaderResourceView* InStrandsVertexBuffer,
-	FRHIShaderResourceView* InStrandsAttributeBuffer,
-	FRDGTextureRef OutDepthTestTexture,
-	FRDGTextureRef OutDepthTexture,
-	FRDGTextureRef OutTangentTexture,
-	FRDGTextureRef OutCoverageTexture,
-	FRDGTextureRef OutAttributeTexture)
-{
-	const FIntPoint AtlasResolution = OutDepthTexture->Desc.Extent;
-
-	FHairCardAtlasTextureRectPS::FParameters* ParametersPS = GraphBuilder.AllocParameters<FHairCardAtlasTextureRectPS::FParameters>();
-	ParametersPS->SampleCount = FMath::Max(GHairCardsAtlasMaxSample, 1);
-
-	ParametersPS->Raster_MinBound = Rect.MinBound;
-	ParametersPS->Raster_MaxBound = Rect.MaxBound;
-
-	ParametersPS->Atlas_Resolution = InAtlas.Resolution;
-	ParametersPS->Atlas_RectOffset = Rect.Offset;
-	ParametersPS->Atlas_RectResolution = Rect.Resolution;
-
-	ParametersPS->Raster_AxisX = Rect.RasterAxisX;
-	ParametersPS->Raster_AxisY = Rect.RasterAxisY;
-	ParametersPS->Raster_AxisZ = Rect.RasterAxisZ;
-
-	ParametersPS->Curve_TotalVertexCount = TotalVertexCount;
-	ParametersPS->Curve_VertexOffset = Rect.VertexOffset;
-	ParametersPS->Curve_VertexCount = FMath::Min(Rect.VertexCount, uint32(GHairCardsMaxStrandsSegmentPerCards));
-	ParametersPS->Curve_TotalVertexCount = TotalVertexCount;
-	ParametersPS->Curve_PositionBuffer = InStrandsVertexBuffer;
-	ParametersPS->Curve_AttributeBuffer = InStrandsAttributeBuffer;
-
-	if (ShaderPrintData)
-	{
-		ShaderPrint::SetParameters(GraphBuilder, *ShaderPrintData, ParametersPS->ShaderDrawParameters);
-	}
-
-	if (bOutputCoverageOnly)
-	{
-		ParametersPS->RenderTargets[0] = FRenderTargetBinding(OutCoverageTexture, bClear ? ERenderTargetLoadAction::EClear : ERenderTargetLoadAction::ELoad);
-	}
-	else
-	{
-		ParametersPS->RenderTargets[0] = FRenderTargetBinding(OutDepthTexture, bClear ? ERenderTargetLoadAction::EClear : ERenderTargetLoadAction::ELoad);
-		ParametersPS->RenderTargets[1] = FRenderTargetBinding(OutTangentTexture, bClear ? ERenderTargetLoadAction::EClear : ERenderTargetLoadAction::ELoad);
-		ParametersPS->RenderTargets[2] = FRenderTargetBinding(OutAttributeTexture, bClear ? ERenderTargetLoadAction::EClear : ERenderTargetLoadAction::ELoad);
-		ParametersPS->RenderTargets.DepthStencil = FDepthStencilBinding(OutDepthTestTexture, bClear ? ERenderTargetLoadAction::EClear : ERenderTargetLoadAction::ELoad, ERenderTargetLoadAction::ELoad, FExclusiveDepthStencil::DepthWrite_StencilNop);
-	}
-
-	FHairCardAtlasTextureRectPS::FPermutationDomain PermutationVector;
-	PermutationVector.Set<FHairCardAtlasTextureRectPS::FOutput>(bOutputCoverageOnly ? 1 : 0);
-	TShaderMapRef<FHairCardAtlasTextureRectVS> VertexShader(ShaderMap);
-	TShaderMapRef<FHairCardAtlasTextureRectPS> PixelShader(ShaderMap, PermutationVector);
-
-	GraphBuilder.AddPass(
-		RDG_EVENT_NAME("HairCardsAtlasTexturePS"),
-		ParametersPS,
-		ERDGPassFlags::Raster | ERDGPassFlags::NeverCull,
-		[ParametersPS, VertexShader, PixelShader, AtlasResolution, Rect, bOutputCoverageOnly](FRHICommandList& RHICmdList)
-		{
-			FHairCardAtlasTextureRectVS::FParameters ParametersVS;
-			ParametersVS.SampleCount = ParametersPS->SampleCount;
-
-			ParametersVS.Raster_MinBound = ParametersPS->Raster_MinBound;
-			ParametersVS.Raster_MaxBound = ParametersPS->Raster_MaxBound;
-
-			ParametersVS.Atlas_Resolution = ParametersPS->Atlas_Resolution;
-			ParametersVS.Atlas_RectOffset = ParametersPS->Atlas_RectOffset;
-			ParametersVS.Atlas_RectResolution = ParametersPS->Atlas_RectResolution;
-
-			ParametersVS.Raster_AxisX = ParametersPS->Raster_AxisX;
-			ParametersVS.Raster_AxisY = ParametersPS->Raster_AxisY;
-			ParametersVS.Raster_AxisZ = ParametersPS->Raster_AxisZ;
-
-			ParametersVS.Curve_TotalVertexCount = ParametersPS->Curve_TotalVertexCount;
-			ParametersVS.Curve_VertexOffset = ParametersPS->Curve_VertexOffset;
-			ParametersVS.Curve_VertexCount = ParametersPS->Curve_VertexCount;
-			ParametersVS.Curve_PositionBuffer = ParametersPS->Curve_PositionBuffer;
-			ParametersVS.Curve_AttributeBuffer = ParametersPS->Curve_AttributeBuffer;
-
-			FGraphicsPipelineStateInitializer GraphicsPSOInit;
-			RHICmdList.ApplyCachedRenderTargets(GraphicsPSOInit);
-
-			if (bOutputCoverageOnly)
-			{
-				GraphicsPSOInit.BlendState = TStaticBlendState<CW_RGBA, BO_Add, BF_One, BF_One, BO_Add, BF_One, BF_One>::GetRHI();
-				GraphicsPSOInit.RasterizerState = TStaticRasterizerState<>::GetRHI();
-				GraphicsPSOInit.DepthStencilState = TStaticDepthStencilState<false, CF_Always> ::GetRHI();
-			}
-			else
-			{
-				GraphicsPSOInit.BlendState = TStaticBlendState<CW_RGBA, BO_Add, BF_One, BF_Zero, BO_Add, BF_One, BF_Zero>::GetRHI();
-				GraphicsPSOInit.RasterizerState = TStaticRasterizerState<>::GetRHI();
-				GraphicsPSOInit.DepthStencilState = TStaticDepthStencilState<true, CF_Less>::GetRHI();
-			}
-
-
-			GraphicsPSOInit.BoundShaderState.VertexDeclarationRHI = GFilterVertexDeclaration.VertexDeclarationRHI;
-			GraphicsPSOInit.BoundShaderState.VertexShaderRHI = VertexShader.GetVertexShader();
-			GraphicsPSOInit.BoundShaderState.PixelShaderRHI = PixelShader.GetPixelShader();
-			GraphicsPSOInit.PrimitiveType = PT_TriangleList;
-			SetGraphicsPipelineState(RHICmdList, GraphicsPSOInit, 0);
-
-			SetShaderParameters(RHICmdList, VertexShader, VertexShader.GetVertexShader(), ParametersVS);
-			SetShaderParameters(RHICmdList, PixelShader, PixelShader.GetPixelShader(), *ParametersPS);
-
-			const uint32 NumSegments = ParametersVS.Curve_VertexCount;
+#define LOCTEXT_NAMESPACE "GroomCardBuilder"
 
-			RHICmdList.SetViewport(Rect.Offset.X, Rect.Offset.Y, 0.0f, Rect.Offset.X + Rect.Resolution.X, Rect.Offset.Y + Rect.Resolution.Y, 1.0f);
-			RHICmdList.SetStreamSource(0, nullptr, 0);
-			RHICmdList.DrawPrimitive(0, 2 * NumSegments, 1);
-		});
-}
+////////////////////////////////////////////////////////////////////////////////////////////////
+// Build public functions
 
-static void AddHairCardAtlasTexturePass(
-	FRDGBuilder& GraphBuilder,
-	FGlobalShaderMap* ShaderMap,
-	const FHairCardsProceduralAtlas& InAtlas,
-	const FShaderPrintData* ShaderPrintData,
-	const uint32 TotalVertexCount,
-	FRHIShaderResourceView* InStrandsVertexBuffer,
-	FRHIShaderResourceView* InStrandsAttributeBuffer,
-	FRDGTextureRef OutDepthTexture,
-	FRDGTextureRef OutTangentTexture,
-	FRDGTextureRef OutCoverageTexture,
-	FRDGTextureRef OutAttributeTexture)
+namespace FHairCardsBuilder
 {
-	FRDGTextureRef DepthTestTexture = GraphBuilder.CreateTexture(FRDGTextureDesc::Create2D(OutAttributeTexture->Desc.Extent, PF_DepthStencil, FClearValueBinding::DepthOne, TexCreate_DepthStencilTargetable | TexCreate_ShaderResource, 1), TEXT("Hair.CardsDepthTest"));
-
-	bool bClear = true;
-	for (const FHairCardsProceduralAtlas::Rect& Rect : InAtlas.Rects)
-	{
-		// 1. Generate depth/tangent/attribute
-		AddCardsTracingPass(
-			GraphBuilder,
-			ShaderMap,
-			false,
-			bClear,
-			InAtlas,
-			Rect,
-			ShaderPrintData,
-			TotalVertexCount,
-			InStrandsVertexBuffer,
-			InStrandsAttributeBuffer,
-			DepthTestTexture,
-			OutDepthTexture,
-			OutTangentTexture,
-			OutCoverageTexture,
-			OutAttributeTexture);
-
-		// 2. Generate Coverage
-		AddCardsTracingPass(
-			GraphBuilder,
-			ShaderMap,
-			true,
-			bClear,
-			InAtlas,
-			Rect,
-			ShaderPrintData,
-			TotalVertexCount,
-			InStrandsVertexBuffer,
-			InStrandsAttributeBuffer,
-			DepthTestTexture,
-			OutDepthTexture,
-			OutTangentTexture,
-			OutCoverageTexture,
-			OutAttributeTexture);
-
-		bClear = false;
-	}
-}
-
-///////////////////////////////////////////////////////////////////////////////////////////////////
-
-#define LOCTEXT_NAMESPACE "GroomCardBuilder"
 
-// Voxelize hair strands into a small low res volume. This used for finding the main orientation of the groom
-namespace HairCards
+FString GetVersion()
 {
-	struct FVoxelVolume
-	{
-		float VoxelSize = 0;
-		FIntVector Resolution = FIntVector::ZeroValue;
-		FVector3f MinBound = FVector3f::ZeroVector;
-		FVector3f MaxBound = FVector3f::ZeroVector;
-		TArray<uint32> Density;
-		TArray<FVector3f> Normal;
-		TArray<FVector3f> Tangent;
-	};
-
-	inline uint32 To8bits(const FVector3f& T)
-	{
-		uint32 X = FMath::Clamp((T.X + 1) * 0.5f, 0.f, 1.f) * 255.f;
-		uint32 Y = FMath::Clamp((T.Y + 1) * 0.5f, 0.f, 1.f) * 255.f;
-		uint32 Z = FMath::Clamp((T.Z + 1) * 0.5f, 0.f, 1.f) * 255.f;
-		return (0xFF & X) | ((0xFF & Y)<<8) | ((0xFF & Z)<<16);
-	}
-
-	inline FVector3f From8bits(const uint32& T)
-	{
-		return FVector3f(
-			((0xFF & (T))     / 255.f) * 2.f - 1.f,
-			((0xFF & (T>>8))  / 255.f) * 2.f - 1.f,
-			((0xFF & (T>>16)) / 255.f) * 2.f - 1.f);
-	}
-
-	inline FIntVector ToCoord(const FVector3f& T, const FVoxelVolume& V)
-	{
-		const FVector3f C = (T  - V.MinBound) / V.VoxelSize;
-		return FIntVector(
-			FMath::Clamp(FMath::FloorToInt(C.X), 0, V.Resolution.X-1), 
-			FMath::Clamp(FMath::FloorToInt(C.Y), 0, V.Resolution.Y-1), 
-			FMath::Clamp(FMath::FloorToInt(C.Z), 0, V.Resolution.Z-1));
-	}
-
-	inline uint32 ToLinearCoord(const FIntVector& T, const FVoxelVolume& V)
-	{
-		// Morton instead for better locality?
-		return T.X + T.Y * V.Resolution.X + T.Z * V.Resolution.X * V.Resolution.Y;
-	}
-
-	// Return true if the voxel coord is within the voxelized volume
-	inline bool IsInside(const FIntVector& C, const FVoxelVolume& V)
-	{
-		return 
-			C.X >=0 && C.Y >= 0 && C.Z >= 0 &&
-			C.X < V.Resolution.X && C.Y < V.Resolution.Y && C.Z < V.Resolution.Z;
-	}
-
-	// Return true if the voxel is empty (no density)
-	inline bool IsEmpty(const FIntVector& C, const FVoxelVolume& V)
-	{
-		return IsInside(C, V) ? (V.Density[ToLinearCoord(C, V)] == 0) : true;
-	}
-
-	// Return true is at least one of the voxel face is visible, i.e., not hidden by another non-empty voxel
-	inline bool IsVisibile(const FIntVector& C, const FVoxelVolume& V)
-	{
-		FIntVector CX0 = C; CX0.X += 1;
-		FIntVector CX1 = C; CX1.X -= 1;
-
-		FIntVector CY0 = C; CY0.Y += 1;
-		FIntVector CY1 = C; CY1.Y -= 1;
-
-		FIntVector CZ0 = C; CZ0.Z += 1;
-		FIntVector CZ1 = C; CZ1.Z -= 1;
-
-		return
-			IsEmpty(CX0, V) || IsEmpty(CX1, V) ||
-			IsEmpty(CY0, V) || IsEmpty(CY1, V) ||
-			IsEmpty(CZ0, V) || IsEmpty(CZ1, V);
-	}
-
-	TArray<FIntVector> CreateKernelCoord(const int32 KernelSize, const FVoxelVolume& InV)
-	{
-		TArray<FIntVector> Out;
-		const uint32 KernelWidth = KernelSize * 2 + 1;
-		Out.Reserve(KernelWidth* KernelWidth* KernelWidth);
-
-		for (int32 Z = -KernelSize; Z <= KernelSize; ++Z)
-		for (int32 Y = -KernelSize; Y <= KernelSize; ++Y)
-		for (int32 X = -KernelSize; X <= KernelSize; ++X)
-		{
-			Out.Add(FIntVector(X, Y, Z));
-		}
-		return Out;
-	}
-
-	// Compute the iso surface normal
-	void ComputeNormal(FVoxelVolume& Out)
-	{
-		Out.Normal.SetNum(Out.Density.Num());
-		TArray<FIntVector> NonVisibleCoord;
-		TArray<FIntVector> VisibleCoord;
-		NonVisibleCoord.Reserve(Out.Density.Num() * 0.5f); // Guesstimate
-		VisibleCoord.Reserve(Out.Density.Num() * 0.5f); // Guesstimate
-
-		const FVector3f CenterBoundCoord = FVector3f(Out.Resolution) * 0.5f;
-
-		const TArray<FIntVector> Kernel1 = CreateKernelCoord(1, Out);
-		const TArray<FIntVector> Kernel2 = CreateKernelCoord(2, Out);
-
-		// Compute normal on at 0-density iso surface
-		{
-			for (int32 Z = 0; Z < Out.Resolution.Z; ++Z)
-			for (int32 Y = 0; Y < Out.Resolution.Y; ++Y)
-			for (int32 X = 0; X < Out.Resolution.X; ++X)
-			{
-				const FIntVector Coord(X,Y,Z);
-				const uint32 LinearCoord = ToLinearCoord(Coord, Out);
-				Out.Normal[LinearCoord] = FVector3f::ZeroVector;
-
-				const uint32 CenterDensity = Out.Density[LinearCoord];
-				if (CenterDensity == 0)
-					continue;
-
-				if (!IsVisibile(Coord, Out))
-				{
-					NonVisibleCoord.Add(Coord);
-					continue;
-				}
-				else
-				{
-					VisibleCoord.Add(Coord);
-				}
-
-				uint32 ValidCount = 0;
-				FVector3f OutDirection = FVector3f::ZeroVector;
-				for (const FIntVector& Offset : Kernel1)
-				{
-					const FIntVector C = Coord + Offset;
-					if (Offset != FIntVector::ZeroValue && IsEmpty(C, Out))
-					{
-						OutDirection += FVector3f(Offset).GetSafeNormal();
-						ValidCount++;
-					}
-				}
-
-				assert(ValidCount > 0);
-				Out.Normal[LinearCoord] = OutDirection.GetSafeNormal();
-
-				// If nothing has been found in a 1x1 kernel, try to look for further
-				if (Out.Normal[LinearCoord] == FVector3f::ZeroVector)
-				{
-					ValidCount = 0;
-					OutDirection = FVector3f::ZeroVector;
-					for (const FIntVector& Offset : Kernel2)
-					{
-						const FIntVector C = Coord + Offset;
-						if (Offset != FIntVector::ZeroValue && IsEmpty(C, Out))
-						{
-							OutDirection += FVector3f(Offset).GetSafeNormal();
-							ValidCount++;
-						}
-					}
-					assert(ValidCount > 0);
-					Out.Normal[LinearCoord] = OutDirection.GetSafeNormal();
-
-					// Nothing has been found, so use the center of the bound
-					if (Out.Normal[LinearCoord] != FVector3f::ZeroVector)
-					{
-						Out.Normal[LinearCoord] = (FVector3f(Coord) - CenterBoundCoord).GetSafeNormal();
-					}
-				}
-			}
-		}
-
-		// Filter normal on the iso surface
-		{	
-			TArray<FVector3f> OutNormal;
-			OutNormal.Init(FVector3f::ZeroVector, Out.Normal.Num());
-
-			// Filter only voxel which have valid normals
-			for (const FIntVector& Coord : VisibleCoord)
-			{
-				const uint32 LinearCoord = ToLinearCoord(Coord, Out);
-
-				uint32 ValidCount = 0;
-				FVector3f OutDirection = FVector3f::ZeroVector;
-				for (const FIntVector& Offset : Kernel2)
-				{
-					const FIntVector C = Coord + Offset;
-					const float w = Offset == FIntVector::ZeroValue ? 2 : 1;
-					if (IsInside(C, Out))
-					{
-						const uint32 LinearC = ToLinearCoord(C, Out);
-
-						const bool bIsSampleValid = Out.Normal[LinearC] != FVector3f::ZeroVector;
-						if (bIsSampleValid)
-						{
-							OutNormal[LinearCoord]  += Out.Normal[LinearC].GetSafeNormal() * w;
-							ValidCount++;
-						}
-					}
-				}
-
-				if (ValidCount > 0)
-				{
-					OutNormal[LinearCoord] = OutNormal[LinearCoord].GetSafeNormal();
-				}
-				else
-				{
-					// Generate normal oriented as a surounding sphere
-					OutNormal[LinearCoord] = (FVector3f(Coord) - FVector3f(CenterBoundCoord)).GetSafeNormal();
-				}
-			}
-
-			Out.Normal = OutNormal;
-		}
-
-		// Fill/propagate normal within the volume with the iso surface normal
-		const uint32 MaxPassCount = 10;
-		uint32 PassIndex = 0;
-		while (NonVisibleCoord.Num() > 0 && PassIndex++ < MaxPassCount)
-		{
-			TArray<FIntVector> NextNonVisibleCoord;
-			NextNonVisibleCoord.Reserve(NonVisibleCoord.Num());
-			for (const FIntVector& Coord : NonVisibleCoord)
-			{
-				const uint32 LinearCoord = ToLinearCoord(Coord, Out);
-				check(Out.Normal[LinearCoord] == FVector3f::ZeroVector);
-				Out.Normal[LinearCoord] = FVector3f::ZeroVector;
-
-				uint32 ValidCount = 0;
-				FVector3f OutDirection = FVector3f::ZeroVector;
-				for (const FIntVector& Offset : Kernel1)
-				{
-					const FIntVector C = Coord + Offset;
-					if (IsInside(C, Out) && Offset != FIntVector::ZeroValue)
-					{
-						const uint32 LinearC = ToLinearCoord(C, Out);
-						const FVector3f SNormal = Out.Normal[LinearC];
-						if (SNormal != FVector3f::ZeroVector)
-						{
-							OutDirection += SNormal;
-							ValidCount++;
-						}
-					}
-				}
-
-				if (ValidCount > 0)
-				{
-					Out.Normal[LinearCoord] = OutDirection.GetSafeNormal();
-				}
-				else
-				{
-					NextNonVisibleCoord.Add(Coord);
-				}
-			}
-
-			NonVisibleCoord = NextNonVisibleCoord;
-		}
-
-		check(PassIndex < MaxPassCount);
-	}
-
-	void FilterTangent(FVoxelVolume& Out)
-	{	
-		const TArray<FIntVector> Kernel2 = CreateKernelCoord(2, Out);
-
-		TArray<FVector3f> OutTangent;
-		OutTangent.SetNum(Out.Tangent.Num());
-
-		for (int32 Z = 0; Z < Out.Resolution.Z; ++Z)
-		for (int32 Y = 0; Y < Out.Resolution.Y; ++Y)
-		for (int32 X = 0; X < Out.Resolution.X; ++X)
-		{
-			const FIntVector Coord(X,Y,Z);
-			const uint32 LinearCoord = ToLinearCoord(Coord, Out);
-			OutTangent[LinearCoord] = FVector3f::ZeroVector;
-
-			const uint32 CenterDensity = Out.Density[LinearCoord];
-			if (CenterDensity == 0)
-				continue;
-
-			uint32 ValidCount = 0;
-			FVector3f OutDirection = FVector3f::ZeroVector;
-			for (const FIntVector& Offset : Kernel2)
-			{
-				const FIntVector C = Coord + Offset;
-				const float w = Offset == FIntVector::ZeroValue ? 2 : 1;
-				if (IsInside(C, Out))
-				{
-					const uint32 LinearC = ToLinearCoord(C, Out);
-					const uint32 Density = Out.Density[LinearC];
-					if (Density > 0)
-					{
-						OutTangent[LinearCoord]  += Out.Tangent[LinearC].GetSafeNormal() * w;
-						ValidCount++;
-					}
-				}
-			}
-
-			assert(ValidCount > 0);
-			OutTangent[LinearCoord] = OutTangent[LinearCoord].GetSafeNormal();
-		}
-
-		Out.Tangent = OutTangent;
-	}
-
-	void VoxelizeGroom(const FHairStrandsDatas& InData, const uint32 InResolution, FVoxelVolume& Out)
-	{
-		const FVector3f BoundSize = FVector3f(InData.BoundingBox.Max - InData.BoundingBox.Min);
-		Out.VoxelSize = FMath::Max3(BoundSize.X, BoundSize.Y, BoundSize.Z) / InResolution;
-
-		Out.Resolution = FIntVector(FMath::CeilToFloat(BoundSize.X / Out.VoxelSize), FMath::CeilToFloat(BoundSize.Y / Out.VoxelSize), FMath::CeilToFloat(BoundSize.Z / Out.VoxelSize));
-		Out.MinBound   = (FVector3f)InData.BoundingBox.Min;
-		Out.MaxBound   = (FVector3f)Out.Resolution * Out.VoxelSize + (FVector3f)InData.BoundingBox.Min;
-
-		const uint32 TotalVoxelCount = Out.Resolution.X * Out.Resolution.Y * Out.Resolution.Z;
-		Out.Density.Init(0, TotalVoxelCount);
-		Out.Tangent.Init(FVector3f::ZeroVector, TotalVoxelCount);
-
-		// Fill in voxel (TODO: make it parallel)
-		const uint32 CurveCount = InData.StrandsCurves.Num();
-		for (uint32 CurveIndex = 0; CurveIndex < CurveCount; ++CurveIndex)
-		{
-			const uint32 PointOffset = InData.StrandsCurves.CurvesOffset[CurveIndex];
-			const uint32 PointCount = InData.StrandsCurves.CurvesCount[CurveIndex];
-			
-			for (uint32 PointIndex = 0; PointIndex < PointCount-1; ++PointIndex)
-			{
-				const FVector3f& P0 = InData.StrandsPoints.PointsPosition[PointOffset + PointIndex];
-				const FVector3f& P1 = InData.StrandsPoints.PointsPosition[PointOffset + PointIndex + 1];
-				const FVector3f Segment = P1 - P0;
-				const FVector3f T = Segment.GetSafeNormal();
-				//const uint32 T8bits = To8bits(T);
-
-				const float Length = Segment.Size();
-				const uint32 StepCount = FMath::CeilToInt(Length / Out.VoxelSize);
-				uint32 PrevLinearCoord = ~0;
-				for (uint32 StepIt = 0; StepIt < StepCount+1; ++StepIt)
-				{
-					const FVector3f P = P0 + Segment * StepIt / float(StepCount);
-					const FIntVector Coord = ToCoord(P, Out);
-					const uint32 LinearCoord = ToLinearCoord(Coord, Out);
-					if (LinearCoord != PrevLinearCoord)
-					{
-						Out.Density[LinearCoord] += 1;
-						Out.Tangent[LinearCoord] += T;
-						PrevLinearCoord = LinearCoord;
-					}
-				}
-			}
-		}
-
-		// Compute normal based on gradient density
-		ComputeNormal(Out);
-
-		// Filter tangent to get smoother result
-		FilterTangent(Out);
-	}
-
-	FVector3f GetTangent(const FVector3f InP, const FVoxelVolume& InV)
-	{
-		const FIntVector Coord = ToCoord(InP, InV);
-		const uint32 LinearCoord = ToLinearCoord(Coord, InV);
-
-		return InV.Tangent[LinearCoord].GetSafeNormal();
-	}
-
-	FVector3f GetNormal(const FVector3f InP, const FVoxelVolume& InV)
-	{
-		const FIntVector Coord = ToCoord(InP, InV);
-		const uint32 LinearCoord = ToLinearCoord(Coord, InV);
-		
-		return InV.Normal[LinearCoord].GetSafeNormal();
-	}
-
-	void GetFrame(const FVector3f InP, const FVoxelVolume& InV, FVector3f& OutT, FVector3f& OutB, FVector3f& OutN)
-	{
-		const FIntVector Coord = ToCoord(InP, InV);
-		const uint32 LinearCoord = ToLinearCoord(Coord, InV);
-
-#if 0
-		// Build ortho-normal frame
-		OutT = InV.Tangent[LinearCoord].GetSafeNormal();
-		OutN = InV.Normal[LinearCoord].GetSafeNormal();
-		OutB = FVector3f::CrossProduct(OutN, OutT).GetSafeNormal();
-		OutN = FVector3f::CrossProduct(OutT, OutN).GetSafeNormal();
-#else
-		// Return a non-ortho-normal frame
-		OutT = InV.Tangent[LinearCoord].GetSafeNormal();
-		OutN = InV.Normal[LinearCoord].GetSafeNormal();
-		OutB = FVector3f::CrossProduct(OutN, OutT).GetSafeNormal();
-#endif
-	}
+	// Important to update the version when cards building or importing changes
+	return TEXT("9g");
 }
 
-// Hair clustering: https://www.lvdiwang.com/publications/hairsynthesis/2009_hairsynthesis.pdf
-// K clustering with several iterations for finding a good k with an error threshold provided by the user
-// Export climp Id for helping the wisp creation?
-
-namespace HairCards
-{
-	struct FHairRoot
-	{
-		FVector3f Position;
-		uint32  VertexCount;
-		FVector3f Normal;
-		uint32  Index;
-		float   Length;
-	};
-
-	template<uint32 NumSamples>
-	inline FVector3f GetCurvePosition(const FHairStrandsDatas& CurvesDatas, const uint32 CurveIndex, const uint32 SampleIndex)
-	{
-		const float PointCount = CurvesDatas.StrandsCurves.CurvesCount[CurveIndex]-1.0;
-		const uint32 PointOffset = CurvesDatas.StrandsCurves.CurvesOffset[CurveIndex];
-
-		const float CurvePoint = static_cast<float>(SampleIndex) * PointCount / (static_cast<float>(NumSamples)-1.0f);
-		const uint32 PointPrev = (SampleIndex == 0) ? 0 : (SampleIndex == (NumSamples-1)) ? PointCount - 1 : floor(CurvePoint);
-		const uint32 PointNext = PointPrev + 1;
-
-		const float PointAlpha = CurvePoint - static_cast<float>(PointPrev);
-		return 
-			CurvesDatas.StrandsPoints.PointsPosition[PointOffset+PointPrev]*(1.0f-PointAlpha) + 
-			CurvesDatas.StrandsPoints.PointsPosition[PointOffset+PointNext]*PointAlpha;
-	}
-
-	template<uint32 NumSamples>
-	inline float ComputeCurvesMetric(
-		const FHairStrandsDatas& RenderCurvesDatas, 
-		const uint32 RenderCurveIndex, 
-		const FHairStrandsDatas& GuideCurvesDatas, 
-		const uint32 GuideCurveIndex, 
-		const float RootImportance, 
-		const float ShapeImportance, 
-		const float ProximityImportance)
-	{
-		const float AverageLength = FMath::Max(0.5f * (RenderCurvesDatas.StrandsCurves.CurvesLength[RenderCurveIndex] + GuideCurvesDatas.StrandsCurves.CurvesLength[GuideCurveIndex]), SMALL_NUMBER);
-
-		static const float DeltaCoord = 1.0f / static_cast<float>(NumSamples-1);
-
-		const FVector3f& RenderRoot = RenderCurvesDatas.StrandsPoints.PointsPosition[RenderCurvesDatas.StrandsCurves.CurvesOffset[RenderCurveIndex]];
-		const FVector3f& GuideRoot = GuideCurvesDatas.StrandsPoints.PointsPosition[GuideCurvesDatas.StrandsCurves.CurvesOffset[GuideCurveIndex]];
-
-		float CurveProximityMetric = 0.0;
-		float CurveShapeMetric = 0.0;
-		for (uint32 SampleIndex = 0; SampleIndex < NumSamples; ++SampleIndex)
-		{
-			const FVector3f GuidePosition = GetCurvePosition<NumSamples>(GuideCurvesDatas, GuideCurveIndex, SampleIndex);
-			const FVector3f RenderPosition = GetCurvePosition<NumSamples>(RenderCurvesDatas, RenderCurveIndex, SampleIndex);
-			const float RootWeight = FMath::Exp(-RootImportance*SampleIndex*DeltaCoord);
-
-			CurveProximityMetric += (GuidePosition - RenderPosition).Size() * RootWeight;
-			CurveShapeMetric += (GuidePosition - GuideRoot - RenderPosition + RenderRoot).Size() * RootWeight;
-		}
-		CurveShapeMetric *= DeltaCoord / AverageLength;
-		CurveProximityMetric *= DeltaCoord / AverageLength;
-
-		return FMath::Exp(-ShapeImportance*CurveShapeMetric) * FMath::Exp(-ProximityImportance * CurveProximityMetric);
-	}
-
-	template<typename T>
-	void SwapValue(T& A, T& B)
-	{
-		T Temp = A;
-		A = B;
-		B = Temp;
-	}
-
-	struct FMetrics
-	{
-		static const int32 Invalid = -1;
-		static const uint32 Count = 3;
-		float KMinMetrics[Count];
-		int32 KClosestGuideIndices[Count];
-	};
-
-	struct FClosestGuides
-	{
-		static const int32 Invalid = -1;
-		static const uint32 Count = 3;
-		int32 Indices[Count];
-	};
-
-	static void SelectFinalGuides(FClosestGuides& ClosestGuides, const FMetrics& InMetric)
-	{
-		check(InMetric.KClosestGuideIndices[0] >= 0);
-
-		ClosestGuides.Indices[0] = InMetric.KClosestGuideIndices[0];
-		ClosestGuides.Indices[1] = InMetric.KClosestGuideIndices[1];
-		ClosestGuides.Indices[2] = InMetric.KClosestGuideIndices[2];
-
-		check(FClosestGuides::Count == 3);
-		if (ClosestGuides.Indices[0] == ClosestGuides.Indices[1]) { ClosestGuides.Indices[1] = FClosestGuides::Invalid; }
-		if (ClosestGuides.Indices[1] == ClosestGuides.Indices[2]) { ClosestGuides.Indices[2] = FClosestGuides::Invalid; }
-		if (ClosestGuides.Indices[0] == ClosestGuides.Indices[2]) { ClosestGuides.Indices[2] = FClosestGuides::Invalid; }
-	}
-
-	// Simple closest distance metric
-	static void ComputeSimpleMetric(
-		FMetrics& Metrics1, 
-		const FHairRoot& RenRoot, 
-		const FHairRoot& GuideRoot, 
-		const int32 RenCurveIndex,
-		const int32 SimCurveIndex)
-	{
-		const float Metric = FVector3f::Dist(GuideRoot.Position, RenRoot.Position);
-		if (Metric < Metrics1.KMinMetrics[FMetrics::Count - 1])
-		{
-			int32 LastGuideIndex = SimCurveIndex;
-			float LastMetric = Metric;
-			for (uint32 Index = 0; Index < FMetrics::Count; ++Index)
-			{
-				if (Metric < Metrics1.KMinMetrics[Index])
-				{
-					SwapValue(Metrics1.KClosestGuideIndices[Index], LastGuideIndex);
-					SwapValue(Metrics1.KMinMetrics[Index], LastMetric);
-				}
-			}
-		}
-	}
-
-	// Complex pairing metric
-	static void ComputeAdvandedMetric(
-		FMetrics& Metrics0,
-		const FHairStrandsDatas& RenStrandsData,
-		const FHairStrandsDatas& SimStrandsData,
-		const int32 RenCurveIndex,
-		const int32 SimCurveIndex)
-	{
-		const float Metric = 1.0 - ComputeCurvesMetric<16>(RenStrandsData, RenCurveIndex, SimStrandsData, SimCurveIndex, 0.0f, 1.0f, 1.0f);
-		if (Metric < Metrics0.KMinMetrics[FMetrics::Count - 1])
-		{
-			int32 LastGuideIndex = SimCurveIndex;
-			float LastMetric = Metric;
-			for (uint32 Index = 0; Index < FMetrics::Count; ++Index)
-			{
-				if (Metric < Metrics0.KMinMetrics[Index])
-				{
-					SwapValue(Metrics0.KClosestGuideIndices[Index], LastGuideIndex);
-					SwapValue(Metrics0.KMinMetrics[Index], LastMetric);
-				}
-			}
-		}
-	}
-
-	struct FRootsGrid
-	{
-		FVector3f MinBound = FVector3f::ZeroVector;
-		FVector3f MaxBound = FVector3f::ZeroVector;
-		
-		const uint32 MaxLookupDistance = 31;
-		const FIntVector GridResolution = FIntVector(32, 32, 32);
-
-		TArray<int32> GridIndirection;
-		TArray<TArray<int32>> RootIndices;
-		
-		FORCEINLINE bool IsValid(const FIntVector& P) const
-		{
-			return	0 <= P.X && P.X < GridResolution.X &&
-					0 <= P.Y && P.Y < GridResolution.Y &&
-					0 <= P.Z && P.Z < GridResolution.Z;
-		}
-
-		FORCEINLINE FIntVector ClampToVolume(const FIntVector& CellCoord, bool& bIsValid) const
-		{
-			bIsValid = IsValid(CellCoord);
-			return FIntVector(
-				FMath::Clamp(CellCoord.X, 0, GridResolution.X - 1),
-				FMath::Clamp(CellCoord.Y, 0, GridResolution.Y - 1),
-				FMath::Clamp(CellCoord.Z, 0, GridResolution.Z - 1));
-		}
-
-		FORCEINLINE FIntVector ToCellCoord(const FVector3f& P) const
-		{
-			bool bIsValid = false;
-			const FVector3f F = ((P - MinBound) / (MaxBound - MinBound));
-			const FIntVector CellCoord = FIntVector(FMath::FloorToInt(F.X * GridResolution.X), FMath::FloorToInt(F.Y * GridResolution.Y), FMath::FloorToInt(F.Z * GridResolution.Z));			
-			return ClampToVolume(CellCoord, bIsValid);
-		}
-
-		uint32 ToIndex(const FIntVector& CellCoord) const
-		{
-			uint32 CellIndex = CellCoord.X + CellCoord.Y * GridResolution.X + CellCoord.Z * GridResolution.X * GridResolution.Y;
-			check(CellIndex < uint32(GridIndirection.Num()));
-			return CellIndex;
-		}
-
-		void InsertRoots(TArray<FHairRoot>& Roots, const FVector3f& InMinBound, const FVector3f& InMaxBound)
-		{
-			MinBound = InMinBound;
-			MaxBound = InMaxBound;
-			GridIndirection.SetNumZeroed(GridResolution.X*GridResolution.Y*GridResolution.Z);
-			RootIndices.Empty();
-			RootIndices.AddDefaulted(); // Add a default empty list for the null element
-
-			const uint32 RootCount = Roots.Num();
-			for (uint32 RootIt = 0; RootIt < RootCount; ++RootIt)
-			{
-				const FHairRoot& Root = Roots[RootIt];
-				const FIntVector CellCoord = ToCellCoord(Root.Position);
-				const uint32 Index = ToIndex(CellCoord);
-				if (GridIndirection[Index] == 0)
-				{
-					GridIndirection[Index] = RootIndices.Num();
-					RootIndices.AddDefaulted();
-				}
-				
-				TArray<int32>& CellGuideIndices = RootIndices[GridIndirection[Index]];
-				CellGuideIndices.Add(RootIt);
-			}
-		}
-
-		FORCEINLINE void SearchCell(
-			const FIntVector& CellCoord,
-			const uint32 RenCurveIndex,
-			const FHairRoot& RenRoot,
-			const TArray<FHairRoot>& SimRoots,
-			FMetrics& Metrics) const 
-		{
-			const uint32 CellIndex = ToIndex(CellCoord);
-
-			if (GridIndirection[CellIndex] == 0)
-				return;
-
-			const TArray<int32>& Elements = RootIndices[GridIndirection[CellIndex]];
-
-			for (int32 SimCurveIndex : Elements)
-			{
-				const FHairRoot& GuideRoot = SimRoots[SimCurveIndex];
-				{
-					ComputeSimpleMetric(Metrics, RenRoot, GuideRoot, RenCurveIndex, SimCurveIndex);
-				}
-			}
-		}
-
-		FClosestGuides FindClosestRoots(
-			const uint32 RenCurveIndex,
-			const TArray<FHairRoot>& RenRoots,
-			const TArray<FHairRoot>& SimRoots,
-			const FHairStrandsDatas& RenStrandsData,
-			const FHairStrandsDatas& SimStrandsData) const
-		{
-			const FHairRoot& RenRoot = RenRoots[RenCurveIndex];
-			const FIntVector PointCoord = ToCellCoord(RenRoot.Position);
-
-			FMetrics Metrics;
-			for (uint32 ClearIt = 0; ClearIt < FMetrics::Count; ++ClearIt)
-			{
-				Metrics.KMinMetrics[ClearIt] = FLT_MAX;
-				Metrics.KClosestGuideIndices[ClearIt] = FMetrics::Invalid;
-			}
-
-			for (int32 Offset = 1; Offset <= int32(MaxLookupDistance); ++Offset)
-			{
-				// Center
-				{
-					bool bIsValid = false;
-					const FIntVector CellCoord = ClampToVolume(PointCoord, bIsValid);
-					if (bIsValid) SearchCell(PointCoord, RenCurveIndex, RenRoot, SimRoots, Metrics);
-				}
-
-				// Top & Bottom
-				for (int32 X = -Offset; X <= Offset; ++X)
-				for (int32 Y = -Offset; Y <= Offset; ++Y)
-				{
-					bool bIsValid0 = false, bIsValid1 = false;
-					const FIntVector CellCoord0 = ClampToVolume(PointCoord + FIntVector(X, Y, Offset), bIsValid0);
-					const FIntVector CellCoord1 = ClampToVolume(PointCoord + FIntVector(X, Y,-Offset), bIsValid1);
-					if (bIsValid0) SearchCell(CellCoord0, RenCurveIndex, RenRoot, SimRoots, Metrics);
-					if (bIsValid1) SearchCell(CellCoord1, RenCurveIndex, RenRoot, SimRoots, Metrics);
-				}
-
-				const int32 OffsetMinusOne = Offset - 1;
-				// Front & Back
-				for (int32 X = -Offset; X <= Offset; ++X)
-				for (int32 Z = -OffsetMinusOne; Z <= OffsetMinusOne; ++Z)
-				{
-					bool bIsValid0 = false, bIsValid1 = false;
-					const FIntVector CellCoord0 = ClampToVolume(PointCoord + FIntVector(X,  Offset, Z), bIsValid0);
-					const FIntVector CellCoord1 = ClampToVolume(PointCoord + FIntVector(X, -Offset, Z), bIsValid1);
-					if (bIsValid0) SearchCell(CellCoord0, RenCurveIndex, RenRoot, SimRoots, Metrics);
-					if (bIsValid1) SearchCell(CellCoord1, RenCurveIndex, RenRoot, SimRoots, Metrics);
-				}
-				
-				// Left & Right
-				for (int32 Y = -OffsetMinusOne; Y <= OffsetMinusOne; ++Y)
-				for (int32 Z = -OffsetMinusOne; Z <= OffsetMinusOne; ++Z)
-				{
-					bool bIsValid0 = false, bIsValid1 = false;
-					const FIntVector CellCoord0 = ClampToVolume(PointCoord + FIntVector( Offset, Y, Z), bIsValid0);
-					const FIntVector CellCoord1 = ClampToVolume(PointCoord + FIntVector(-Offset, Y, Z), bIsValid1);
-					if (bIsValid0) SearchCell(CellCoord0, RenCurveIndex, RenRoot, SimRoots, Metrics);
-					if (bIsValid1) SearchCell(CellCoord1, RenCurveIndex, RenRoot, SimRoots, Metrics);
-				}
-
-				// Early out if we have found all closest guide during a ring/layer step.
-				// This early out is not conservative, as the complex metric might find better guides one or multiple step further.
-				if (Metrics.KClosestGuideIndices[FMetrics::Count-1] != FMetrics::Invalid && Offset >= 2)
-				{
-					break;
-				}
-			}
-
-			// If no valid guide have been found, switch to a simpler metric
-			FClosestGuides ClosestGuides;
-			SelectFinalGuides(ClosestGuides, Metrics);
-			return ClosestGuides;
-		}
-
-		FORCEINLINE void SearchCell(
-			const FIntVector& CellCoord,
-			const uint32 RenCurveIndex,
-			const FHairRoot& RenRoot,
-			const TArray<FHairRoot>& SimRoots,
-			const FHairStrandsDatas& RenStrandsData,
-			const FHairStrandsDatas& SimStrandsData,
-			FMetrics& Metrics0,
-			FMetrics& Metrics1) const
-		{
-			const uint32 CellIndex = ToIndex(CellCoord);
-
-			if (GridIndirection[CellIndex] == 0)
-				return;
-
-			const TArray<int32>& Elements = RootIndices[GridIndirection[CellIndex]];
-
-			for (int32 SimCurveIndex : Elements)
-			{
-				const FHairRoot& GuideRoot = SimRoots[SimCurveIndex];
-				{
-					ComputeSimpleMetric(Metrics1, RenRoot, GuideRoot, RenCurveIndex, SimCurveIndex);
-					ComputeAdvandedMetric(Metrics0, RenStrandsData, SimStrandsData, RenCurveIndex, SimCurveIndex);
-				}
-			}
-		}
-
-		FClosestGuides FindBestClosestRoots(
-			const uint32 RenCurveIndex,
-			const TArray<FHairRoot>& RenRoots,
-			const TArray<FHairRoot>& SimRoots,
-			const FHairStrandsDatas& RenStrandsData,
-			const FHairStrandsDatas& SimStrandsData,
-			const bool bRandomized,
-			const bool bUnique,
-			const FIntVector& RandomIndices) const
-		{
-			const FHairRoot& RenRoot = RenRoots[RenCurveIndex];
-			const FIntVector PointCoord = ToCellCoord(RenRoot.Position);
-
-			FMetrics Metrics0;
-			FMetrics Metrics1;
-			for (uint32 ClearIt = 0; ClearIt < FMetrics::Count; ++ClearIt)
-			{
-				Metrics0.KMinMetrics[ClearIt] = FLT_MAX;
-				Metrics0.KClosestGuideIndices[ClearIt] = -1;
-				Metrics1.KMinMetrics[ClearIt] = FLT_MAX;
-				Metrics1.KClosestGuideIndices[ClearIt] = -1;
-			}
-
-			for (int32 Offset = 1; Offset <= int32(MaxLookupDistance); ++Offset)
-			{
-				// Center
-				{
-					bool bIsValid = false;
-					const FIntVector CellCoord = ClampToVolume(PointCoord, bIsValid);
-					if (bIsValid) SearchCell(CellCoord, RenCurveIndex, RenRoot, SimRoots, RenStrandsData, SimStrandsData, Metrics0, Metrics1);
-				}
-
-				// Top & Bottom
-				for (int32 X = -Offset; X <= Offset; ++X)
-				for (int32 Y = -Offset; Y <= Offset; ++Y)
-				{
-					bool bIsValid0 = false, bIsValid1 = false;
-					const FIntVector CellCoord0 = ClampToVolume(PointCoord + FIntVector(X, Y, Offset), bIsValid0);
-					const FIntVector CellCoord1 = ClampToVolume(PointCoord + FIntVector(X, Y,-Offset), bIsValid1);
-					if (bIsValid0) SearchCell(CellCoord0, RenCurveIndex, RenRoot, SimRoots, RenStrandsData, SimStrandsData, Metrics0, Metrics1);
-					if (bIsValid1) SearchCell(CellCoord1, RenCurveIndex, RenRoot, SimRoots, RenStrandsData, SimStrandsData, Metrics0, Metrics1);
-				}
-
-				const int32 OffsetMinusOne = Offset - 1;
-				// Front & Back
-				for (int32 X = -Offset; X <= Offset; ++X)
-				for (int32 Z = -OffsetMinusOne; Z <= OffsetMinusOne; ++Z)
-				{
-					bool bIsValid0 = false, bIsValid1 = false;
-					const FIntVector CellCoord0 = ClampToVolume(PointCoord + FIntVector(X, Offset, Z), bIsValid0);
-					const FIntVector CellCoord1 = ClampToVolume(PointCoord + FIntVector(X, -Offset, Z), bIsValid1);
-					if (bIsValid0) SearchCell(CellCoord0, RenCurveIndex, RenRoot, SimRoots, RenStrandsData, SimStrandsData, Metrics0, Metrics1);
-					if (bIsValid1) SearchCell(CellCoord1, RenCurveIndex, RenRoot, SimRoots, RenStrandsData, SimStrandsData, Metrics0, Metrics1);
-				}
-				
-				// Left & Right
-				for (int32 Y = -OffsetMinusOne; Y <= OffsetMinusOne; ++Y)
-				for (int32 Z = -OffsetMinusOne; Z <= OffsetMinusOne; ++Z)
-				{
-					bool bIsValid0 = false, bIsValid1 = false;
-					const FIntVector CellCoord0 = ClampToVolume(PointCoord + FIntVector( Offset, Y, Z), bIsValid0);
-					const FIntVector CellCoord1 = ClampToVolume(PointCoord + FIntVector(-Offset, Y, Z), bIsValid1);
-					if (bIsValid0) SearchCell(CellCoord0, RenCurveIndex, RenRoot, SimRoots, RenStrandsData, SimStrandsData, Metrics0, Metrics1);
-					if (bIsValid1) SearchCell(CellCoord1, RenCurveIndex, RenRoot, SimRoots, RenStrandsData, SimStrandsData, Metrics0, Metrics1);
-				}
-
-				// Early out if we have found all closest guide during a ring/layer step.
-				// This early out is not conservative, as the complex metric might find better guides one or multiple step further.
-				if ((Metrics0.KClosestGuideIndices[FMetrics::Count-1] != -1 || Metrics1.KClosestGuideIndices[FMetrics::Count - 1] != -1) && Offset >= 2)
-				{
-					break;
-				}
-			}
-
-			// If no valid guide have been found, switch to a simpler metric
-			FClosestGuides ClosestGuides;
-			if (Metrics0.KClosestGuideIndices[0] != -1)
-			{
-				SelectFinalGuides(ClosestGuides, Metrics0);
-			}
-			else
-			{
-				SelectFinalGuides(ClosestGuides, Metrics1);
-			}
-
-			// Check there is at least one valid guide
-			check(ClosestGuides.Indices[0] != FClosestGuides::Invalid || ClosestGuides.Indices[1] != FClosestGuides::Invalid || ClosestGuides.Indices[2] != FClosestGuides::Invalid);
-
-			return ClosestGuides;
-		}
-	};
-
-	static FClosestGuides FindBestRoots(
-		const uint32 RenCurveIndex,
-		const TArray<FHairRoot>& RenRoots,
-		const TArray<FHairRoot>& SimRoots,
-		const FHairStrandsDatas& RenStrandsData,
-		const FHairStrandsDatas& SimStrandsData)
-	{
-		FMetrics Metrics;
-		for (uint32 ClearIt = 0; ClearIt < FMetrics::Count; ++ClearIt)
-		{
-			Metrics.KMinMetrics[ClearIt] = FLT_MAX;
-			Metrics.KClosestGuideIndices[ClearIt] = -1;
-		}
-
-		const uint32 SimRootsCount = SimRoots.Num();
-		for (uint32 SimCurveIndex =0; SimCurveIndex<SimRootsCount; ++SimCurveIndex)
-		{
-			ComputeAdvandedMetric(Metrics, RenStrandsData, SimStrandsData, RenCurveIndex, SimCurveIndex);
-		}
-			
-		FClosestGuides ClosestGuides;
-		SelectFinalGuides(ClosestGuides, Metrics);
-
-		// Check there is at least one valid guide
-		check(ClosestGuides.Indices[0] != FClosestGuides::Invalid || ClosestGuides.Indices[1] != FClosestGuides::Invalid || ClosestGuides.Indices[2] != FClosestGuides::Invalid);
-
-		return ClosestGuides;
-	}
-
-	// Extract strand roots
-	static void ExtractRoots(const FHairStrandsDatas& InData, TArray<FHairRoot>& OutRoots, FVector3f& MinBound, FVector3f& MaxBound)
-	{
-		MinBound = FVector3f(FLT_MAX, FLT_MAX, FLT_MAX);
-		MaxBound = FVector3f(-FLT_MAX, -FLT_MAX, -FLT_MAX);
-		const uint32 CurveCount = InData.StrandsCurves.Num();
-		OutRoots.Reserve(CurveCount);
-		for (uint32 CurveIndex = 0; CurveIndex < CurveCount; ++CurveIndex)
-		{
-			const uint32 PointOffset = InData.StrandsCurves.CurvesOffset[CurveIndex];
-			const uint32 PointCount = InData.StrandsCurves.CurvesCount[CurveIndex];
-			const float  CurveLength = InData.StrandsCurves.CurvesLength[CurveIndex];
-			check(PointCount > 1);
-			const FVector3f& P0 = InData.StrandsPoints.PointsPosition[PointOffset];
-			const FVector3f& P1 = InData.StrandsPoints.PointsPosition[PointOffset + 1];
-			FVector3f N = (P1 - P0).GetSafeNormal();
-
-			// Fallback in case the initial points are too close (this happens on certain assets)
-			if (FVector3f::DotProduct(N, N) == 0)
-			{
-				N = FVector3f(0, 0, 1);
-			}
-			OutRoots.Add({ P0, PointCount, N, PointOffset, CurveLength });
-
-			MinBound = MinBound.ComponentMin(P0);
-			MaxBound = MaxBound.ComponentMax(P0);
-		}
-	}
-
-	struct FVertexInterpolationDesc
-	{
-		uint32 Index0 = 0;
-		uint32 Index1 = 0;
-		float T = 0;
-	};
-
-	// Find the vertex along a sim curve 'SimCurveIndex', which has the same parametric distance than the render distance 'RenPointDistance'
-	static FVertexInterpolationDesc FindMatchingVertex(const float RenPointDistance, const FHairStrandsDatas& SimStrandsData, const uint32 SimCurveIndex)
-	{
-		const uint32 SimOffset = SimStrandsData.StrandsCurves.CurvesOffset[SimCurveIndex];
-
-		const float CurveLength = SimStrandsData.StrandsCurves.CurvesLength[SimCurveIndex];
-
-		// Find with with vertex the vertex should be paired
-		const uint32 SimPointCount = SimStrandsData.StrandsCurves.CurvesCount[SimCurveIndex];
-		for (uint32 SimPointIndex = 0; SimPointIndex < SimPointCount-1; ++SimPointIndex)
-		{
-			const float SimPointDistance0 = SimStrandsData.StrandsPoints.PointsCoordU[SimPointIndex + SimOffset] * CurveLength;
-			const float SimPointDistance1 = SimStrandsData.StrandsPoints.PointsCoordU[SimPointIndex + SimOffset + 1] * CurveLength;
-			if (SimPointDistance0 <= RenPointDistance && RenPointDistance <= SimPointDistance1)
-			{
-				const float SegmentLength = SimPointDistance1 - SimPointDistance0;
-				FVertexInterpolationDesc Out;
-				Out.Index0 = SimPointIndex;
-				Out.Index1 = SimPointIndex+1;
-				Out.T = (RenPointDistance - SimPointDistance0) / (SegmentLength>0? SegmentLength : 1);
-				Out.T = FMath::Clamp(Out.T, 0.f, 1.f);
-				return Out;
-			}
-		}
-		FVertexInterpolationDesc Desc;
-		Desc.Index0 = SimPointCount-2;
-		Desc.Index1 = SimPointCount-1;
-		Desc.T = 1;
-		return Desc;
-	}
-
-	void BuildInternalData(FHairStrandsDatas& HairStrands)
-	{
-		TRACE_CPUPROFILER_EVENT_SCOPE(HairStrandsBuilder::BuildInternalData);
-
-		FHairStrandsCurves& Curves = HairStrands.StrandsCurves;
-		FHairStrandsPoints& Points = HairStrands.StrandsPoints;
-
-		HairStrands.BoundingBox.Min = { FLT_MAX,  FLT_MAX ,  FLT_MAX };
-		HairStrands.BoundingBox.Max = { -FLT_MAX, -FLT_MAX , -FLT_MAX };
-
-		if (HairStrands.GetNumCurves() > 0 && HairStrands.GetNumPoints() > 0)
-		{
-			TArray<FVector3f>::TIterator PositionIterator = Points.PointsPosition.CreateIterator();
-			TArray<float>::TIterator RadiusIterator = Points.PointsRadius.CreateIterator();
-			TArray<float>::TIterator CoordUIterator = Points.PointsCoordU.CreateIterator();
-
-			TArray<uint16>::TIterator CountIterator = Curves.CurvesCount.CreateIterator();
-			TArray<uint32>::TIterator OffsetIterator = Curves.CurvesOffset.CreateIterator();
-			TArray<float>::TIterator LengthIterator = Curves.CurvesLength.CreateIterator();
-
-			uint32 StrandOffset = 0;
-			*OffsetIterator = StrandOffset; ++OffsetIterator;
-
-			for (uint32 CurveIndex = 0; CurveIndex < HairStrands.GetNumCurves(); ++CurveIndex, ++OffsetIterator, ++LengthIterator, ++CountIterator)
-			{
-				const uint16& StrandCount = *CountIterator;
-
-				StrandOffset += StrandCount;
-				*OffsetIterator = StrandOffset;
-
-				float StrandLength = 0.0;
-				FVector3f PreviousPosition(0.0, 0.0, 0.0);
-				for (uint32 PointIndex = 0; PointIndex < StrandCount; ++PointIndex, ++PositionIterator, ++RadiusIterator, ++CoordUIterator)
-				{
-					HairStrands.BoundingBox += (FVector)*PositionIterator;
-
-					if (PointIndex > 0)
-					{
-						StrandLength += (*PositionIterator - PreviousPosition).Size();
-					}
-					*CoordUIterator = StrandLength;
-					PreviousPosition = *PositionIterator;
-
-				}
-				*LengthIterator = StrandLength;
-			}
-
-			CountIterator.Reset();
-			LengthIterator.Reset();
-			RadiusIterator.Reset();
-			CoordUIterator.Reset();
-
-			for (uint32 CurveIndex = 0; CurveIndex < HairStrands.GetNumCurves(); ++CurveIndex, ++LengthIterator, ++CountIterator)
-			{
-				const uint16& StrandCount = *CountIterator;
-
-				for (uint32 PointIndex = 0; PointIndex < StrandCount; ++PointIndex, ++RadiusIterator, ++CoordUIterator)
-				{
-					*CoordUIterator /= *LengthIterator;
-				}
-			}
-		}
-	}
-
-	void GenerateGuides(const FHairStrandsDatas& InData, float DecimationPercentage, FHairStrandsDatas& OutData)
-	{
-		// Pick randomly strand as guide 
-		// Divide strands in buckets and pick randomly one stand per bucket
-		const uint32 CurveCount = InData.StrandsCurves.Num();
-		const uint32 OutCurveCount = FMath::Clamp(uint32(CurveCount * DecimationPercentage), 1u, CurveCount);
-
-		const uint32 BucketSize = CurveCount / OutCurveCount;
-		const bool bHasClumpID = InData.StrandsCurves.ClumpIDs.Num() > 0;
-
-		TArray<uint32> CurveIndices;
-		CurveIndices.SetNum(OutCurveCount);
-
-		uint32 OutTotalPointCount = 0;
-		FRandomStream Random;
-		Random.Initialize(0xdeedbeed);
-		for (uint32 BucketIndex = 0; BucketIndex < OutCurveCount; BucketIndex++)
-		{
-			const uint32 CurveIndex = BucketIndex * BucketSize;// +BucketSize * Random.FRand();
-			CurveIndices[BucketIndex] = CurveIndex;
-			OutTotalPointCount += InData.StrandsCurves.CurvesCount[CurveIndex];
-		}
-
-		OutData.StrandsCurves.SetNum(OutCurveCount, 0u);
-		OutData.StrandsPoints.SetNum(OutTotalPointCount, 0u);
-		OutData.HairDensity = InData.HairDensity;
-
-		const uint32 InAttribute = InData.GetAttributes();
-
-		uint32 OutPointOffset = 0;
-		for (uint32 OutCurveIndex = 0; OutCurveIndex < OutCurveCount; ++OutCurveIndex)
-		{
-			const uint32 InCurveIndex = CurveIndices[OutCurveIndex];
-			const uint32 InPointOffset = InData.StrandsCurves.CurvesOffset[InCurveIndex];
-			const uint32 PointCount = InData.StrandsCurves.CurvesCount[InCurveIndex];
-			OutData.StrandsCurves.CurvesCount[OutCurveIndex] = PointCount;
-			OutData.StrandsCurves.CurvesOffset[OutCurveIndex] = OutPointOffset;
-			FHairStrandsDatas::CopyCurve(InData, OutData, InAttribute, InCurveIndex, OutCurveIndex);
-
-			for (uint32 PointIndex = 0; PointIndex < PointCount; ++PointIndex)
-			{
-				FHairStrandsDatas::CopyPoint(InData, OutData, InAttribute, PointIndex + InPointOffset, PointIndex + OutPointOffset);
-			}
-			OutPointOffset += PointCount;
-		}
-
-		BuildInternalData(OutData);
-	}
-
-	struct FHairStrandsClusters
-	{
-		TArray<uint32>						AtlasRectIndex; 		// Representative rect index (i.e., index to AtlasUniqueRectIndex)
-		TArray<uint32>						AtlasUniqueRectIndex;	// Index corresponduing to the representative cluster 
-		TArray<FHairCardsProceduralGeometry::Rect>	AtlasUniqueRect;		// Rect info corresponding to the representative cluster
-
-		TArray<uint32> CurveIndices;
-		TArray<uint32> CurveOffset;
-		TArray<uint32> CurveCount;
-
-		TArray<uint32>  GuidePointCount;
-		TArray<uint32>  GuidePointOffset;
-		TArray<FVector3f> GuidePoints;
-		TArray<FVector2f> GuideRootUVs;
-
-		TArray<FHairOrientedBound> Bounds;
-
-		// Cluster bound/envelop
-		TArray<FVector3f> BoundPositions;
-		TArray<FVector3f> BoundAxis_T;
-		TArray<FVector3f> BoundAxis_B;
-		TArray<FVector3f> BoundAxis_N;
-
-		TArray<uint32> BoundOffset;
-		TArray<uint32> BoundCount;
-		TArray<float>  BoundLength; // Total length of the cluster
-
-		uint32 GetNum() const { return CurveCount.Num();  }
-		void SetNum(uint32 Count)
-		{
-			BoundPositions.Empty();
-			BoundAxis_T.Empty();
-			BoundAxis_B.Empty();
-			BoundAxis_N.Empty();
-
-			BoundOffset.SetNum(Count);
-			BoundCount.SetNum(Count);
-			BoundLength.SetNum(Count);
-			
-			GuidePointCount.SetNum(Count);
-			GuidePointOffset.SetNum(Count);
-			GuideRootUVs.SetNum(Count);
-			GuidePoints.Empty();
-
-			CurveIndices.Empty();
-			CurveOffset.SetNum(Count);
-			CurveCount.SetNum(Count);
-			Bounds.SetNum(Count);
-
-			AtlasRectIndex.SetNum(Count);
-			AtlasUniqueRectIndex.Empty();
-			AtlasUniqueRect.Empty();
-		}
-	};
-
-	static void BuildClusterData(
-		const FHairStrandsDatas& RenStrandsData,
-		const FHairStrandsDatas& InSimStrandsData,
-		FHairStrandsClusters& OutClusters,
-		const FHairCardsGeometrySettings& InSettings)
-	{
-		TRACE_CPUPROFILER_EVENT_SCOPE(HairCards::BuildClusterData);
-
-		FHairStrandsDatas SimStrandsData;
-		if (InSettings.GenerationType == EHairCardsGenerationType::UseGuides)
-		{
-			SimStrandsData = InSimStrandsData;
-		}
-		else if (InSettings.GenerationType == EHairCardsGenerationType::CardsCount)
-		{
-			const float DecimationFactor = FMath::Clamp(InSettings.CardsCount / float(RenStrandsData.GetNumCurves()), 0.f, 1.f);
-			GenerateGuides(RenStrandsData, FMath::Clamp(DecimationFactor, 0.f, 1.f), SimStrandsData);
-		}
-		else
-		{
-			check(false); // Not implemented
-		}
-
-		typedef TArray<FHairRoot> FRoots;
-
-		// Build acceleration structure for fast nearest-neighbors lookup.
-		// This is used only for low quality interpolation as high quality 
-		// interpolation require broader search
-		FRoots RenRoots, SimRoots;
-		FRootsGrid RootsGrid;
-		{
-			FVector3f RenMinBound, RenMaxBound;
-			FVector3f SimMinBound, SimMaxBound;
-			ExtractRoots(RenStrandsData, RenRoots, RenMinBound, RenMaxBound);
-			ExtractRoots(SimStrandsData, SimRoots, SimMinBound, SimMaxBound);
-
-			if (InSettings.ClusterType == EHairCardsClusterType::Low)
-			{
-				// Build a conservative bound, to insure all queries will fall 
-				// into the grid volume.
-				const FVector3f MinBound = RenMinBound.ComponentMin(SimMinBound);
-				const FVector3f MaxBound = RenMaxBound.ComponentMax(SimMaxBound);
-				RootsGrid.InsertRoots(SimRoots, MinBound, MaxBound);
-			}
-		}
-
-		// Find k-closest guide:
-		const static float MinWeightDistance = 0.0001f;
-
-		const uint32 RenCurveCount = RenStrandsData.GetNumCurves();
-		const uint32 SimCurveCount = SimStrandsData.GetNumCurves();
-
-		TAtomic<uint32> CompletedTasks(0);
-		FScopedSlowTask SlowTask(RenCurveCount, LOCTEXT("BuildHairClusterData", "Building hair cluster for cards generation"));
-		SlowTask.MakeDialog();
-
-		uint32 ClusterCount = SimStrandsData.GetNumCurves();
-		TArray<TAtomic<uint32>> ClusterCurveIndexCount;
-		ClusterCurveIndexCount.SetNumZeroed(ClusterCount);
-		TAtomic<uint32> TotalCurveIndices(0);
-
-		const float SimStrandsMaxLength = GetHairStrandsMaxLength(SimStrandsData);
-
-		TArray<FClosestGuides> ClosestGuides;
-		ClosestGuides.SetNum(RenCurveCount);
-		const FDateTime StartTime = FDateTime::UtcNow();
-		ParallelFor(RenCurveCount, 
-		[
-			StartTime,
-			InSettings,
-			RenCurveCount, &RenRoots, &RenStrandsData,
-			SimCurveCount, &SimRoots, &SimStrandsData, 
-			SimStrandsMaxLength,
-			&RootsGrid,
-			&CompletedTasks,
-			&SlowTask,
-			&ClosestGuides,
-			&TotalCurveIndices,
-			&ClusterCurveIndexCount
-		] (uint32 RenCurveIndex) 
-		//for (uint32 RenCurveIndex = 0; RenCurveIndex < RenCurveCount; ++RenCurveIndex)
-		{
-			TRACE_CPUPROFILER_EVENT_SCOPE(HairCards::ComputeClusters);
-
-			++CompletedTasks;
-
-			if (IsInGameThread())
-			{
-				const uint32 CurrentCompletedTasks = CompletedTasks.Exchange(0);
-				const float RemainingTasks = FMath::Clamp(SlowTask.TotalAmountOfWork - SlowTask.CompletedWork, 1.f, SlowTask.TotalAmountOfWork);
-				const FTimespan ElaspedTime = FDateTime::UtcNow() - StartTime;
-				const double RemainingTimeInSeconds = RemainingTasks * double(ElaspedTime.GetSeconds() / SlowTask.CompletedWork);
-
-				FTextBuilder TextBuilder;
-				TextBuilder.AppendLineFormat(LOCTEXT("ComputeCluster", "Computing closest guides and weights ({0})"), FText::AsTimespan(FTimespan::FromSeconds(RemainingTimeInSeconds)));
-				SlowTask.EnterProgressFrame(CurrentCompletedTasks, TextBuilder.ToText());
-			}
-
-			if (InSettings.ClusterType == EHairCardsClusterType::Low)
-			{
-				ClosestGuides[RenCurveIndex] = RootsGrid.FindClosestRoots(RenCurveIndex, RenRoots, SimRoots, RenStrandsData, SimStrandsData);
-			}
-			else // (InSettings.ClusterType == EHairCardsClusterType::High)
-			{
-				ClosestGuides[RenCurveIndex] = FindBestRoots(RenCurveIndex, RenRoots, SimRoots, RenStrandsData, SimStrandsData);
-			}
-
-			for (uint32 KIndex = 0; KIndex < FClosestGuides::Count; ++KIndex)
-			{
-				const uint32 GuideIndex = ClosestGuides[RenCurveIndex].Indices[KIndex];
-				if (GuideIndex != FClosestGuides::Invalid)
-				{
-					ClusterCurveIndexCount[GuideIndex] += 1;
-					++TotalCurveIndices;
-				}
-			}
-		});
-
-		// Insure the min length threshold is lower than the max curves length, to generate at least one cluster. 
-		// This is arbitrarely set to 95% of the longest strands
-		float MinLengthTreshold = FMath::Min(InSettings.MinCardsLength, SimStrandsMaxLength * 0.95f);
-		float MaxLengthTreshold = InSettings.MaxCardsLength > 0 ? FMath::Min(InSettings.MaxCardsLength, SimStrandsMaxLength) : SimStrandsMaxLength;
-
-		// Compute the number of clusters
-		// * Some cluster might have no curve this is why we remove them
-		// * Cluster/guide which are shorter than a certain length are cut
-		const uint32 GuideCount = SimStrandsData.GetNumCurves();
-		TArray<int32> GuideIndexToClusterIndex;
-		GuideIndexToClusterIndex.SetNum(GuideCount);
-		{
-			uint32 ValidClusterCount = 0;
-			while (ValidClusterCount == 0)
-			{
-				for (uint32 GuideIndex = 0; GuideIndex < GuideCount; ++GuideIndex)
-				{				
-					const float GuideLength = SimStrandsData.StrandsCurves.CurvesLength[GuideIndex];
-					const bool bIsCurveValid = GuideLength >= MinLengthTreshold && GuideLength <= MaxLengthTreshold;
-					if (bIsCurveValid && ClusterCurveIndexCount[GuideIndex] > 0)
-					{
-						GuideIndexToClusterIndex[GuideIndex] = ValidClusterCount;
-						++ValidClusterCount;
-					}
-					else
-					{
-						GuideIndexToClusterIndex[GuideIndex] = -1;
-					}
-				}
-
-				// If we haven't found a single valid cluster, slowly decrease/increase the min/max length thresholds
-				if (ValidClusterCount == 0)
-				{
-					MinLengthTreshold = MinLengthTreshold > 0 ? MinLengthTreshold * 0.1f : 0;
-					MaxLengthTreshold = MaxLengthTreshold < SimStrandsMaxLength ? FMath::Min(SimStrandsMaxLength, MaxLengthTreshold * 1.1f) : SimStrandsMaxLength;
-				}				
-			}
-			ClusterCount = FMath::Min(ValidClusterCount, uint32(FMath::Max(1, GHairCardsMaxClusterCount)));
-			OutClusters.SetNum(ClusterCount);
-			OutClusters.CurveIndices.SetNum(TotalCurveIndices);
-		}
-
-		// Compute curve offset (need to be serial)
-		{
-			uint32 ClusterIt = 0;
-			uint32 Offset = 0;
-			for (uint32 GuideIndex = 0; GuideIndex < GuideCount; ++GuideIndex)
-			{
-				const float GuideLength = SimStrandsData.StrandsCurves.CurvesLength[GuideIndex];
-				const bool bIsCurveValid = GuideLength >= MinLengthTreshold && GuideLength <= MaxLengthTreshold;
-				if (bIsCurveValid && ClusterCurveIndexCount[GuideIndex] > 0)
-				{
-					OutClusters.CurveOffset[ClusterIt] = Offset;
-					Offset += ClusterCurveIndexCount[GuideIndex];
-					++ClusterIt;
-				}
-			}
-		}
-
-		// Fill in curve indices (TODO: run in parallel)
-		for (uint32 RenCurveIndex = 0; RenCurveIndex < RenCurveCount; ++RenCurveIndex)
-		{
-			const FClosestGuides& G = ClosestGuides[RenCurveIndex];
-			for (uint32 KIndex = 0; KIndex < FClosestGuides::Count; ++KIndex)
-			{
-				const int32 GuideIndex = G.Indices[KIndex];
-				if (GuideIndex != FClosestGuides::Invalid && GuideIndex < int32(GuideCount) && GuideIndexToClusterIndex[GuideIndex] >= 0)
-				{
-					const uint32 ClusterIt = GuideIndexToClusterIndex[GuideIndex];
-
-					const uint32 Offset = OutClusters.CurveOffset[ClusterIt];
-					const uint32 Count  = OutClusters.CurveCount[ClusterIt];
-					OutClusters.CurveIndices[Offset + Count] = RenCurveIndex;
-					OutClusters.CurveCount[ClusterIt]++;
-				}
-			}			
-		}
-
-		const uint32 InAttribute = SimStrandsData.GetAttributes();
-
-		// Store the cluster guide points
-		uint32 ClusterOffset = 0;
-		OutClusters.GuidePoints.SetNum(SimStrandsData.GetNumPoints());
-		for (uint32 GuideIndex = 0; GuideIndex < GuideCount; ++GuideIndex)
-		{
-			const uint32 SimCount  = SimStrandsData.StrandsCurves.CurvesCount[GuideIndex];
-			const uint32 SimOffset = SimStrandsData.StrandsCurves.CurvesOffset[GuideIndex];
-			const FVector2f SimRootUV = HasHairAttribute(InAttribute, EHairAttribute::RootUV) ? SimStrandsData.StrandsCurves.CurvesRootUV[GuideIndex] : FVector2f(0,0);
-
-			const uint32 ClusterIt = GuideIndexToClusterIndex[GuideIndex];
-			if (GuideIndexToClusterIndex[GuideIndex] >= 0)
-			{
-				OutClusters.GuidePointCount[ClusterIt]  = SimCount;
-				OutClusters.GuidePointOffset[ClusterIt] = ClusterOffset;
-				OutClusters.GuideRootUVs[ClusterIt] = SimRootUV;
-				for (uint32 ClusterVertIt = 0; ClusterVertIt < SimCount; ++ClusterVertIt)
-				{
-					OutClusters.GuidePoints[ClusterOffset + ClusterVertIt] = SimStrandsData.StrandsPoints.PointsPosition[SimOffset + ClusterVertIt];
-				}
-
-				ClusterOffset += SimCount;
-			}
-		}
-	}
-
-	static Eigen::Matrix3f ComputeCovarianceMatrix(TArray<FVector3f> InPoints, FVector3f& Mean)
-	{
-		// Mean
-		const uint32 PointCount = InPoints.Num();
-		Mean = FVector3f::ZeroVector;
-		for (FVector3f& P : InPoints)
-		{
-			Mean += P;
-		}
-		Mean /= PointCount;
-
-		// Covariance
-		Eigen::Matrix3f Out;
-		Out << 0, 0, 0, 0, 0, 0, 0, 0, 0;
-		for (FVector3f& P : InPoints)
-		{
-			Out(0, 0) += (P.X - Mean.X) * (P.X - Mean.X);
-			Out(1, 0) += (P.X - Mean.X) * (P.Y - Mean.Y);
-			Out(2, 0) += (P.X - Mean.X) * (P.Z - Mean.Z);
-
-			Out(0, 1) += (P.Y - Mean.Y) * (P.X - Mean.X);
-			Out(1, 1) += (P.Y - Mean.Y) * (P.Y - Mean.Y);
-			Out(2, 1) += (P.Y - Mean.Y) * (P.Z - Mean.Z);
-
-			Out(0, 2) += (P.Z - Mean.Z) * (P.X - Mean.X);
-			Out(1, 2) += (P.Z - Mean.Z) * (P.Y - Mean.Y);
-			Out(2, 2) += (P.Z - Mean.Z) * (P.Z - Mean.Z);
-		}
-		Out(0, 0) /= PointCount;
-		Out(1, 0) /= PointCount;
-		Out(2, 0) /= PointCount;
-
-		Out(0, 1) /= PointCount;
-		Out(1, 1) /= PointCount;
-		Out(2, 1) /= PointCount;
-
-		Out(0, 2) /= PointCount;
-		Out(1, 2) /= PointCount;
-		Out(2, 2) /= PointCount;
-		return Out;
-	}
-
-	//static void ComputePCA(const TArray<FVector3f>& InPoints, FVector3f& Mean, FVector3f& B0, FVector3f& B1, FVector3f& B2)
-	static void ComputePCA(const TArray<FVector3f>& InPoints, const uint32 InPointOffset, const uint32 InPointCount, FVector3f& Mean, FVector3f& B0, FVector3f& B1, FVector3f& B2)
-	{
-		Eigen::Matrix3f covariance = ComputeCovarianceMatrix(InPoints, Mean);
-		Eigen::SelfAdjointEigenSolver<Eigen::Matrix3f> eigen_solver(covariance, Eigen::ComputeEigenvectors);
-		Eigen::Matrix3f eigenVectorsPCA = eigen_solver.eigenvectors();
-
-		B0 = FVector3f(eigenVectorsPCA.col(0)(0), eigenVectorsPCA.col(0)(1), eigenVectorsPCA.col(0)(2));
-		B1 = FVector3f(eigenVectorsPCA.col(1)(0), eigenVectorsPCA.col(1)(1), eigenVectorsPCA.col(1)(2));
-		B2 = FVector3f(eigenVectorsPCA.col(2)(0), eigenVectorsPCA.col(2)(1), eigenVectorsPCA.col(2)(2));
-
-		B0 = B0.GetSafeNormal();
-		B1 = B1.GetSafeNormal();
-		B2 = B2.GetSafeNormal();
-
-		FVector3f LocalMinBound( FLT_MAX);
-		FVector3f LocalMaxBound(-FLT_MAX);
-		//for (const FVector3f& P : InPoints)
-		for (uint32 PointIt=0;PointIt<InPointCount;++PointIt)
-		{
-			FVector3f PLocal = InPoints[InPointOffset+PointIt] - Mean;
-			PLocal = FVector3f(
-				FVector3f::DotProduct(PLocal, B0),
-				FVector3f::DotProduct(PLocal, B1),
-				FVector3f::DotProduct(PLocal, B2));
-
-			LocalMinBound.X = FMath::Min(LocalMinBound.X, PLocal.X);
-			LocalMinBound.Y = FMath::Min(LocalMinBound.Y, PLocal.Y);
-			LocalMinBound.Z = FMath::Min(LocalMinBound.Z, PLocal.Z);
-
-			LocalMaxBound.X = FMath::Max(LocalMaxBound.X, PLocal.X);
-			LocalMaxBound.Y = FMath::Max(LocalMaxBound.Y, PLocal.Y);
-			LocalMaxBound.Z = FMath::Max(LocalMaxBound.Z, PLocal.Z);
-		}
-
-		// Recompte the new mean in local space, and then compute its world position
-		const FVector3f LocalMean(
-			(LocalMaxBound.X + LocalMinBound.X) * 0.5f,
-			(LocalMaxBound.Y + LocalMinBound.Y) * 0.5f,
-			(LocalMaxBound.Z + LocalMinBound.Z) * 0.5f);
-
-		const FVector3f LocalExtent(
-			(LocalMaxBound.X - LocalMinBound.X) * 0.5f,
-			(LocalMaxBound.Y - LocalMinBound.Y) * 0.5f,
-			(LocalMaxBound.Z - LocalMinBound.Z) * 0.5f);
-
-		const FVector3f MeanPrime = Mean + B0 * LocalMean.X + B1 * LocalMean.Y + B2 * LocalMean.Z;
-		Mean = MeanPrime;
-
-		B0 *= LocalExtent.X;
-		B1 *= LocalExtent.Y;
-		B2 *= LocalExtent.Z;
-	}
-
-	static FHairOrientedBound ComputeOrientedBoundingBox(const TArray<FVector3f>& InPoints)
-	{
-		FHairOrientedBound Out;
-		ComputePCA(InPoints, 0, InPoints.Num(), Out.Center, Out.ExtentX, Out.ExtentY, Out.ExtentZ);
-		return Out;
-	}
-
-	static FHairOrientedBound ComputeOrientedBoundingBox(const TArray<FVector3f>& InPoints, const uint32 InPointOffset, const uint32 InPointCount)
-	{
-		FHairOrientedBound Out;
-		ComputePCA(InPoints, InPointOffset, InPointCount, Out.Center, Out.ExtentX, Out.ExtentY, Out.ExtentZ);
-		return Out;
-	}
-
-	static uint8 GetMajorAxis(const FHairOrientedBound& B)
-	{
-		uint8 MajorAxis = 0;
-		float S = B.ExtentX.Size(); 
-		if (B.ExtentY.Size() > S) { MajorAxis = 1; S = B.ExtentY.Size(); }
-		if (B.ExtentZ.Size() > S) { MajorAxis = 2; S = B.ExtentZ.Size(); }
-		return MajorAxis;
-	}
-
-	// Returns the axis index by decreasing size
-	static void GetSortedAxis(const FHairOrientedBound& B, uint8& Axis0, uint8& Axis1, uint8& Axis2)
-	{
-		float Size[3] = { (float)B.ExtentX.Size(), (float)B.ExtentY.Size(), (float)B.ExtentZ.Size() };
-		uint8 Axis[3] = { 0, 1, 2 };
-		if (Size[0] < Size[1]) { Swap(Axis[0], Axis[1]); Swap(Size[0], Size[1]); }
-		if (Size[1] < Size[2]) { Swap(Axis[1], Axis[2]); Swap(Size[1], Size[2]); }
-		if (Size[0] < Size[1]) { Swap(Axis[0], Axis[1]); Swap(Size[0], Size[1]); }
-
-		Axis0 = Axis[0];
-		Axis1 = Axis[1];
-		Axis2 = Axis[2];
-	}
-
-
-	inline FVector3f ToLocal(
-		const FVector3f& InCenter, 
-		const FVector3f& InAxisX, 
-		const FVector3f& InAxisY, 
-		const FVector3f& InAxisZ, 
-		const FVector3f& InP)
-	{
-		const FVector3f Local = InP - InCenter;
-		return FVector3f(
-			FVector3f::DotProduct(Local, InAxisX),
-			FVector3f::DotProduct(Local, InAxisY),
-			FVector3f::DotProduct(Local, InAxisZ));
-	}
-
-	inline FVector3f ToWorld(
-		const FVector3f& InCenter,
-		const FVector3f& InAxisX,
-		const FVector3f& InAxisY,
-		const FVector3f& InAxisZ,
-		const FVector3f& InP)
-	{
-		return 
-			InCenter + 
-			InP.X * InAxisX +
-			InP.Y * InAxisY +
-			InP.Z * InAxisZ;
-	}
-
-	inline FVector3f ToLocal(const FHairOrientedBound& InB, const FVector3f& InP)
-	{
-		return ToLocal(InB.Center, InB.ExtentX, InB.ExtentY, InB.ExtentZ, InP);
-	}
-
-
-	static void GetOBBPlane(const FHairOrientedBound& InB, FVector3f& OutT, FVector3f& OutB, FVector3f& OutN)
-	{
-		uint8 Axis0;
-		uint8 Axis1; 
-		uint8 Axis2;
-		GetSortedAxis(InB, Axis0, Axis1, Axis2);
-		OutT = InB.Extent(Axis0);
-		OutB = InB.Extent(Axis1);
-		OutN = InB.Extent(Axis2);
-	}
-	
-	static void Split(uint8 Axis, const FHairOrientedBound& In, FHairOrientedBound& Out0, FHairOrientedBound& Out1)
-	{
-		Out0 = In;
-		Out1 = In;
-
-		Out0.Center = In.Center + In.Extent(Axis) * 0.5f;
-		Out1.Center = In.Center - In.Extent(Axis) * 0.5f;
-
-		Out0.Extent(Axis) = In.Extent(Axis) * 0.5f;
-		Out1.Extent(Axis) = In.Extent(Axis) * 0.5f;
-	}
-
-	static TArray<FVector3f> Inside(const TArray<FVector3f>& InPoints, const FHairOrientedBound& InBound)
-	{
-		TArray<FVector3f> Out;
-		Out.Reserve(InPoints.Num() / 2);
-		for (const FVector3f& P : InPoints)
-		{
-			FVector3f PLocal = ToLocal(InBound, P);
-			const bool bInside =
-				FMath::Abs(PLocal.X) <= 1 &&
-				FMath::Abs(PLocal.Y) <= 1 &&
-				FMath::Abs(PLocal.Z) <= 1;
-
-			if (bInside)
-			{
-				Out.Add(P);
-			}
-		}
-
-		return Out;
-	}
-
-	static void ComputeOrientedBoundingBoxes(const TArray<FVector3f>& InPoints, float MinSectionLength, TArray<FHairOrientedBound>& Out, TArray<FHairOrientedBound>& OutAll, uint8 Depth=0, uint8 MaxDepth=0)
-	{
-		const float Scale = 1.2f;
-		FHairOrientedBound B;
-		ComputePCA(InPoints, 0, InPoints.Num(), B.Center, B.ExtentX, B.ExtentY, B.ExtentZ);
-		B.ExtentX *= Scale;
-		B.ExtentY *= Scale;
-		B.ExtentZ *= Scale;
-
-		// Split at least once the OBB to have two OBBs, as we nee to infer a tangent direction
-		if (Depth == 0)
-		{
-			const uint8 MajorAxis = GetMajorAxis(B);
-			const float Size = FMath::Max(B.ExtentX.Size(), FMath::Max(B.ExtentY.Size(), B.ExtentZ.Size()));
-			MaxDepth = FMath::Max(1, FMath::CeilToInt(Size / MinSectionLength));
-		}
-
-		if (Depth < MaxDepth)
-		{
-			const uint8 MajorAxis = GetMajorAxis(B);
-			FHairOrientedBound B0;
-			FHairOrientedBound B1;
-			Split(MajorAxis, B, B0, B1);
-			OutAll.Add(B0);
-			OutAll.Add(B1);
-
-			const TArray<FVector3f> InPoints0 = Inside(InPoints, B0);
-			const TArray<FVector3f> InPoints1 = Inside(InPoints, B1);
-			ComputeOrientedBoundingBoxes(InPoints0, MinSectionLength, Out, OutAll, Depth + 1, MaxDepth);
-			ComputeOrientedBoundingBoxes(InPoints1, MinSectionLength, Out, OutAll, Depth + 1, MaxDepth);
-		}
-		else
-		{
-			Out.Add(B);
-		}
-	}
-
-	static void FindTBN(FHairOrientedBound& B0, FHairOrientedBound& B1, FVector3f& T, FVector3f& B, FVector3f& N)
-	{
-		const FVector3f& C0 = B0.Center;
-		const FVector3f& C1 = B1.Center;
-		T = FVector3f(C1 - C0).GetSafeNormal();
-
-		const float DX = FVector3f::DotProduct(T, B0.ExtentX);
-		const float DY = FVector3f::DotProduct(T, B0.ExtentY);
-		const float DZ = FVector3f::DotProduct(T, B0.ExtentZ);
-
-		uint8 T_Axis = 0;
-		float MaxD = DX;
-		if (DY > MaxD) { MaxD = DY; T_Axis = 1; }
-		if (DZ > MaxD) { MaxD = DZ; T_Axis = 2; }
-
-		uint8 B_Axis = T_Axis == 0 ? 1 : (T_Axis == 1 ? 2 : 0);
-		uint8 N_Axis = T_Axis == 0 ? 2 : (T_Axis == 1 ? 0 : 1);
-
-		if (B0.Extent(B_Axis).Size() > B0.Extent(N_Axis).Size())
-		{
-			Swap(N_Axis, B_Axis);
-		}
-
-		// Sanity check
-		check(T_Axis != B_Axis);
-		check(T_Axis != N_Axis);
-		check(B_Axis != N_Axis);
-
-		N = B0.Extent(N_Axis);
-		B = B0.Extent(B_Axis);
-	}
-
-	static void DecimateCurve(const float InAngularThresholdInDegree, const float OutTargetCount, TArray<FVector3f>& OutPoints)
-	{
-		check(OutPoints.Num() > 2);
-		check(OutTargetCount >= 2 && OutTargetCount <= OutPoints.Num());
-			
-		const float AngularThresholdInRad = FMath::DegreesToRadians(InAngularThresholdInDegree);
-
-		// 'bCanDecimate' tracks if it is possible to reduce the remaining vertives even more while respecting the user angular constrain
-		bool bCanDecimate = true;
-		while (OutPoints.Num() > OutTargetCount && bCanDecimate)
-		{
-			float MinError = FLT_MAX;
-			int32 ElementToRemove = -1;
-			const uint32 Count = OutPoints.Num();
-			for (uint32 IndexIt = 1; IndexIt < Count - 1; ++IndexIt)
-			{
-				const FVector3f& P0 = OutPoints[IndexIt - 1];
-				const FVector3f& P1 = OutPoints[IndexIt];
-				const FVector3f& P2 = OutPoints[IndexIt + 1];
-
-				const float Area = FVector3f::CrossProduct(P0 - P1, P2 - P1).Size() * 0.5f;
-
-				//     P0 .       . P2
-				//         \Inner/
-				//   ` .    \   /
-				// Thres(` . \^/ ) Angle
-				//    --------.---------
-				//            P1
-				const FVector3f V0 = (P0 - P1).GetSafeNormal();
-				const FVector3f V1 = (P2 - P1).GetSafeNormal();
-				const float InnerAngle = FMath::Abs(FMath::Acos(FVector3f::DotProduct(V0, V1)));
-				const float Angle = (PI - InnerAngle) * 0.5f;
-
-				if (Area < MinError && Angle < AngularThresholdInRad)
-				{
-					MinError = Area;
-					ElementToRemove = IndexIt;
-				}
-			}
-			bCanDecimate = ElementToRemove >= 0;
-			if (bCanDecimate)
-			{
-				OutPoints.RemoveAt(ElementToRemove);
-			}
-		}
-	}
-
-	static void CreateHairStrandsCluster(
-		const FHairStrandsDatas& In, 
-		const FHairStrandsDatas& InSim, 
-		const FVoxelVolume& InVoxels,
-		FHairStrandsClusters& Out, 
-		const FHairGroupsProceduralCards& InSettings)
-	{
-		// Generate cluster only based on proximity for now. This will be iterativaly improved overtime
-		// Reuse the code from the interpolation code. 
-		// TODO: change this for https://www.lvdiwang.com/publications/hairsynthesis/2009_hairsynthesis.pdf
-		BuildClusterData(In, InSim, Out, InSettings.GeometrySettings);
-
-		// For each cluster 
-		// compute the 'guide'/'center' strands in the middle of the strands
-		// compute the extent and principal axis at each guides vertex
-		// Simplify/collapse based on angular threshold, but insure overally uniform distribution at least
-		uint32 Offset = 0;
-		const uint32 ClusterCount = Out.GetNum();
-		for (uint32 ClusterIt = 0; ClusterIt < ClusterCount; ++ClusterIt)
-		{
-			const uint32 CurveOffset = Out.CurveOffset[ClusterIt];
-			const uint32 CurveCount  = Out.CurveCount[ClusterIt];
-			
-			float MaxLength = 0;
-
-			// Select the longest strands
-			// Compute initial size based on the root positions
-			float RootRadius = 0;
-			FVector3f RootCenter = FVector3f::ZeroVector;
-			uint32 LongestCurveIndex = 0;
-			float LongestLength = 0;
-			FVector2f RootUV(0, 0);
-			TArray<FVector3f> Points;
-			Points.Reserve(CurveCount * 20); // Guess number for avoiding too much reallocation
-			for (uint32 StrandIt = 0; StrandIt < CurveCount; ++StrandIt)
-			{
-				const uint32 StrandIndex = Out.CurveIndices[CurveOffset + StrandIt];
-				const uint32 StrandOffset = In.StrandsCurves.CurvesOffset[StrandIndex];
-				const uint32 PointCount = In.StrandsCurves.CurvesCount[StrandIndex];
-				const float CurveLength = In.StrandsCurves.CurvesLength[StrandIndex];
-
-				if (CurveLength > LongestLength)
-				{
-					LongestLength = CurveLength;
-					LongestCurveIndex = StrandIndex;
-				}
-
-				for (uint32 PointIt = 0; PointIt < PointCount; ++PointIt)
-				{
-					FVector3f P0 = In.StrandsPoints.PointsPosition[StrandOffset + PointIt];
-					Points.Add(P0);
-				}
-
-				RootCenter += In.StrandsPoints.PointsPosition[StrandOffset];
-			}
-			RootCenter /= CurveCount;
-			for (uint32 StrandIt = 0; StrandIt < CurveCount; ++StrandIt)
-			{
-				const uint32 StrandIndex = Out.CurveIndices[CurveOffset + StrandIt];
-				const uint32 StrandOffset = In.StrandsCurves.CurvesOffset[StrandIndex];
-				RootRadius = FMath::Max((In.StrandsPoints.PointsPosition[StrandOffset] - RootCenter).Size(), RootRadius);
-			}
-
-			// Compute global OBB to compute the main orientation
-			Out.Bounds[ClusterIt] = ComputeOrientedBoundingBox(Points);
-
-			uint32 VertexCount = 0;
-			{
-				const uint32 GuidePointCount  = Out.GuidePointCount[ClusterIt];
-				const uint32 GuidePointOffset = Out.GuidePointOffset[ClusterIt];
-				const FVector2f GuideRootUV = Out.GuideRootUVs[ClusterIt];
-				TArray<FVector3f> BoundPoints;
-				TArray<float> BoundLengths;
-				float TotalLength = 0;
-				for (uint32 VertexIt = 0; VertexIt < GuidePointCount; ++VertexIt)
-				{
-					const FVector3f& P0 = Out.GuidePoints[GuidePointOffset + VertexIt];
-					BoundPoints.Add(P0);
-					BoundLengths.Add(TotalLength);
-
-					if (VertexIt > 0)
-					{
-						const FVector3f& PPrev = Out.GuidePoints[GuidePointOffset + VertexIt - 1];
-						TotalLength += (P0 - PPrev).Size();
-					}
-				}
-				MaxLength = TotalLength;
-
-				const uint32 OutCount = FMath::CeilToInt(FMath::Clamp(TotalLength / InSettings.GeometrySettings.MinSegmentLength, 2.f, float(BoundPoints.Num())));
-				if (OutCount < GuidePointCount)
-				{
-					DecimateCurve(InSettings.GeometrySettings.AngularThreshold, OutCount, BoundPoints);
-				}
-				VertexCount = BoundPoints.Num();
-
-				const float UseCurveOrientation = 1;
-				const float OrientationU = FMath::Clamp(UseCurveOrientation, 0.f, 1.0f);
-
-				FVector3f PrevN;
-				for (uint32 VertexIt = 0; VertexIt < VertexCount; ++VertexIt)
-				{
-					const FVector3f& P0 = BoundPoints[VertexIt];
-
-					FVector3f T = FVector3f::ZeroVector;
-					FVector3f B = FVector3f::ZeroVector;
-					FVector3f N = FVector3f::ZeroVector;
-					GetFrame(P0, InVoxels, T, B, N);
-
-					if (VertexIt == 0)
-					{
-						// Bias the initial normal by assuming the groom is centered aroung a head (spherical hypothesis)
-						const FVector3f StartN = (FVector3f)((FVector)P0 - In.BoundingBox.GetCenter()).GetSafeNormal();
-						//PrevN = N;
-						PrevN = StartN;
-					}
-
-					const float LengthU = BoundLengths[VertexIt] / TotalLength;
-					PrevN = FMath::Lerp(N, PrevN, FMath::Lerp(LengthU, 1.f, OrientationU));
-
-					FVector3f PT;
-					if (VertexIt > 0 && VertexIt < VertexCount - 1)
-					{
-						const FVector3f& PPrev = BoundPoints[VertexIt - 1];
-						const FVector3f& PNext = BoundPoints[VertexIt + 1];
-
-						PT = ((P0 - PPrev).GetSafeNormal() + (PNext - P0).GetSafeNormal()).GetSafeNormal();
-					}
-					else if (VertexIt == 0)
-					{
-						const FVector3f& PNext = BoundPoints[VertexIt + 1];
-						PT = (PNext - P0).GetSafeNormal();
-					}
-					else if (VertexIt == VertexCount - 1)
-					{
-						const FVector3f& PPrev = BoundPoints[VertexIt - 1];
-						PT = (P0 - PPrev).GetSafeNormal();
-					}
-
-					if (PT.Size() > 0)
-					{
-						T = PT;
-						B = FVector3f::CrossProduct(PrevN, T).GetSafeNormal();
-						N = FVector3f::CrossProduct(T, B).GetSafeNormal();
-					}
-
-					Out.BoundAxis_T.Add(T);
-					Out.BoundAxis_B.Add(B * FMath::Clamp(GHairCardsWidthScale, 0.01f, 10.f));
-					Out.BoundAxis_N.Add(N);
-					Out.BoundPositions.Add(P0);
-
-					PrevN = N;
-				}
-			}
-
-			Out.BoundOffset[ClusterIt] = Offset;
-			Out.BoundCount[ClusterIt]  = VertexCount;
-			Out.BoundLength[ClusterIt] = MaxLength;
-
-			Offset += VertexCount;
-		}
-	}
-}
-
-namespace HairCards
-{
-	// Note:
-	// A card can only be affected by one cluster
-	// A cluster can belongs to differnt cards
-	// A strand can belong to several cluster
-
-	// This is done per hair group/
-	// Operation could be break down into cached to avoid costly operation?
-
-	static void CreateCardsGeometry(
-		const FHairStrandsDatas& InDatas, 
-		const FHairStrandsClusters& InClusters, 
-		const FIntPoint& AtlasResolution, 
-		const FHairGroupsProceduralCards& InSettings, 
-		const HairCards::FVoxelVolume& InVoxel, 
-		FHairCardsProceduralGeometry& OutCards)
-	{
-		const uint32 ClusterCount = InClusters.GetNum();
-		const uint32 CardsPerCluster = 1;
-		OutCards.SetNum(ClusterCount * CardsPerCluster);
-
-		for (uint32 ClusterIt = 0; ClusterIt < ClusterCount; ++ClusterIt)
-		{
-			OutCards.Bounds[ClusterIt] = InClusters.Bounds[ClusterIt];
-		}
-
-		// Store the offset and the count of the strands vertices belonging to a given cluster
-		{
-			OutCards.CardIndexToClusterOffsetAndCount.SetNum(ClusterCount * CardsPerCluster);
-			OutCards.ClusterIndexToVertexOffsetAndCount.SetNum(InClusters.CurveIndices.Num());
-			uint32 ClusterStrandIt = 0;
-			for (uint32 ClusterIt = 0; ClusterIt < ClusterCount; ++ClusterIt)
-			{
-				const uint32 CurveOffset = InClusters.CurveOffset[ClusterIt];
-				const uint32 CurveCount = InClusters.CurveCount[ClusterIt];
-
-				for (uint32 CurveIt = 0; CurveIt < CurveCount; ++CurveIt)
-				{
-					const uint32 CurveIndex  = InClusters.CurveIndices[CurveOffset + CurveIt];					
-					const uint32 PointCount  = InDatas.StrandsCurves.CurvesCount[CurveIndex];
-					const uint32 PointOffset = InDatas.StrandsCurves.CurvesOffset[CurveIndex];
-
-					OutCards.ClusterIndexToVertexOffsetAndCount[ClusterStrandIt] = { PointOffset, PointCount };
-					++ClusterStrandIt;
-				}
-
-				for (uint32 CardIt = 0; CardIt < CardsPerCluster; ++CardIt)
-				{
-					const uint32 CardIndex = CardIt + ClusterIt * CardsPerCluster;
-					OutCards.CardIndexToClusterOffsetAndCount[CardIndex] = { CurveOffset, CurveCount };
-				}
-			}
-//			check(ClusterStrandIt == InClusters.CurveIndices.Num());
-		}
-
-		// Compute the total number of point and 
-		uint32 TotalPointCount = 0;
-		uint32 TotalIndexCount = 0;
-		for (uint32 ClusterIt = 0; ClusterIt < ClusterCount; ++ClusterIt)
-		{
-			for (uint32 CardIt = 0; CardIt < CardsPerCluster; ++CardIt)
-			{
-				const uint32 PointCount =  InClusters.BoundCount[ClusterIt] * 2;
-				const uint32 IndexCount	= (InClusters.BoundCount[ClusterIt]-1) * 6;
-
-				OutCards.IndexCounts[ClusterIt]  = IndexCount;
-				OutCards.IndexOffsets[ClusterIt] = TotalIndexCount;
-
-				TotalPointCount += PointCount;
-				TotalIndexCount += IndexCount;
-			}
-		}
-
-		// Based on principle axis and cluster strands, generate the position, indices and normal of the cards
-		OutCards.Positions.SetNum(TotalPointCount);
-		OutCards.Normals.SetNum(TotalPointCount);
-		OutCards.Tangents.SetNum(TotalPointCount);
-		OutCards.UVs.SetNum(TotalPointCount);
-		OutCards.Indices.SetNum(TotalIndexCount);
-		OutCards.CardIndices.SetNum(TotalPointCount);		
-		OutCards.CoordU.SetNum(TotalPointCount);
-
-		uint32 IndexIt = 0;
-		uint32 VertexIt = 0;
-		for (uint32 ClusterIt = 0; ClusterIt < ClusterCount; ++ClusterIt)
-		{
-			const float TotalLength = InClusters.BoundLength[ClusterIt];
-			for (uint32 CardIt = 0; CardIt < CardsPerCluster; ++CardIt)
-			{
-				const uint32 CardIndex = CardIt + ClusterIt * CardsPerCluster;
-
-				// Sanity check
-				check(OutCards.IndexOffsets[CardIndex] == IndexIt);
-				check(OutCards.IndexCounts[CardIndex] == (InClusters.BoundCount[ClusterIt]-1) * 6);
-
-				const uint32 PointCount = InClusters.BoundCount[ClusterIt];
-				const uint32 PointOffset = InClusters.BoundOffset[ClusterIt];
-
-				float CardMaxWidth = 0;
-				float CardLength = 0;
-				FVector3f PrevP0(0, 0, 0);
-				FVector3f PrevT(0, 0, 0);
-				const uint32 StartVertexIt = VertexIt;
-				const uint32 StartIndexIt = IndexIt;
-				for (uint32 PointIt = 0; PointIt < PointCount; ++PointIt)
-				{
-					// Simple geometric normal based on adjacent vertices
-					const FVector3f P0 = InClusters.BoundPositions[PointOffset + PointIt];
-					FVector3f T;
-					float SegmentLength = 0;
-					if (PointIt < PointCount - 1)
-					{
-						const FVector3f P1 = InClusters.BoundPositions[PointOffset + PointIt + 1];
-						SegmentLength = (P1 - P0).Size();
-						T = (P1 - P0).GetSafeNormal();
-					}
-					else
-					{
-						T = PrevT;
-					}
-
-					const FVector3f Cluster_N = InClusters.BoundAxis_N[PointOffset + PointIt];
-					const FVector3f B = InClusters.BoundAxis_B[PointOffset + PointIt].GetSafeNormal();
-					const float CardWidth = InClusters.BoundAxis_B[PointOffset + PointIt].Size();
-					CardMaxWidth = FMath::Max(CardWidth, CardMaxWidth);
-
-					OutCards.Positions[VertexIt    ] = P0 + B * CardWidth;
-					OutCards.Positions[VertexIt + 1] = P0 - B * CardWidth;
-					check(FMath::IsFinite(OutCards.Positions[VertexIt].X)     && FMath::IsFinite(OutCards.Positions[VertexIt].Y)     && FMath::IsFinite(OutCards.Positions[VertexIt].Z));
-					check(FMath::IsFinite(OutCards.Positions[VertexIt + 1].X) && FMath::IsFinite(OutCards.Positions[VertexIt + 1].Y) && FMath::IsFinite(OutCards.Positions[VertexIt + 1].Z));
-
-					OutCards.Normals[VertexIt    ] = FVector3f::ZeroVector; //MinorAxis0;
-					OutCards.Normals[VertexIt + 1] = FVector3f::ZeroVector; //MinorAxis0;
-
-					// Compute smooth tangent from the previous & next vertices
-					if (PointIt > 0)
-					{
-						const FVector3f T00 = OutCards.Positions[VertexIt]     - OutCards.Positions[VertexIt - 2];
-						const FVector3f T01 = OutCards.Positions[VertexIt + 1] - OutCards.Positions[VertexIt - 2 + 1];
-
-						OutCards.Tangents[VertexIt-2    ] = T00.GetSafeNormal();
-						OutCards.Tangents[VertexIt-2 + 1] = T01.GetSafeNormal();
-
-						// For the last vertex we use the previous vertex tangent
-						if (PointIt <= PointCount - 1)
-						{
-							OutCards.Tangents[VertexIt    ] = OutCards.Tangents[VertexIt-2    ];
-							OutCards.Tangents[VertexIt + 1] = OutCards.Tangents[VertexIt-2 + 1];
-						}
-					}
-
-					const FVector2f& RootUV = InClusters.GuideRootUVs[ClusterIt];
-					OutCards.UVs[VertexIt    ] = FVector4f(CardLength, 0, RootUV.X, RootUV.Y);
-					OutCards.UVs[VertexIt + 1] = FVector4f(CardLength, 1, RootUV.X, RootUV.Y);
-
-					OutCards.CardIndices[VertexIt]	   = CardIndex;
-					OutCards.CardIndices[VertexIt + 1] = CardIndex;
-
-
-					if (PointIt < PointCount - 1)
-					{
-						OutCards.Indices[IndexIt+0] = VertexIt;
-						OutCards.Indices[IndexIt+1] = VertexIt + 3;
-						OutCards.Indices[IndexIt+2] = VertexIt + 1;
-
-						OutCards.Indices[IndexIt+3] = VertexIt;
-						OutCards.Indices[IndexIt+4] = VertexIt + 2;
-						OutCards.Indices[IndexIt+5] = VertexIt + 3;
-
-						IndexIt += 6;
-					}
-
-					VertexIt += 2;
-					PrevP0 = P0;
-					PrevT = T;
-					CardLength += SegmentLength;
-				}
-
-				// Compute face normal. Use adjacent triangle angle to weighting the face contribution
-				for (uint32 InnerIndexIt = StartIndexIt; InnerIndexIt < IndexIt; InnerIndexIt += 3)
-				{
-					const uint32 I0 = OutCards.Indices[InnerIndexIt + 0];
-					const uint32 I1 = OutCards.Indices[InnerIndexIt + 1];
-					const uint32 I2 = OutCards.Indices[InnerIndexIt + 2];
-
-					const FVector3f P0 = OutCards.Positions[I0];
-					const FVector3f P1 = OutCards.Positions[I1];
-					const FVector3f P2 = OutCards.Positions[I2];
-
-					const FVector3f E0 = (P1 - P0).GetSafeNormal();
-					const FVector3f E1 = (P2 - P1).GetSafeNormal();
-					const FVector3f E2 = (P0 - P2).GetSafeNormal();
-
-					const FVector3f N = -FVector3f::CrossProduct(E0, -E2).GetSafeNormal();
-
-					OutCards.Normals[I0] += N * FMath::Acos(FVector3f::DotProduct(E0, -E0));
-					OutCards.Normals[I1] += N * FMath::Acos(FVector3f::DotProduct(E1, -E1));
-					OutCards.Normals[I2] += N * FMath::Acos(FVector3f::DotProduct(E2, -E2));
-				}
-
-				const uint32 RepresentativeRectIndex = InClusters.AtlasRectIndex[ClusterIt];
-				const FHairCardsProceduralGeometry::Rect Rect = InClusters.AtlasUniqueRect[RepresentativeRectIndex];
-				OutCards.Rects[CardIndex] = Rect;
-
-				// Average vertex normal
-				// Normalize UV
-				for (uint32 PointIt = 0; PointIt < PointCount; ++PointIt)
-				{
-					const uint32 VertexPerPoint = 2;
-
-					for (uint32 InnerVertexIt = 0; InnerVertexIt < VertexPerPoint; ++InnerVertexIt)
-					{
-						const uint32 VIndex = StartVertexIt + PointIt*2 + InnerVertexIt;
-
-						const FVector3f N = OutCards.Normals[VIndex].GetSafeNormal();
-						FVector3f T		= OutCards.Tangents[VIndex];
-						const FVector3f B = FVector3f::CrossProduct(N, T);
-						T = FVector3f::CrossProduct(B, N).GetSafeNormal();
-
-						OutCards.Normals [VIndex] = N;
-						OutCards.Tangents[VIndex] = T;
-
-						// Relative UVs
-						OutCards.UVs[VIndex].X /= CardLength;
-						OutCards.CoordU[VIndex] = OutCards.UVs[VIndex].X;
-						
-						// Absolute/Global UVs
-						const FVector2f AtlasCoord = FVector2f(OutCards.UVs[VIndex].X, OutCards.UVs[VIndex].Y) * Rect.Resolution + FVector2f(Rect.Offset.X, Rect.Offset.Y) + FVector2f(0.5f, 0.5f);
-						OutCards.UVs[VIndex].X = AtlasCoord.X / AtlasResolution.X;
-						OutCards.UVs[VIndex].Y = AtlasCoord.Y / AtlasResolution.Y;
-					}
-				}
-
-				//const float CentimeterToMeter = 10.f;
-				//OutCards.Rects[CardIndex].Offset.X = 0;
-				//OutCards.Rects[CardIndex].Offset.X = 0;
-				//OutCards.Rects[CardIndex].Offset.Y = 0;
-				//OutCards.Rects[CardIndex].Resolution.X = 0;//CardLength   * InSettings.TextureSettings.PixelPerCentimeters;
-				//OutCards.Rects[CardIndex].Resolution.Y = 0;//CardMaxWidth * InSettings.TextureSettings.PixelPerCentimeters * FMath::Max(0.1f, GHairCardsAtlasWidthScale);
-
-				OutCards.Lengths[CardIndex] = CardLength;
-			}
-		}
-
-		// Sanity check 
-		check(IndexIt == OutCards.Indices.Num());
-		check(VertexIt == OutCards.Positions.Num());
-	}
-
-	// Select a few representative cluster for generating texture from, and assign it to similar cluster. 
-	// Use length & density similarity to do the pairing.
-	static void CreateCardsTexture(
-		const FHairStrandsDatas& InDatas,
-		FHairStrandsClusters& InClusters,
-		const FHairGroupsProceduralCards& InSettings)
-	{
-		struct FClusterTextureDescriptor
-		{
-			uint32 ClusterIndex;
-			uint32 CurveCount;
-			float MinLength;
-			float MaxLength;
-		};
-
-		const uint32 DescriptorCount = InClusters.GetNum();
-		TArray<FClusterTextureDescriptor> Descriptors;
-		Descriptors.Reserve(DescriptorCount);
-
-		// Extra the min/max length & Build cluster texture descriptors
-		float GlobalMinLength = FLT_MAX;
-		float GlobalMaxLength = -FLT_MAX;
-		for (uint32 ClusterIt = 0, ClusterCount = InClusters.GetNum(); ClusterIt < ClusterCount; ++ClusterIt)
-		{
-			FClusterTextureDescriptor& Desc = Descriptors.AddDefaulted_GetRef();
-			Desc.ClusterIndex = ClusterIt;
-			Desc.CurveCount = InClusters.CurveCount[ClusterIt];
-			Desc.MinLength = FLT_MAX;
-			Desc.MaxLength = -FLT_MAX;
-
-			const uint32 CurveOffset = InClusters.CurveOffset[ClusterIt];
-			const uint32 CurveCount = InClusters.CurveCount[ClusterIt];
-			check(CurveCount > 0);
-			for (uint32 StrandIt = 0; StrandIt < CurveCount; ++StrandIt)
-			{
-				const uint32 StrandIndex = InClusters.CurveIndices[CurveOffset + StrandIt];
-				const uint32 StrandOffset = InDatas.StrandsCurves.CurvesOffset[StrandIndex];
-				const uint32 PointCount = InDatas.StrandsCurves.CurvesCount[StrandIndex];
-				const float CurveLength = InDatas.StrandsCurves.CurvesLength[StrandIndex];
-
-				Desc.MinLength = FMath::Min(CurveLength, Desc.MinLength);
-				Desc.MaxLength = FMath::Max(CurveLength, Desc.MaxLength);
-			}
-
-			GlobalMinLength = FMath::Min(GlobalMinLength, Desc.MinLength);
-			GlobalMaxLength = FMath::Max(GlobalMaxLength, Desc.MaxLength);
-
-			InClusters.AtlasRectIndex[ClusterIt] = ~0u;
-		}
-
-		// Sort by length first (from longest to shortest)
-		Descriptors.Sort(
-			[](const FClusterTextureDescriptor& A, const FClusterTextureDescriptor& B)
-			{
-				return A.MaxLength > B.MaxLength;
-			});
-
-		const float DeltaLengthPerBucket = (GlobalMaxLength - GlobalMinLength) / float(InSettings.TextureSettings.LengthTextureCount);
-
-		InClusters.AtlasUniqueRectIndex.Empty();
-		InClusters.AtlasUniqueRect.Empty();
-
-		uint32 UniqueRectIndex = 0;
-		uint32 DescriptorIt = 0;
-		float NextBucketLength = GlobalMaxLength - DeltaLengthPerBucket;
-		while (DescriptorIt < DescriptorCount)
-		{
-			uint32 StartDescriptorIt = DescriptorIt;
-			uint32 EndDescriptorIt = DescriptorIt;
-			while (EndDescriptorIt < DescriptorCount && Descriptors[EndDescriptorIt].MaxLength >= NextBucketLength)
-			{
-				++EndDescriptorIt;
-				++DescriptorIt;
-			}
-			DescriptorIt = EndDescriptorIt;
-
-			// Sort by curve count (from denser to less denser)
-			const uint32 BucketDescriptorCount = EndDescriptorIt - StartDescriptorIt;
-			if (BucketDescriptorCount > 1)
-			{
-				Algo::SortBy(MakeArrayView(Descriptors.GetData() + StartDescriptorIt, BucketDescriptorCount),
-					&FClusterTextureDescriptor::CurveCount, TGreater<>());
-			}
-
-			// Take the first descritor (Longer & denser as the representative of the bucket)
-			const uint32 RepresentativeClusterIndex = Descriptors[StartDescriptorIt].ClusterIndex;
-
-			// The actual dimension of the rect are setup later (based on the straighten curves dimensions)
-			FHairCardsProceduralGeometry::Rect Rect;
-			Rect.Offset.X = 0;
-			Rect.Offset.Y = 0;
-			Rect.Resolution.X = 0; //InSettings.TextureSettings.PixelPerCentimeters * Descriptors[StartDescriptorIt].MaxLength;
-			Rect.Resolution.Y = 0; //InSettings.TextureSettings.PixelPerCentimeters * 20u * FMath::Max(1.f, GHairCardsAtlasWidthScale) ; //TODO
-
-			InClusters.AtlasUniqueRect.Add(Rect);
-			InClusters.AtlasUniqueRectIndex.Add(RepresentativeClusterIndex);
-
-			for (uint32 DescIt = StartDescriptorIt; DescIt < EndDescriptorIt; ++DescIt)
-			{
-				const uint32 ClusterIt = Descriptors[DescIt].ClusterIndex;
-				InClusters.AtlasRectIndex[ClusterIt] = UniqueRectIndex;
-			}
-
-			NextBucketLength -= DeltaLengthPerBucket;
-			++UniqueRectIndex;
-		}
-	}
-
-	struct FAtlasPackingRect
-	{
-		FHairCardsProceduralGeometry::Rect Rect;
-		uint32 UniqueRectIt = 0;
-	};
-
-	static void Packing_InitAtlasRects(
-		const TArray<FHairCardsProceduralAtlas::Rect>& Rects,
-		TArray<FAtlasPackingRect>& OutRects)
-	{
-		const uint32 RectCount = Rects.Num();
-		OutRects.SetNum(RectCount);
-		for (uint32 RectIt = 0; RectIt < RectCount; ++RectIt)
-		{
-			OutRects[RectIt].Rect.Offset	 = FIntPoint(0, 0);
-			OutRects[RectIt].Rect.Resolution = Rects[RectIt].Resolution; //InClusters.AtlasUniqueRect[RectIt].Resolution;
-			OutRects[RectIt].UniqueRectIt	 = RectIt;
-		}
-	}
-
-	static void Packing_RescaleAtlasRects(TArray<FHairCardsProceduralAtlas::Rect>& OutRects, const float ScaleFactor)
-	{
-		for (FHairCardsProceduralAtlas::Rect& Rect : OutRects)
-		{
-			Rect.Resolution.X = FMath::Max(FMath::CeilToInt(Rect.Resolution.X * ScaleFactor), 2);
-			Rect.Resolution.Y = FMath::Max(FMath::CeilToInt(Rect.Resolution.Y * ScaleFactor), 2);
-		}
-	}
-
-	static void Packing_ComputeAtlasMaxDimensionAndArea(
-		const TArray<FHairCardsProceduralAtlas::Rect>& Rects,
-		const float ScaleFactor,
-		float& OutArea,
-		int32& OutMaxHeight)
-	{
-		OutArea = 0;
-		OutMaxHeight = 0;
-		for (const FHairCardsProceduralAtlas::Rect& Rect : Rects)
-		{
-			FHairCardsProceduralAtlas::Rect ScaledRect = Rect;
-			ScaledRect.Resolution.X = FMath::Max(FMath::CeilToInt(ScaledRect.Resolution.X * ScaleFactor), 2);
-			ScaledRect.Resolution.Y = FMath::Max(FMath::CeilToInt(ScaledRect.Resolution.Y * ScaleFactor), 2);
-			OutMaxHeight = FMath::Max(OutMaxHeight, ScaledRect.Resolution.Y);
-			OutArea += ScaledRect.Resolution.X * ScaledRect.Resolution.Y;
-		}
-	}
-
-	static FIntPoint Packing_ComputeAtlasDefaultResolution(float Area, int32 MaxHeight)
-	{
-		int32 DefaultResolution = FMath::CeilToInt(FMath::Sqrt(Area));
-		FIntPoint OutResolution;
-		OutResolution.X = DefaultResolution;
-		OutResolution.Y = FMath::Max(DefaultResolution, MaxHeight);
-		OutResolution   = FIntPoint(FMath::RoundUpToPowerOfTwo(OutResolution.X), FMath::RoundUpToPowerOfTwo(OutResolution.Y));
-		return OutResolution;
-	}
-
-	void PackHairCardsRects(
-		const FHairCardsTextureSettings& InSettings,
-		FHairStrandsClusters& InClusters,
-		FHairCardsProceduralAtlas& InAtlas)
-	{
-		// 1. Initial resolution estimate
-		// Compute an estimate of the atlas resolution which satisfies the max resolution requirement
-		// If the initial resolution is larger than the requested max resolution, cards' rects 
-		// are rescaled until they fit
-		const int32 AtlasMaxResolution = FMath::RoundUpToPowerOfTwo(FMath::Min(InSettings.AtlasMaxResolution, 16384));
-		FIntPoint AtlasResolution(0, 0);
-		{			
-			float ScaleFactor = 1;
-			float Area = 0;
-			int32 MaxHeight = 0;
-			Packing_ComputeAtlasMaxDimensionAndArea(InAtlas.Rects, ScaleFactor, Area, MaxHeight);
-
-			// Insure that the largest element fit into the atlas. Otherwise compute the overall rescaling factor;
-			if (MaxHeight > AtlasMaxResolution)
-			{
-				ScaleFactor = AtlasMaxResolution / float(FMath::Max(MaxHeight, 1));
-				Packing_ComputeAtlasMaxDimensionAndArea(InAtlas.Rects, ScaleFactor, Area, MaxHeight);
-			}
-
-			// Insure that the initial resolution is smaller than the requested max resolution
-			// Start with a small 512 x 512 texture so that we can have tighter packing. 
-			// The packing algo grows the resolution until max. resolution is reached. If the 
-			// packing still does not work, then we start to shrink cards resolution.
-			AtlasResolution = FIntPoint(FMath::Min(AtlasMaxResolution, 512), FMath::Min(AtlasMaxResolution, 512));
-
-			// Rescale rects so that they can fit into the atlas
-			if (ScaleFactor < 1.f)
-			{
-				Packing_RescaleAtlasRects(InAtlas.Rects, ScaleFactor);
-			}
-		}
-
-		// 2. Packing
-		// Use binary rect insertion, starting with the rect having the largest width/height first. 
-		// The algo is relatively naive, but since we have a low number of rects this is acceptable
-		TArray<FAtlasPackingRect> FinalAtlasRects;
-		bool bFitInAtlas = false;
-		while (!bFitInAtlas)
-		{
-			// 2.1 Transfer atlas rect into pack rect
-			// Transfer the rect packing from the cluster to the sorting struct
-			TArray<FAtlasPackingRect> AtlasRects;
-			Packing_InitAtlasRects(InAtlas.Rects, AtlasRects);
-
-			// 2.2. Sort rect from larger to lower dimensions (width/height). 
-			// Since the rect cards are mostly anisotropic, it seems better to use the longest dimension 
-			// as a comparison metric, rather than perimeter
-			AtlasRects.Sort([](const FAtlasPackingRect& A, const FAtlasPackingRect& B)
-			{
-				uint32 LongestA = FMath::Max(A.Rect.Resolution.X, A.Rect.Resolution.Y);
-				uint32 LongestB = FMath::Max(B.Rect.Resolution.X, B.Rect.Resolution.Y);
-				return LongestA >= LongestB;
-			});
-
-			// 2.3 Reset packing
-			// Reset the packing to a single rect which covers the entire atlas
-			TArray<FHairCardsProceduralGeometry::Rect> FreeRects;
-			FreeRects.Reserve(AtlasRects.Num() * 2);
-			FreeRects.Add({ FIntPoint(0,0), FIntPoint(AtlasResolution.X, AtlasResolution.Y) });
-
-			// 2.4 Pack
-			bFitInAtlas = true;
-			for (FAtlasPackingRect& R : AtlasRects)
-			{
-				int32 FreeRectIndex = -1;
-				for (int32 FreeRIt=0, FreeRCount = FreeRects.Num(); FreeRIt<FreeRCount; ++FreeRIt)
-				{
-					FHairCardsProceduralGeometry::Rect& FreeR = FreeRects[FreeRIt];
-					if (R.Rect.Resolution.X <= FreeR.Resolution.X &&
-						R.Rect.Resolution.Y <= FreeR.Resolution.Y)
-					{
-						FreeRectIndex = FreeRIt;
-						break;
-					}
-				}
-
-				if (FreeRectIndex >=0)
-				{
-					FHairCardsProceduralGeometry::Rect& FreeR = FreeRects[FreeRectIndex];
-					// Perfect Fit
-					if (R.Rect.Resolution.X == FreeR.Resolution.X &&
-						R.Rect.Resolution.Y == FreeR.Resolution.Y)
-					{
-						R.Rect.Offset = FreeR.Offset;
-						FreeRects.RemoveAt(FreeRectIndex);
-					}
-					// Same width
-					else if (R.Rect.Resolution.X == FreeR.Resolution.X &&
-							 R.Rect.Resolution.Y <= FreeR.Resolution.Y)
-					{
-						R.Rect.Offset = FreeR.Offset;
-
-						FreeR.Offset.Y     += R.Rect.Resolution.Y;
-						FreeR.Resolution.Y -= R.Rect.Resolution.Y;
-					}
-					// Same height
-					else if (R.Rect.Resolution.X <= FreeR.Resolution.X &&
-							 R.Rect.Resolution.Y == FreeR.Resolution.Y)
-					{
-						R.Rect.Offset = FreeR.Offset;
-
-						FreeR.Offset.X     += R.Rect.Resolution.X;
-						FreeR.Resolution.X -= R.Rect.Resolution.X;
-					}
-					// Within
-					else
-					{
-						R.Rect.Offset = FreeR.Offset;
-
-						// Make the cut along the longest side
-						const bool bAlongX = R.Rect.Resolution.X > R.Rect.Resolution.Y;
-
-						// Technically the incoding rect takes on free rect, and create two new Free rects. 
-						// However we can instead update one, and create another one
-						FHairCardsProceduralGeometry::Rect FreeR2 = FreeR;
-						if (bAlongX)
-						{
-							FreeR2.Offset.Y    += R.Rect.Resolution.Y;
-							FreeR2.Resolution.Y-= R.Rect.Resolution.Y;
-
-							FreeR.Offset.X     += R.Rect.Resolution.X;
-							FreeR.Resolution.X -= R.Rect.Resolution.X;
-							FreeR.Resolution.Y  = R.Rect.Resolution.Y;
-						}
-						else
-						{
-							FreeR2.Offset.Y    += R.Rect.Resolution.Y;
-							FreeR2.Resolution.X = R.Rect.Resolution.X;
-							FreeR2.Resolution.Y-= R.Rect.Resolution.Y;
-
-							FreeR.Offset.X     += R.Rect.Resolution.X;
-							FreeR.Resolution.X -= R.Rect.Resolution.X;
-						}
-
-						FreeRects.Add(FreeR2);
-					}
-
-					// Sort the free rects, to order them from smallest to largest to have tigher packing
-					FreeRects.Sort([](const FHairCardsProceduralGeometry::Rect& A, const FHairCardsProceduralGeometry::Rect& B)
-					{
-						uint32 PerimeterA = (A.Resolution.X + A.Resolution.Y) * 2;
-						uint32 PerimeterB = (B.Resolution.X + B.Resolution.Y) * 2;
-						return PerimeterA < PerimeterB;
-					});
-				}
-				else
-				{
-					bFitInAtlas = false;
-					break;
-				}				
-			}
-			
-			// 2.5 Adjust atlas resolution or scale down rects if it does not fit
-			// * If the atlas size has already reached the max, size, then down scale the rect elements
-			// * Otherwise, increase the atlas size
-			if (!bFitInAtlas)
-			{
-				if (AtlasResolution.X == AtlasMaxResolution || AtlasResolution.Y == AtlasMaxResolution)
-				{
-					// Shrink by 10%
-					float ScaleFactor = 0.9f;
-					Packing_RescaleAtlasRects(InAtlas.Rects, ScaleFactor);
-				}
-				else
-				{
-					AtlasResolution = FIntPoint(FMath::RoundUpToPowerOfTwo(2 * AtlasResolution.X), FMath::RoundUpToPowerOfTwo(2 * AtlasResolution.Y));
-					AtlasResolution.X = FMath::Min(AtlasResolution.X, AtlasMaxResolution);
-					AtlasResolution.Y = FMath::Min(AtlasResolution.Y, AtlasMaxResolution);
-				}
-			}
-			else
-			{
-				FinalAtlasRects = AtlasRects;
-			}
-		}
-		InAtlas.Resolution = AtlasResolution;
-
-		// Transfer the packing information from the sorting struct back to the cluster 
-		for (FAtlasPackingRect& R : FinalAtlasRects)
-		{
-			const uint32 RectIt = R.UniqueRectIt;
-			InClusters.AtlasUniqueRect[RectIt] = R.Rect;
-			InAtlas.Rects[RectIt].Offset = R.Rect.Offset;
-
-			check(InAtlas.Rects[RectIt].Resolution == R.Rect.Resolution);
-		}
-	}
-
-
-	static void CreateCardsGuides(
-		const FHairStrandsClusters& InClusters,
-		FHairStrandsDatas& OutGuides)
-	{
-		const uint32 ClusterCount = InClusters.GetNum();
-		OutGuides.StrandsPoints.SetNum(InClusters.GuidePoints.Num(), 0u);
-		OutGuides.StrandsCurves.SetNum(InClusters.GetNum(), 0u);
-		OutGuides.BoundingBox.Init();
-
-		uint32 IndexIt = 0;
-		uint32 VertexIt = 0;
-		for (uint32 ClusterIt = 0; ClusterIt < ClusterCount; ++ClusterIt)
-		{
-			const uint32 PointCount = InClusters.BoundCount[ClusterIt];
-			const uint32 PointOffset = InClusters.BoundOffset[ClusterIt];
-			const float TotalLength = InClusters.BoundLength[ClusterIt];
-
-			OutGuides.StrandsCurves.CurvesCount[ClusterIt] = PointCount;
-			OutGuides.StrandsCurves.CurvesOffset[ClusterIt] = PointOffset;
-			OutGuides.StrandsCurves.CurvesLength[ClusterIt] = TotalLength;
-			//OutGuides.StrandsCurves.CurvesRootUV[ClusterIt] = FVector2f(0, 0);
-			//OutGuides.StrandsCurves.StrandIDs[ClusterIt] = ClusterIt;
-			//OutGuides.StrandsCurves.GroomIDToIndex[ClusterIt] = ;
-
-			float CurrentLength = 0;
-			for (uint32 PointIt = 0; PointIt < PointCount; ++PointIt)
-			{
-				const uint32 PointIndex = PointOffset + PointIt;
-				const FVector3f P0 = InClusters.BoundPositions[PointIndex];
-
-				OutGuides.BoundingBox += (FVector)P0;
-
-				OutGuides.StrandsPoints.PointsPosition[PointIndex] = P0;
-				//OutGuides.StrandsPoints.PointsBaseColor[PointIndex] = FLinearColor(FVector3f::ZeroVector);
-				//OutGuides.StrandsPoints.PointsRoughness[PointIndex] = 0;
-				OutGuides.StrandsPoints.PointsCoordU[PointIndex] = FMath::Clamp(CurrentLength / TotalLength, 0.f, 1.f);
-				OutGuides.StrandsPoints.PointsRadius[PointIndex] = 1;
-
-				// Simple geometric normal based on adjacent vertices
-				if (PointIt < PointCount - 1)
-				{
-					const FVector3f P1 = InClusters.BoundPositions[PointOffset + PointIt + 1];
-					const float SegmentLength = (P1 - P0).Size();
-					CurrentLength += SegmentLength;
-				}
-			}
-		}
-	}
-
-	static bool CreateCardsGuides(
-		FHairCardsGeometry& InCards,
-		FHairStrandsDatas& OutGuides,
-		TArray<float>& OutCardLengths
-	)
-	{
-		// Build the guides from the triangles that form the card
-		// The guides are derived from the line that passes through the middle of each quad
-
-		const uint32 NumCards = InCards.IndexOffsets.Num();
-		const uint32 MaxGuidePoints = InCards.Indices.Num() / 3;
-		OutCardLengths.SetNum(NumCards);
-
-		OutGuides.StrandsPoints.PointsPosition.Reserve(MaxGuidePoints);
-		OutGuides.StrandsPoints.PointsRadius.Reserve(MaxGuidePoints);
-		OutGuides.StrandsPoints.PointsCoordU.Reserve(MaxGuidePoints);
-		OutGuides.StrandsCurves.SetNum(NumCards, 0u);
-		OutGuides.BoundingBox.Init();
-
-		InCards.CoordU.SetNum(InCards.Positions.Num());
-		InCards.LocalUVs.SetNum(InCards.Positions.Num());
-
-
-		// Find the principal direction of the atlas by comparing the number of segment along U and along V
-		bool bIsMainDirectionU = true;
-		{
-			uint32 MainDirectionUCount = 0;
-			uint32 ValidCount = 0;
-			FVector4f TriangleUVs[3];
-			for (uint32 CardIt = 0; CardIt < NumCards; ++CardIt)
-			{
-				const uint32 NumTriangles = InCards.IndexCounts[CardIt] / 3;
-				const uint32 VertexOffset = InCards.IndexOffsets[CardIt];
-				struct FIndexAndCoord { uint32 Index; float TexCoord; };
-				struct FSimilarUVVertices
-				{
-					float TexCoord = 0;
-					TArray<uint32> Indices;
-					TArray<FIndexAndCoord> AllIndices; // Store vertex index and Tex.coord perpendicular to the principal axis (stored in TexCoord)
-				};
-				TArray<FSimilarUVVertices> SimilarVertexU;
-				TArray<FSimilarUVVertices> SimilarVertexV;
-
-				auto AddSimilarUV = [](TArray<FSimilarUVVertices>& In, uint32 Index, float TexCoord, float MinorTexCoord, float Threshold)
-				{
-					bool bFound = false;
-					for (int32 It = 0, Count = In.Num(); It < Count; ++It)
-					{
-						if (FMath::Abs(In[It].TexCoord - TexCoord) < Threshold)
-						{
-							// We add only unique vertices per segment, so that the average position land in the center of the cards
-							In[It].Indices.AddUnique(Index);
-							In[It].AllIndices.Add({ Index, MinorTexCoord });
-							bFound = true;
-							break;
-						}
-					}
-
-					if (!bFound)
-					{
-						FSimilarUVVertices& SimilarUV = In.AddDefaulted_GetRef();
-						SimilarUV.TexCoord = TexCoord;
-						SimilarUV.Indices.Add(Index);
-						SimilarUV.AllIndices.Add({ Index, MinorTexCoord });
-					}
-				};
-
-				// Iterate over all triangles of a cards, and find vertices which share either same U or same V. We add them to separate lists. 
-				// We then use the following heuristic: the main axis will have more segments. This is what determine the principal axis of the cards
-				const float UVCoordTreshold = 1.f / 1024.f; // 1 pixel for a 1k texture
-				for (uint32 TriangleIt = 0; TriangleIt < NumTriangles; ++TriangleIt)
-				{
-					const uint32 VertexIndexOffset = VertexOffset + TriangleIt * 3;
-					for (uint32 VertexIt = 0; VertexIt < 3; ++VertexIt)
-					{
-						const uint32 VertexIndex = InCards.Indices[VertexIndexOffset + VertexIt];
-						const FVector4f UV = TriangleUVs[VertexIt] = InCards.UVs[VertexIndex];
-						AddSimilarUV(SimilarVertexU, VertexIndex, UV.X, UV.Y, UVCoordTreshold);
-						AddSimilarUV(SimilarVertexV, VertexIndex, UV.Y, UV.X, UVCoordTreshold);
-					}
-				}
+bool InternalCreateCardsGuides(
+	FHairCardsGeometry& InCards,
+	FHairStrandsDatas& OutGuides,
+	TArray<float>& OutCardLengths)
+{
+	// Build the guides from the triangles that form the card
+	// The guides are derived from the line that passes through the middle of each quad
 
-				// Use global UV orientation
-				// Use global segment count orientation
+	const uint32 NumCards = InCards.IndexOffsets.Num();
+	const uint32 MaxGuidePoints = InCards.Indices.Num() / 3;
+	OutCardLengths.SetNum(NumCards);
 
-				// Find the perpendicular direction by comparing the number of segment along U and along V
-				const bool bIsValid = SimilarVertexU.Num() != SimilarVertexV.Num() ? 1u : 0u;
-				if (bIsValid)
-				{
-					++ValidCount;
-					MainDirectionUCount += SimilarVertexU.Num() > SimilarVertexV.Num() ? 1u : 0u;
-				}
-			}
+	OutGuides.StrandsPoints.PointsPosition.Reserve(MaxGuidePoints);
+	OutGuides.StrandsPoints.PointsRadius.Reserve(MaxGuidePoints);
+	OutGuides.StrandsPoints.PointsCoordU.Reserve(MaxGuidePoints);
+	OutGuides.StrandsCurves.SetNum(NumCards, 0u);
+	OutGuides.BoundingBox.Init();
 
-			bIsMainDirectionU = (float(MainDirectionUCount) / float(FMath::Max(1u,ValidCount))) > 0.5f;
-		}
+	InCards.CoordU.SetNum(InCards.Positions.Num());
+	InCards.LocalUVs.SetNum(InCards.Positions.Num());
 
+	// Find the principal direction of the atlas by comparing the number of segment along U and along V
+	bool bIsMainDirectionU = true;
+	{
+		uint32 MainDirectionUCount = 0;
+		uint32 ValidCount = 0;
 		FVector4f TriangleUVs[3];
 		for (uint32 CardIt = 0; CardIt < NumCards; ++CardIt)
 		{
@@ -2922,7 +119,7 @@ namespace HairCards
 					{
 						// We add only unique vertices per segment, so that the average position land in the center of the cards
 						In[It].Indices.AddUnique(Index);
-						In[It].AllIndices.Add({Index, MinorTexCoord});
+						In[It].AllIndices.Add({ Index, MinorTexCoord });
 						bFound = true;
 						break;
 					}
@@ -2933,7 +130,7 @@ namespace HairCards
 					FSimilarUVVertices& SimilarUV = In.AddDefaulted_GetRef();
 					SimilarUV.TexCoord = TexCoord;
 					SimilarUV.Indices.Add(Index);
-					SimilarUV.AllIndices.Add({Index, MinorTexCoord});
+					SimilarUV.AllIndices.Add({ Index, MinorTexCoord });
 				}
 			};
 
@@ -2952,684 +149,267 @@ namespace HairCards
 				}
 			}
 
+			// Use global UV orientation
+			// Use global segment count orientation
+
 			// Find the perpendicular direction by comparing the number of segment along U and along V
-			TArray<FVector3f> CenterPoints;
+			const bool bIsValid = SimilarVertexU.Num() != SimilarVertexV.Num() ? 1u : 0u;
+			if (bIsValid)
 			{
-				// Sort vertices along the main axis so that, when we iterate through them, we get a correct linear ordering
-				TArray<FSimilarUVVertices>& SimilarVertex = bIsMainDirectionU ? SimilarVertexU : SimilarVertexV;
-				SimilarVertex.Sort([](const FSimilarUVVertices& A, const FSimilarUVVertices& B)
-				{
-					return A.TexCoord < B.TexCoord;
-				});
-
-				// For each group of vertex with similar 'principal' tex coord, sort them with growing 'perpendicular/secondary' tex coord
-				for (FSimilarUVVertices& Similar : SimilarVertex)
-				{
-					Similar.AllIndices.Sort([](const FIndexAndCoord& A, const FIndexAndCoord& B)
-					{
-						return A.TexCoord < B.TexCoord;
-					});
-
-					// Compute normalize coordinate
-					float MinTexCoord = Similar.AllIndices[0].TexCoord;
-					float MaxTexCoord = Similar.AllIndices[0].TexCoord;
-					for (const FIndexAndCoord& A : Similar.AllIndices)
-					{
-						MinTexCoord = FMath::Min(A.TexCoord, MinTexCoord);
-						MaxTexCoord = FMath::Max(A.TexCoord, MaxTexCoord);
-					}
-					MaxTexCoord = FMath::Max(MaxTexCoord, MinTexCoord + KINDA_SMALL_NUMBER);
-					for (FIndexAndCoord& A : Similar.AllIndices)
-					{
-						A.TexCoord = (A.TexCoord - MinTexCoord) / (MaxTexCoord - MinTexCoord);
-					}
-				}
-
-				CenterPoints.Reserve(SimilarVertex.Num());
-				FVector3f PrevCenterPoint = FVector3f::ZeroVector;
-				float TotalLength = 0;
-				for (const FSimilarUVVertices& Similar : SimilarVertex)
-				{
-					// Compute avg center point of the guide
-					FVector3f CenterPoint = FVector3f::ZeroVector;
-					for (uint32 VertexIndex : Similar.Indices)
-					{
-						CenterPoint += InCards.Positions[VertexIndex];
-					}
-					CenterPoint /= Similar.Indices.Num() > 0 ? Similar.Indices.Num() : 1;
-
-					// Update length along the guide
-					if (CenterPoints.Num() > 0)
-					{
-						const float SegmentLength = (CenterPoint - PrevCenterPoint).Size();
-						TotalLength += SegmentLength;
-					}
+				++ValidCount;
+				MainDirectionUCount += SimilarVertexU.Num() > SimilarVertexV.Num() ? 1u : 0u;
+			}
+		}
 
-					// Update neighbor vertex with current length (to compute the parametric distance at the end)
-					for (const FIndexAndCoord& VertexIndex : Similar.AllIndices)
-					{
-						InCards.CoordU[VertexIndex.Index] = TotalLength;
-						InCards.LocalUVs[VertexIndex.Index].X = TotalLength;
-						InCards.LocalUVs[VertexIndex.Index].Y = VertexIndex.TexCoord;
-					}
+		bIsMainDirectionU = (float(MainDirectionUCount) / float(FMath::Max(1u,ValidCount))) > 0.5f;
+	}
 
-					CenterPoints.Add(CenterPoint);
-					PrevCenterPoint = CenterPoint;
-				}
+	FVector4f TriangleUVs[3];
+	for (uint32 CardIt = 0; CardIt < NumCards; ++CardIt)
+	{
+		const uint32 NumTriangles = InCards.IndexCounts[CardIt] / 3;
+		const uint32 VertexOffset = InCards.IndexOffsets[CardIt];
+		struct FIndexAndCoord { uint32 Index; float TexCoord; };
+		struct FSimilarUVVertices
+		{
+			float TexCoord = 0;
+			TArray<uint32> Indices;
+			TArray<FIndexAndCoord> AllIndices; // Store vertex index and Tex.coord perpendicular to the principal axis (stored in TexCoord)
+		};
+		TArray<FSimilarUVVertices> SimilarVertexU;
+		TArray<FSimilarUVVertices> SimilarVertexV;
 
-				if (SimilarVertex.Num() > 1)
+		auto AddSimilarUV = [](TArray<FSimilarUVVertices>& In, uint32 Index, float TexCoord, float MinorTexCoord, float Threshold)
+		{
+			bool bFound = false;
+			for (int32 It = 0, Count = In.Num(); It < Count; ++It)
+			{
+				if (FMath::Abs(In[It].TexCoord - TexCoord) < Threshold)
 				{
-					// Normalize length to have a parametric distance
-					for (const FSimilarUVVertices& Similar : SimilarVertex)
-					{
-						for (uint32 VertexIndex : Similar.Indices)
-						{
-							InCards.CoordU[VertexIndex] /= TotalLength;
-							InCards.LocalUVs[VertexIndex].X = InCards.CoordU[VertexIndex];
-						}
-					}
+					// We add only unique vertices per segment, so that the average position land in the center of the cards
+					In[It].Indices.AddUnique(Index);
+					In[It].AllIndices.Add({Index, MinorTexCoord});
+					bFound = true;
+					break;
 				}
 			}
 
-			// Insure that cards as at least two points to build a segment as a lot of the runtime code assume at have
-			// at least one valid segment
-			check(CenterPoints.Num() > 0);
-			if (CenterPoints.Num() == 1)
+			if (!bFound)
 			{
-				const FVector3f CenterPoint = CenterPoints[0];
-				const float SegmentSize = 0.5f;
-				const FVector3f P1 = CenterPoint + SegmentSize * (CenterPoint - InCards.BoundingBox.GetCenter()).GetSafeNormal();
-				CenterPoints.Add(P1);
+				FSimilarUVVertices& SimilarUV = In.AddDefaulted_GetRef();
+				SimilarUV.TexCoord = TexCoord;
+				SimilarUV.Indices.Add(Index);
+				SimilarUV.AllIndices.Add({Index, MinorTexCoord});
 			}
+		};
 
-			// Compute and store the guide's total length
-			const uint32 PointCount = CenterPoints.Num();
-			float TotalLength = 0.f;
-			for (uint32 PointIt = 0; PointIt < PointCount - 1; ++PointIt)
-			{
-				TotalLength += (CenterPoints[PointIt + 1] - CenterPoints[PointIt]).Size();
-			}
-			OutCardLengths[CardIt] = TotalLength;
-
-			const uint32 PointOffset = OutGuides.StrandsPoints.PointsPosition.Num();
-
-			OutGuides.StrandsCurves.CurvesCount[CardIt] = PointCount;
-			OutGuides.StrandsCurves.CurvesOffset[CardIt] = PointOffset;
-			OutGuides.StrandsCurves.CurvesLength[CardIt] = TotalLength;
-			if (OutGuides.StrandsCurves.CurvesRootUV.Num()) OutGuides.StrandsCurves.CurvesRootUV[CardIt] = FVector2f(0, 0);
-
-			float CurrentLength = 0;
-			for (uint32 PointIt = 0; PointIt < PointCount; ++PointIt)
+		// Iterate over all triangles of a cards, and find vertices which share either same U or same V. We add them to separate lists. 
+		// We then use the following heuristic: the main axis will have more segments. This is what determine the principal axis of the cards
+		const float UVCoordTreshold = 1.f / 1024.f; // 1 pixel for a 1k texture
+		for (uint32 TriangleIt = 0; TriangleIt < NumTriangles; ++TriangleIt)
+		{
+			const uint32 VertexIndexOffset = VertexOffset + TriangleIt * 3;
+			for (uint32 VertexIt = 0; VertexIt < 3; ++VertexIt)
 			{
-				const FVector3f P0 = CenterPoints[PointIt];
-
-				OutGuides.BoundingBox += (FVector)P0;
-
-				OutGuides.StrandsPoints.PointsPosition.Add(P0);
-				OutGuides.StrandsPoints.PointsCoordU.Add(FMath::Clamp(CurrentLength / TotalLength, 0.f, 1.f));
-				OutGuides.StrandsPoints.PointsRadius.Add(1);
-
-				// Simple geometric normal based on adjacent vertices
-				if (PointIt < PointCount - 1)
-				{
-					const FVector3f P1 = CenterPoints[PointIt + 1];
-					const float SegmentLength = (P1 - P0).Size();
-					CurrentLength += SegmentLength;
-				}
+				const uint32 VertexIndex = InCards.Indices[VertexIndexOffset + VertexIt];
+				const FVector4f UV = TriangleUVs[VertexIt] = InCards.UVs[VertexIndex];
+				AddSimilarUV(SimilarVertexU, VertexIndex, UV.X, UV.Y, UVCoordTreshold);
+				AddSimilarUV(SimilarVertexV, VertexIndex, UV.Y, UV.X, UVCoordTreshold);
 			}
 		}
 
-		return true;
-	}
-
-	// Build interpolation data between the cards geometry and the guides
-	void CreateCardsInterpolation(
-		const FHairCardsGeometry& InCards,
-		const FHairStrandsDatas& InGuides,
-		FHairCardsInterpolationDatas& Out,
-		const TArray<float>& CardLengths)
-	{		
-		const uint32 TotalVertexCount = InCards.Positions.Num();
-		Out.SetNum(TotalVertexCount);
-
-		// For each cards, and for each cards vertex,
-		// Compute the closest guide points (two guide points to interpolation in-between), 
-		// and compute their indices and lerping value
-		const uint32 CardsCount = InCards.IndexOffsets.Num();
-		for (uint32 CardIt = 0; CardIt < CardsCount; ++CardIt)
+		// Find the perpendicular direction by comparing the number of segment along U and along V
+		TArray<FVector3f> CenterPoints;
 		{
-			const uint32 GuidePointOffset = InGuides.StrandsCurves.CurvesOffset[CardIt];
-			const uint32 GuidePointCount = InGuides.StrandsCurves.CurvesCount[CardIt];
-			check(GuidePointCount >= 2);
-
-			const uint32 IndexOffset = InCards.IndexOffsets[CardIt];
-			const uint32 IndexCount  = InCards.IndexCounts[CardIt];
-
-			const float CardLength = CardLengths[CardIt];
+			// Sort vertices along the main axis so that, when we iterate through them, we get a correct linear ordering
+			TArray<FSimilarUVVertices>& SimilarVertex = bIsMainDirectionU ? SimilarVertexU : SimilarVertexV;
+			SimilarVertex.Sort([](const FSimilarUVVertices& A, const FSimilarUVVertices& B)
+			{
+				return A.TexCoord < B.TexCoord;
+			});
 
-			for (uint32 IndexIt = 0; IndexIt < IndexCount; ++IndexIt)
+			// For each group of vertex with similar 'principal' tex coord, sort them with growing 'perpendicular/secondary' tex coord
+			for (FSimilarUVVertices& Similar : SimilarVertex)
 			{
-				const uint32 VertexIndex = InCards.Indices[IndexOffset + IndexIt];
-				const float CoordU = InCards.CoordU[VertexIndex];
-				check(CoordU >= 0.f && CoordU <= 1.f);
-
-				uint32 GuideIndex0 = GuidePointOffset + 0;
-				uint32 GuideIndex1 = GuidePointOffset + 1;
-				float  GuideLerp   = 0;
-				bool bFoundMatch = false;
-				for (uint32 GuidePointIt = 0; GuidePointIt < GuidePointCount-1; ++GuidePointIt)
+				Similar.AllIndices.Sort([](const FIndexAndCoord& A, const FIndexAndCoord& B)
 				{
-					const float GuideCoordU0 = InGuides.StrandsPoints.PointsCoordU[GuidePointOffset + GuidePointIt];
-					const float GuideCoordU1 = InGuides.StrandsPoints.PointsCoordU[GuidePointOffset + GuidePointIt + 1];
-					if (GuideCoordU0 <= CoordU && CoordU <= GuideCoordU1)
-					{
-						GuideIndex0 = GuidePointOffset + GuidePointIt;
-						GuideIndex1 = GuidePointOffset + GuidePointIt + 1;
-						const float LengthDiff = GuideCoordU1 - GuideCoordU0;
-						GuideLerp = (CoordU - GuideCoordU0) / (LengthDiff>0 ? LengthDiff : 1);
-						GuideLerp = FMath::Clamp(GuideLerp, 0.f, 1.f);
-						bFoundMatch = true;
-						break;
-					}
-				}
+					return A.TexCoord < B.TexCoord;
+				});
 
-				if (!bFoundMatch)
+				// Compute normalize coordinate
+				float MinTexCoord = Similar.AllIndices[0].TexCoord;
+				float MaxTexCoord = Similar.AllIndices[0].TexCoord;
+				for (const FIndexAndCoord& A : Similar.AllIndices)
 				{
-					GuideIndex0 = GuidePointOffset + GuidePointCount - 2;
-					GuideIndex1 = GuidePointOffset + GuidePointCount - 1;
-					GuideLerp = 1;
+					MinTexCoord = FMath::Min(A.TexCoord, MinTexCoord);
+					MaxTexCoord = FMath::Max(A.TexCoord, MaxTexCoord);
+				}
+				MaxTexCoord = FMath::Max(MaxTexCoord, MinTexCoord + KINDA_SMALL_NUMBER);
+				for (FIndexAndCoord& A : Similar.AllIndices)
+				{
+					A.TexCoord = (A.TexCoord - MinTexCoord) / (MaxTexCoord - MinTexCoord);
 				}
-
-				Out.PointsSimCurvesIndex[VertexIndex] = CardIt;
-				Out.PointsSimCurvesVertexIndex[VertexIndex] = GuideIndex0;
-				Out.PointsSimCurvesVertexLerp[VertexIndex] = GuideLerp;
 			}
-		}
-	}
-
-	void CreateCardsInterpolation(
-		const FHairCardsProceduralGeometry& InCards,
-		const FHairStrandsDatas& InGuides,
-		FHairCardsInterpolationDatas& Out)
-	{
-		CreateCardsInterpolation(InCards, InGuides, Out, InCards.Lengths);
-	}
-
-	static void CreateCardsAtlas(
-		const FHairStrandsDatas& InStrands,
-		FHairStrandsClusters& InClusters,
-		FHairCardsProceduralAtlas& InAtlas,
-		const FHairCardsTextureSettings& InSettings)
-	{
-		InAtlas.Rects.Empty();
-		InAtlas.StrandsPositions.Empty();
-		InAtlas.StrandsAttributes.Empty();
 
-		uint32 VertexOffset = 0;
-		for (uint32 UniqueIt=0, UniqueCount=InClusters.AtlasUniqueRect.Num(); UniqueIt<UniqueCount; ++UniqueIt)
-		{
-			FHairCardsProceduralAtlas::Rect& OutRect = InAtlas.Rects.AddDefaulted_GetRef();
-			OutRect.Offset		= 0; //InClusters.AtlasUniqueRect[UniqueIt].Offset;
-			OutRect.Resolution	= 0; //InClusters.AtlasUniqueRect[UniqueIt].Resolution;
-			OutRect.VertexOffset= VertexOffset;
-
-			const uint32 ClusterIndex = InClusters.AtlasUniqueRectIndex[UniqueIt];
-			const uint32 CurveOffset  = InClusters.CurveOffset[ClusterIndex];
-			const uint32 CurveCount   = InClusters.CurveCount[ClusterIndex];
-
-			OutRect.MinBound = FVector3f( FLT_MAX);
-			OutRect.MaxBound = FVector3f(-FLT_MAX);
-
-			// Precompute guide rest & deformed position as well as the length of the guides.
-			TArray<FVector3f> Guide_RestPositions;
-			TArray<FVector3f> Guide_DeformedPositions;
-			TArray<float>   Guide_Lengths;
-			FVector3f AlignmentDir = FVector3f::ZeroVector;
+			CenterPoints.Reserve(SimilarVertex.Num());
+			FVector3f PrevCenterPoint = FVector3f::ZeroVector;
+			float TotalLength = 0;
+			for (const FSimilarUVVertices& Similar : SimilarVertex)
 			{
-				const uint32 GuidePointCount  = InClusters.GuidePointCount[ClusterIndex];
-				const uint32 GuidePointOffset = InClusters.GuidePointOffset[ClusterIndex];
-
-				check(GuidePointCount >= 2);
-				const FVector3f Dir_P0 = InClusters.GuidePoints[GuidePointOffset];
-				const FVector3f Dir_P1 = InClusters.GuidePoints[GuidePointOffset + 1];
-				AlignmentDir = (Dir_P1 - Dir_P0).GetSafeNormal();
-
-				// Align the direction onto one of the cardinal axis for easing the projection
+				// Compute avg center point of the guide
+				FVector3f CenterPoint = FVector3f::ZeroVector;
+				for (uint32 VertexIndex : Similar.Indices)
 				{
-					FVector3f MajorAxis = FVector3f(1,0,0) * FMath::Sign(AlignmentDir.X);
-					float S = FMath::Abs(AlignmentDir.X);
-					if (FMath::Abs(AlignmentDir.Y) > S) { MajorAxis = FMath::Sign(AlignmentDir.Y) * FVector3f(0, 1, 0); S = FMath::Abs(AlignmentDir.Y); }
-					if (FMath::Abs(AlignmentDir.Z) > S) { MajorAxis = FMath::Sign(AlignmentDir.Z) * FVector3f(0, 0, 1); S = FMath::Abs(AlignmentDir.Z); }
-					AlignmentDir = MajorAxis;
+					CenterPoint += InCards.Positions[VertexIndex];
 				}
-				const FVector3f Root = Dir_P0;
+				CenterPoint /= Similar.Indices.Num() > 0 ? Similar.Indices.Num() : 1;
 
-				float Length = 0;
-				for (uint32 PointIt = 0; PointIt < GuidePointCount; ++PointIt)
+				// Update length along the guide
+				if (CenterPoints.Num() > 0)
 				{
-					const uint32 GuidePointIt = GuidePointOffset + PointIt;
-
-					const FVector3f& P0_Rest = InClusters.GuidePoints[GuidePointIt];
-					const FVector3f  P0_Deformed = Root + AlignmentDir * Length;
-
-					Guide_RestPositions.Add(P0_Rest);
-					Guide_DeformedPositions.Add(P0_Deformed);
-					Guide_Lengths.Add(Length);
+					const float SegmentLength = (CenterPoint - PrevCenterPoint).Size();
+					TotalLength += SegmentLength;
+				}
 
-					if (PointIt < GuidePointCount-1)
-					{
-						const FVector3f& P1_Rest = InClusters.GuidePoints[GuidePointIt+1];
-						const float L = (P1_Rest - P0_Rest).Size();
-						Length += L > 0 ? L : 0.001f;
-					}
+				// Update neighbor vertex with current length (to compute the parametric distance at the end)
+				for (const FIndexAndCoord& VertexIndex : Similar.AllIndices)
+				{
+					InCards.CoordU[VertexIndex.Index] = TotalLength;
+					InCards.LocalUVs[VertexIndex.Index].X = TotalLength;
+					InCards.LocalUVs[VertexIndex.Index].Y = VertexIndex.TexCoord;
 				}
+
+				CenterPoints.Add(CenterPoint);
+				PrevCenterPoint = CenterPoint;
 			}
-			const uint32 GuidePointCount = Guide_RestPositions.Num();
-			check(GuidePointCount == Guide_DeformedPositions.Num());
-			check(GuidePointCount == Guide_Lengths.Num());
 
-			// Compute the strands/curves deformed (straighten) points to generate the cards texture
-			uint32 VertexCount = 0;
-			for (uint32 CurveIt = 0; CurveIt < CurveCount; ++CurveIt)
+			if (SimilarVertex.Num() > 1)
 			{
-				const uint32 CurveIndex  = InClusters.CurveIndices[CurveIt + CurveOffset];
-				const uint32 PointOffset = InStrands.StrandsCurves.CurvesOffset[CurveIndex];
-				const uint32 PointCount  = InStrands.StrandsCurves.CurvesCount[CurveIndex];
-
-				float Length = 0;
-				for (uint32 PointIt = 0; PointIt < PointCount; ++PointIt)
+				// Normalize length to have a parametric distance
+				for (const FSimilarUVVertices& Similar : SimilarVertex)
 				{
-					const uint32 PointIndex = PointOffset + PointIt;
-
-					const float StrandsU = InStrands.StrandsPoints.PointsCoordU[PointIndex];
-					const FVector2f RootUV = InStrands.StrandsCurves.CurvesRootUV.Num() ? InStrands.StrandsCurves.CurvesRootUV[CurveIndex] : FVector2f(0,0);
-					const float Seed = 0.f; // TODO float(InStrands.RenderData.Attributes[PointIndex].Seed) / 255.f;
-
-					const FVector3f& P0_Rest = InStrands.StrandsPoints.PointsPosition[PointIndex];
-					const float Radius = InStrands.StrandsPoints.PointsRadius[PointIndex];
-
-					// Find the closest guide point using parametric length at a correspondance/matching metric
-					uint32 ClosestGuidePointIt0 = 0;
-					uint32 ClosestGuidePointIt1 = 1;
-					while (ClosestGuidePointIt1 < GuidePointCount && !(Guide_Lengths[ClosestGuidePointIt0] <= Length && Length <= Guide_Lengths[ClosestGuidePointIt1]))
-					{
-						++ClosestGuidePointIt0;
-						++ClosestGuidePointIt1;
-					}
-
-					// Once the bounded indices are compute the displacement vector used for deforming the curves
-					FVector3f Displacement = FVector3f::ZeroVector;
-					if (ClosestGuidePointIt1 < GuidePointCount)
-					{
-						const float U = (Length - Guide_Lengths[ClosestGuidePointIt0]) / (Guide_Lengths[ClosestGuidePointIt1] - Guide_Lengths[ClosestGuidePointIt0]);
-						const FVector3f& GuideP_Rest = FMath::Lerp(Guide_RestPositions[ClosestGuidePointIt0], Guide_RestPositions[ClosestGuidePointIt1], U);
-						const FVector3f& GuideP_Deformed = FMath::Lerp(Guide_DeformedPositions[ClosestGuidePointIt0], Guide_DeformedPositions[ClosestGuidePointIt1], U);
-						Displacement = GuideP_Deformed - GuideP_Rest;
-					}
-					else
-					{
-						// Extrapolate the guide rest position, and guide deformed position to compute the displacement
-						const FVector3f& GuideP_Rest0 = Guide_RestPositions[GuidePointCount - 2];
-						const FVector3f& GuideP_Rest1 = Guide_RestPositions[GuidePointCount - 1];
-
-						const FVector3f& GuideP_Deformed0 = Guide_DeformedPositions[GuidePointCount - 2];
-						const FVector3f& GuideP_Deformed1 = Guide_DeformedPositions[GuidePointCount - 1];
-
-						const FVector3f GuideDir_Rest		= (GuideP_Rest1 - GuideP_Rest0).GetSafeNormal();
-						const FVector3f GuideDir_Deformed = (GuideP_Deformed1 - GuideP_Deformed0).GetSafeNormal();
-
-						const float ExtraLength = Length - Guide_Lengths[GuidePointCount - 1];
-						const FVector3f& GuideP_RestExtrapolated		= GuideP_Rest1 + GuideDir_Rest * ExtraLength;
-						const FVector3f& GuideP_DeformedExtrapolated	= GuideP_Deformed1 + GuideDir_Deformed * ExtraLength;
-
-						Displacement = GuideP_DeformedExtrapolated - GuideP_RestExtrapolated;
-					}
-
-					const FVector3f P0_Deformed = P0_Rest + Displacement;
-
-					if (PointIt < PointCount - 1)
+					for (uint32 VertexIndex : Similar.Indices)
 					{
-						const FVector3f& P1_Rest = InStrands.StrandsPoints.PointsPosition[PointOffset + PointIt + 1];
-						Length += (P1_Rest - P0_Rest).Size();
+						InCards.CoordU[VertexIndex] /= TotalLength;
+						InCards.LocalUVs[VertexIndex].X = InCards.CoordU[VertexIndex];
 					}
-
-					// Bounding box update. This is used later on to find the rasterizaztion axis, and do the actual rasterization of the strands (on GPU)
-					OutRect.MinBound.X = FMath::Min(P0_Deformed.X, OutRect.MinBound.X);
-					OutRect.MinBound.Y = FMath::Min(P0_Deformed.Y, OutRect.MinBound.Y);
-					OutRect.MinBound.Z = FMath::Min(P0_Deformed.Z, OutRect.MinBound.Z);
-
-					OutRect.MaxBound.X = FMath::Max(P0_Deformed.X, OutRect.MaxBound.X);
-					OutRect.MaxBound.Y = FMath::Max(P0_Deformed.Y, OutRect.MaxBound.Y);
-					OutRect.MaxBound.Z = FMath::Max(P0_Deformed.Z, OutRect.MaxBound.Z);
-					InAtlas.StrandsPositions.Add(FVector4f(P0_Deformed, (PointIt < PointCount - 1) ? Radius : 0));
-					InAtlas.StrandsAttributes.Add(FVector4f(RootUV.X, RootUV.Y, StrandsU, Seed));
-						
-					++VertexCount;
 				}
 			}
+		}
 
-			// Compute rasterization info (rasterization axis) and cards dimension (width/length)
-			{
-
-				FVector3f::FReal Size[3] = 
-				{ 
-					OutRect.MaxBound.X - OutRect.MinBound.X, 
-					OutRect.MaxBound.Y - OutRect.MinBound.Y, 
-					OutRect.MaxBound.Z - OutRect.MinBound.Z 
-				};
-				FVector3f Axis[3] = 
-				{ 
-					Size[0] * FVector3f(1,0,0),
-					Size[1] * FVector3f(0,1,0), 
-					Size[2] * FVector3f(0,0,1) 
-				};
-
-				if (Size[0] < Size[1]) { Swap(Axis[0], Axis[1]); Swap(Size[0], Size[1]); }
-				if (Size[1] < Size[2]) { Swap(Axis[1], Axis[2]); Swap(Size[1], Size[2]); }
-				if (Size[0] < Size[1]) { Swap(Axis[0], Axis[1]); Swap(Size[0], Size[1]); }
+		// Insure that cards as at least two points to build a segment as a lot of the runtime code assume at have
+		// at least one valid segment
+		check(CenterPoints.Num() > 0);
+		if (CenterPoints.Num() == 1)
+		{
+			const FVector3f CenterPoint = CenterPoints[0];
+			const float SegmentSize = 0.5f;
+			const FVector3f P1 = CenterPoint + SegmentSize * (CenterPoint - InCards.BoundingBox.GetCenter()).GetSafeNormal();
+			CenterPoints.Add(P1);
+		}
 
-				FVector3f T = Axis[0];
-				FVector3f B = Axis[1];
-				FVector3f N = Axis[2];
+		// Compute and store the guide's total length
+		const uint32 PointCount = CenterPoints.Num();
+		float TotalLength = 0.f;
+		for (uint32 PointIt = 0; PointIt < PointCount - 1; ++PointIt)
+		{
+			TotalLength += (CenterPoints[PointIt + 1] - CenterPoints[PointIt]).Size();
+		}
+		OutCardLengths[CardIt] = TotalLength;
 
-				const FVector3f NormalizedT = T.GetSafeNormal();
+		const uint32 PointOffset = OutGuides.StrandsPoints.PointsPosition.Num();
 
-				if (FMath::Sign(T.X) != FMath::Sign(AlignmentDir.X) ||
-					FMath::Sign(T.Y) != FMath::Sign(AlignmentDir.Y) ||
-					FMath::Sign(T.Z) != FMath::Sign(AlignmentDir.Z))
-				{
-					T = -T;
-					B = -B;
-					N = -N;
-				}
+		OutGuides.StrandsCurves.CurvesCount[CardIt] = PointCount;
+		OutGuides.StrandsCurves.CurvesOffset[CardIt] = PointOffset;
+		OutGuides.StrandsCurves.CurvesLength[CardIt] = TotalLength;
+		if (OutGuides.StrandsCurves.CurvesRootUV.Num()) OutGuides.StrandsCurves.CurvesRootUV[CardIt] = FVector2f(0, 0);
 
-				// Orient the cards along the vertical or horizontal axis
-			#define CARDS_ATLAS_ORIENTATION_ALONG_Y 1
-			#if CARDS_ATLAS_ORIENTATION_ALONG_Y == 0
-				OutRect.RasterAxisX =  B; // Not normalized on purpose
-				OutRect.RasterAxisY =  T; // Not normalized on purpose
-				OutRect.RasterAxisZ = -N; // Not normalized on purpose
-
-				OutRect.CardWidth  = B.Size();
-				OutRect.CardLength = T.Size();
-
-				OutRect.Resolution.X = InSettings.PixelPerCentimeters * OutRect.CardWidth;
-				OutRect.Resolution.Y = InSettings.PixelPerCentimeters * OutRect.CardLength;
-			#else
-				OutRect.RasterAxisX = T; // Not normalized on purpose
-				OutRect.RasterAxisY = B; // Not normalized on purpose
-				OutRect.RasterAxisZ =-N; // Not normalized on purpose
-
-				OutRect.CardWidth = B.Size();
-				OutRect.CardLength = T.Size();
-
-				OutRect.Resolution.X = InSettings.PixelPerCentimeters * OutRect.CardLength;
-				OutRect.Resolution.Y = InSettings.PixelPerCentimeters * OutRect.CardWidth;
-				#endif
-			}
+		float CurrentLength = 0;
+		for (uint32 PointIt = 0; PointIt < PointCount; ++PointIt)
+		{
+			const FVector3f P0 = CenterPoints[PointIt];
 
-			OutRect.VertexCount = VertexCount;
-			VertexOffset += VertexCount;
-		}
+			OutGuides.BoundingBox += (FVector)P0;
 
-		// Pack the atlas texture
-		PackHairCardsRects(InSettings, InClusters, InAtlas);
+			OutGuides.StrandsPoints.PointsPosition.Add(P0);
+			OutGuides.StrandsPoints.PointsCoordU.Add(FMath::Clamp(CurrentLength / TotalLength, 0.f, 1.f));
+			OutGuides.StrandsPoints.PointsRadius.Add(1);
 
-		// Update the width of the cluster based on the computed width/length ratio for cards texture generation
-		{
-			const uint32 ClusterCount = InClusters.GetNum();
-			for (uint32 ClusterIt = 0; ClusterIt < ClusterCount; ++ClusterIt)
+			// Simple geometric normal based on adjacent vertices
+			if (PointIt < PointCount - 1)
 			{
-				const float ClusterBoundLength  = InClusters.BoundLength[ClusterIt];
-				const uint32 ClusterBoundCount  = InClusters.BoundCount[ClusterIt];
-				const uint32 ClusterBoundOffset = InClusters.BoundOffset[ClusterIt];
-
-				const uint32 UniqueRectIndex = InClusters.AtlasRectIndex[ClusterIt];
-				const FIntPoint CardResolution = InClusters.AtlasUniqueRect[UniqueRectIndex].Resolution;
-
-				for (uint32 VertexIt = 0; VertexIt < ClusterBoundCount; ++VertexIt)
-				{
-					const uint32 Index = ClusterBoundOffset + VertexIt;
-					const float Width = float(CardResolution.Y) / float(CardResolution.X) * ClusterBoundLength;
-					InClusters.BoundAxis_B[Index] = InClusters.BoundAxis_B[Index].GetSafeNormal() * Width * 0.5f;
-				}
+				const FVector3f P1 = CenterPoints[PointIt + 1];
+				const float SegmentLength = (P1 - P0).Size();
+				CurrentLength += SegmentLength;
 			}
 		}
 	}
-}
-
-////////////////////////////////////////////////////////////////////////////////////////////////
-// Build public functions
-
-namespace FHairCardsBuilder
-{
-
-FString GetVersion()
-{
-	// Important to update the version when cards building or importing changes
-	return TEXT("9g");
-}
-
-void AllocateAtlasTexture(UTexture2D* Out, const FIntPoint& Resolution, uint32 MipCount, EPixelFormat PixelFormat, ETextureSourceFormat SourceFormat)
-{
-	if (!Out)
-	{
-		return;
-	}
-
-	FTextureFormatSettings FormatSettings;
-	FormatSettings.CompressionNone = true;
-	FormatSettings.CompressionSettings = TC_Masks;
-	FormatSettings.SRGB = false;
-
-#if WITH_EDITORONLY_DATA
-	Out->Source.Init(Resolution.X, Resolution.Y, 1, MipCount, SourceFormat, nullptr);
-#endif // #if WITH_EDITORONLY_DATA
-	Out->LODGroup = TEXTUREGROUP_EffectsNotFiltered; // Mipmap filtering, no compression
-#if WITH_EDITORONLY_DATA
-	Out->SetLayerFormatSettings(0, FormatSettings);
-#endif // #if WITH_EDITORONLY_DATA
-
-	// No need to allocate the platform data as they will be allocating once the source data are filled in
-#if 0
-	Out->PlatformData = new FTexturePlatformData();
-	Out->PlatformData->SizeX = Resolution.X;
-	Out->PlatformData->SizeY = Resolution.Y;
-	Out->PlatformData->PixelFormat = PixelFormat;
-
-	Out->UpdateResource();
-#endif
-}
-
-void AllocateAtlasTexture_Depth(UTexture2D* Out, const FIntPoint& Resolution, uint32 MipCount)		{ AllocateAtlasTexture(Out, Resolution, MipCount, PF_R8G8B8A8, ETextureSourceFormat::TSF_BGRA8); }
-void AllocateAtlasTexture_Coverage(UTexture2D* Out, const FIntPoint& Resolution, uint32 MipCount)	{ AllocateAtlasTexture(Out, Resolution, MipCount, PF_R8G8B8A8, ETextureSourceFormat::TSF_BGRA8); }
-void AllocateAtlasTexture_Tangent(UTexture2D* Out, const FIntPoint& Resolution, uint32 MipCount)	{ AllocateAtlasTexture(Out, Resolution, MipCount, PF_B8G8R8A8, ETextureSourceFormat::TSF_BGRA8); }
-void AllocateAtlasTexture_Attribute(UTexture2D* Out, const FIntPoint& Resolution, uint32 MipCount)	{ AllocateAtlasTexture(Out, Resolution, MipCount, PF_B8G8R8A8, ETextureSourceFormat::TSF_BGRA8); }
 
-void AllocateAtlasTexture(UTexture2D*& Out, const FIntPoint& InResolution, const FString& GroomPackageName, const FString& Suffix, TTextureAllocation TextureAllocation)
-{
-	if (!Out)
-	{
-		Out = FHairStrandsCore::CreateTexture(GroomPackageName, InResolution, Suffix, TextureAllocation);
-	}
-	else if (Out->GetSizeX() != InResolution.X || Out->GetSizeY() != InResolution.Y)
-	{
-		 FHairStrandsCore::ResizeTexture(Out, InResolution, TextureAllocation);
-	}
+	return true;
 }
 
-void BuildGeometry(
-	const FString& LODName,
-	const FHairStrandsDatas& InRen, 
-	const FHairStrandsDatas& InSim,
-	const FHairGroupsProceduralCards& Settings,
-	FHairCardsProceduralDatas& Out,
-	FHairCardsBulkData& OutBulk,
-	FHairStrandsDatas& OutGuides,
-	FHairCardsInterpolationBulkData& OutInterpolationBulk,
-	FHairGroupCardsTextures& OutTextures)
-{
-	// Basic algo:
-	// * Geometry
-	//	* Find cluster of similar strands (length, CVs, curls, ...)
-	//	* Find the median axis within a cluster
-	//	* Find the major/minor axis of the cluster
-	//	* Generate along the median axis, and orient the cluster card(s) around the major/minor axis
-	//
-	// * Texture
-	//  * Define cards resolution based on cards size (width/length) and the hair world size
-	//	* Project strands onto the cards (compute shader for having non-linear projection?)
-	//	* Store Depth, coverage, tangent, Root UV, UVs, vertex base color, vertex roughness, strands ID/seed, ...
-	// 
-	// * Atlas
-	//  * Allocate space into the atlas
-	//  * Copy and generate MIPs
-	//	* Reduce atlas based on cluster lengths/hair count/width similarity
-
-	// Card generation are refreshed when:
-	// * A groom's group is loaded
-	// * A groom's group parameters changed (width, clip length, root scale, tip scale, )
-	// * A cards settings parameters changed
-	// Note: clip length could maybe be implemented with a face along the cards
-	//
-	// 
-	// TODO: add some caching about the different par so that we don't reprocess the entire group when something change
-	FHairCardsInterpolationDatas OutInterpolation;
-	HairCards::FVoxelVolume Voxels;
-	{
-		// Build voxel structure
-		HairCards::VoxelizeGroom(InRen, 32, Voxels);
-		Out.Voxels.Resolution = Voxels.Resolution;
-		Out.Voxels.MinBound = Voxels.MinBound;
-		Out.Voxels.MaxBound = Voxels.MaxBound;
-		Out.Voxels.VoxelSize = Voxels.VoxelSize;
-
-		// Build cluster
-		HairCards::FHairStrandsClusters Clusters;
-		HairCards::CreateHairStrandsCluster(InRen, InSim, Voxels, Clusters, Settings);
-		check(Clusters.GetNum() > 0);
-
-		// Allocate atlas rects
-		HairCards::CreateCardsTexture(InRen, Clusters, Settings);
-		HairCards::CreateCardsAtlas(InRen, Clusters, Out.Atlas, Settings.TextureSettings);
-		check(Out.Atlas.Resolution.X > 0);
-		check(Out.Atlas.Resolution.Y > 0);
-
-		// Build cards
-		HairCards::CreateCardsGeometry(InRen, Clusters, Out.Atlas.Resolution, Settings, Voxels, Out.Cards);
-		HairCards::CreateCardsGuides(Clusters, OutGuides);
-		HairCards::CreateCardsInterpolation(Out.Cards, OutGuides, OutInterpolation);
-		check(Out.Cards.GetNum() > 0);
-
-		AllocateAtlasTexture(static_cast<UTexture2D*&>(OutTextures.DepthTexture),		Out.Atlas.Resolution, LODName, TEXT("_CardsAtlas_Depth"), AllocateAtlasTexture_Depth);
-		AllocateAtlasTexture(static_cast<UTexture2D*&>(OutTextures.CoverageTexture),	Out.Atlas.Resolution, LODName, TEXT("_CardsAtlas_Coverage"), AllocateAtlasTexture_Coverage);
-		AllocateAtlasTexture(static_cast<UTexture2D*&>(OutTextures.TangentTexture),		Out.Atlas.Resolution, LODName, TEXT("_CardsAtlas_Tangent"), AllocateAtlasTexture_Tangent);
-		AllocateAtlasTexture(static_cast<UTexture2D*&>(OutTextures.AttributeTexture),	Out.Atlas.Resolution, LODName, TEXT("_CardsAtlas_Attribute"), AllocateAtlasTexture_Attribute);
-	}
-
-	Out.Cards.BoundingBox.Init();
-
-	// Fill in render resources (do we need to keep it separated? e.g, format compression, packing)
-	const uint32 PointCount = Out.Cards.Positions.Num();
-	Out.RenderData.Positions.SetNum(PointCount);
-	Out.RenderData.Normals.SetNum(PointCount * FHairCardsNormalFormat::ComponentCount);
-	Out.RenderData.UVs.SetNum(PointCount);
-	for (uint32 PointIt = 0; PointIt < PointCount; ++PointIt)
-	{
-		check(FMath::IsFinite(Out.Cards.Positions[PointIt].X) && FMath::IsFinite(Out.Cards.Positions[PointIt].Y) && FMath::IsFinite(Out.Cards.Positions[PointIt].Z));
-		Out.RenderData.Positions[PointIt] = FVector4f(Out.Cards.Positions[PointIt], *((float*)&Out.Cards.CardIndices[PointIt]));
-		Out.RenderData.UVs[PointIt] = FVector4f(Out.Cards.UVs[PointIt]);
-		Out.RenderData.Normals[PointIt * 2] = FVector4f(Out.Cards.Tangents[PointIt], 0);
-		Out.RenderData.Normals[PointIt * 2 + 1] = FVector4f(Out.Cards.Normals[PointIt], 1);
-
-		Out.Cards.BoundingBox += FVector4f(Out.RenderData.Positions[PointIt]);
-	}
-
-	const uint32 IndexCount = Out.Cards.Indices.Num();
-	Out.RenderData.Indices.SetNum(IndexCount);
-	for (uint32 IndexIt = 0; IndexIt < IndexCount; ++IndexIt)
-	{
-		Out.RenderData.Indices[IndexIt] = Out.Cards.Indices[IndexIt];
-	}
-
-	// Copy transient Procedural cards bulk data into the final bulk data
-	OutBulk.Positions 	= Out.RenderData.Positions;
-	OutBulk.Normals 	= Out.RenderData.Normals;
-	OutBulk.Indices 	= Out.RenderData.Indices;
-	OutBulk.UVs 		= Out.RenderData.UVs;
-	OutBulk.BoundingBox = ToFBox3d(Out.Cards.BoundingBox);
-
-	const uint32 CardCount = Out.Cards.Lengths.Num();
-	Out.RenderData.CardsRect.SetNum(CardCount);
-	Out.RenderData.CardsLengths.SetNum(CardCount);
-	for (uint32 CardIt = 0; CardIt < CardCount; ++CardIt)
-	{
-		FHairCardsAtlasRectFormat::Type Rect16;
-		Rect16.X = Out.Cards.Rects[CardIt].Offset.X;
-		Rect16.Y = Out.Cards.Rects[CardIt].Offset.Y;
-		Rect16.Z = Out.Cards.Rects[CardIt].Resolution.X;
-		Rect16.W = Out.Cards.Rects[CardIt].Resolution.Y;
-
-		Out.RenderData.CardsRect[CardIt] = Rect16;
-		Out.RenderData.CardsLengths[CardIt] = Out.Cards.Lengths[CardIt];
-	}
-
-	Out.RenderData.CardItToCluster.SetNum(Out.Cards.CardIndexToClusterOffsetAndCount.Num());
-	for (uint32 It = 0, ItCount = Out.RenderData.CardItToCluster.Num(); It < ItCount; ++It)
+// Build interpolation data between the cards geometry and the guides
+void InternalCreateCardsInterpolation(
+	const FHairCardsGeometry& InCards,
+	const FHairStrandsDatas& InGuides,
+	FHairCardsInterpolationDatas& Out,
+	const TArray<float>& CardLengths)
+{		
+	const uint32 TotalVertexCount = InCards.Positions.Num();
+	Out.SetNum(TotalVertexCount);
+
+	// For each cards, and for each cards vertex,
+	// Compute the closest guide points (two guide points to interpolation in-between), 
+	// and compute their indices and lerping value
+	const uint32 CardsCount = InCards.IndexOffsets.Num();
+	for (uint32 CardIt = 0; CardIt < CardsCount; ++CardIt)
 	{
-		Out.RenderData.CardItToCluster[It] = Out.Cards.CardIndexToClusterOffsetAndCount[It];
-	}
+		const uint32 GuidePointOffset = InGuides.StrandsCurves.CurvesOffset[CardIt];
+		const uint32 GuidePointCount = InGuides.StrandsCurves.CurvesCount[CardIt];
+		check(GuidePointCount >= 2);
 
-	Out.RenderData.ClusterIdToVertices.SetNum(Out.Cards.ClusterIndexToVertexOffsetAndCount.Num());
-	for (uint32 It = 0, ItCount = Out.RenderData.ClusterIdToVertices.Num(); It < ItCount; ++It)
-	{
-		Out.RenderData.ClusterIdToVertices[It] = Out.Cards.ClusterIndexToVertexOffsetAndCount[It];
-	}
+		const uint32 IndexOffset = InCards.IndexOffsets[CardIt];
+		const uint32 IndexCount  = InCards.IndexCounts[CardIt];
 
-	Out.RenderData.ClusterBounds.SetNum(Out.Cards.Bounds.Num() * 4);
-	for (uint32 It = 0, ItCount = Out.Cards.Bounds.Num(); It < ItCount; ++It)
-	{
-		Out.RenderData.ClusterBounds[It * 4    ] = Out.Cards.Bounds[It].Center;
-		Out.RenderData.ClusterBounds[It * 4 + 1] = Out.Cards.Bounds[It].ExtentX;
-		Out.RenderData.ClusterBounds[It * 4 + 2] = Out.Cards.Bounds[It].ExtentY;
-		Out.RenderData.ClusterBounds[It * 4 + 3] = Out.Cards.Bounds[It].ExtentZ;
-	}
+		const float CardLength = CardLengths[CardIt];
 
-	const uint32 VoxelCount = Voxels.Density.Num();
-	Out.RenderData.VoxelDensity.SetNum(Voxels.Density.Num());
-	Out.RenderData.VoxelTangent.SetNum(Voxels.Tangent.Num());
-	Out.RenderData.VoxelNormal.SetNum(Voxels.Normal.Num());
-	for (uint32 It = 0; It < VoxelCount; ++It)
-	{
-		Out.RenderData.VoxelDensity[It] = Voxels.Density[It];
-		Out.RenderData.VoxelTangent[It] = Voxels.Tangent[It];
-		Out.RenderData.VoxelNormal[It] = Voxels.Normal[It];
-	}
+		for (uint32 IndexIt = 0; IndexIt < IndexCount; ++IndexIt)
+		{
+			const uint32 VertexIndex = InCards.Indices[IndexOffset + IndexIt];
+			const float CoordU = InCards.CoordU[VertexIndex];
+			check(CoordU >= 0.f && CoordU <= 1.f);
 
-	// Strands point for cards texture generation
-	const uint32 StrandsPointCount = Out.Atlas.StrandsPositions.Num();
-	Out.RenderData.CardsStrandsPositions.SetNum(StrandsPointCount);
-	for (uint32 PointIt = 0; PointIt < StrandsPointCount; ++PointIt)
-	{
-		Out.RenderData.CardsStrandsPositions[PointIt] = FVector4f(Out.Atlas.StrandsPositions[PointIt]);
-	}
+			uint32 GuideIndex0 = GuidePointOffset + 0;
+			uint32 GuideIndex1 = GuidePointOffset + 1;
+			float  GuideLerp   = 0;
+			bool bFoundMatch = false;
+			for (uint32 GuidePointIt = 0; GuidePointIt < GuidePointCount-1; ++GuidePointIt)
+			{
+				const float GuideCoordU0 = InGuides.StrandsPoints.PointsCoordU[GuidePointOffset + GuidePointIt];
+				const float GuideCoordU1 = InGuides.StrandsPoints.PointsCoordU[GuidePointOffset + GuidePointIt + 1];
+				if (GuideCoordU0 <= CoordU && CoordU <= GuideCoordU1)
+				{
+					GuideIndex0 = GuidePointOffset + GuidePointIt;
+					GuideIndex1 = GuidePointOffset + GuidePointIt + 1;
+					const float LengthDiff = GuideCoordU1 - GuideCoordU0;
+					GuideLerp = (CoordU - GuideCoordU0) / (LengthDiff>0 ? LengthDiff : 1);
+					GuideLerp = FMath::Clamp(GuideLerp, 0.f, 1.f);
+					bFoundMatch = true;
+					break;
+				}
+			}
 
-	Out.RenderData.CardsStrandsAttributes.SetNum(StrandsPointCount);
-	for (uint32 PointIt = 0; PointIt < StrandsPointCount; ++PointIt)
-	{
-		Out.RenderData.CardsStrandsAttributes[PointIt] = FVector4f(Out.Atlas.StrandsAttributes[PointIt]);
-	}
+			if (!bFoundMatch)
+			{
+				GuideIndex0 = GuidePointOffset + GuidePointCount - 2;
+				GuideIndex1 = GuidePointOffset + GuidePointCount - 1;
+				GuideLerp = 1;
+			}
 
-	// Interpolation data
-	OutInterpolationBulk.Interpolation.SetNum(PointCount);
-	for (uint32 PointIt = 0; PointIt < PointCount; ++PointIt)
-	{
-		const uint32 VertexIndex = OutInterpolation.PointsSimCurvesVertexIndex[PointIt];
-		const float VertexLerp = OutInterpolation.PointsSimCurvesVertexLerp[PointIt];
-		FHairCardsInterpolationVertex PackedData;
-		PackedData.VertexIndex = VertexIndex;
-		PackedData.VertexLerp  = FMath::Clamp(uint32(VertexLerp * 0xFF), 0u, 0xFFu);
-		OutInterpolationBulk.Interpolation[PointIt] = PackedData;
+			Out.PointsSimCurvesIndex[VertexIndex] = CardIt;
+			Out.PointsSimCurvesVertexIndex[VertexIndex] = GuideIndex0;
+			Out.PointsSimCurvesVertexLerp[VertexIndex] = GuideLerp;
+		}
 	}
 }
 
@@ -3809,11 +589,11 @@ bool InternalImportGeometry(
 
 	TArray<float> CardLengths;
 	CardLengths.Reserve(CardsCount);
-	bool bSuccess = HairCards::CreateCardsGuides(Out.Cards, OutGuides, CardLengths);
+	bool bSuccess = InternalCreateCardsGuides(Out.Cards, OutGuides, CardLengths);
 	if (bSuccess)
 	{
 		FHairCardsInterpolationDatas InterpolationData;
-		HairCards::CreateCardsInterpolation(Out.Cards, OutGuides, InterpolationData, CardLengths);
+		InternalCreateCardsInterpolation(Out.Cards, OutGuides, InterpolationData, CardLengths);
 
 		// Fill out the interpolation data
 		OutInterpolationBulkData.Interpolation.SetNum(PointCount);
@@ -4086,153 +866,8 @@ bool ExtractCardsData(const UStaticMesh* StaticMesh, const FHairStrandsDatas& In
 		OutInterpolationBulkData); 
 }
 
-///////////////////////////////////////////////////////////////////////////////////////////////////
-// Asynchronous queuing for follicle texture mask generation
-struct FCardsAtlasQuery
-{
-	FHairCardsProceduralDatas* ProceduralData = nullptr;
-	FHairCardsRestResource* RestResource = nullptr;
-	FHairCardsProceduralResource* ProceduralResource = nullptr;
-	FHairGroupCardsTextures* Textures = nullptr;
-};
-TQueue<FCardsAtlasQuery> GCardsAtlasQuery;
-
-
-void BuildTextureAtlas(
-	FHairCardsProceduralDatas* ProceduralData,
-	FHairCardsRestResource* RestResource,
-	FHairCardsProceduralResource* ProceduralResource,
-	FHairGroupCardsTextures* Textures)
-{
-	FCardsAtlasQuery Q;
-	Q.ProceduralData = ProceduralData;
-	Q.RestResource = RestResource;
-	Q.ProceduralResource = ProceduralResource;
-	Q.Textures = Textures;
-	GCardsAtlasQuery.Enqueue(Q);
-}
-
 } // namespace FHairCardsBuilder
 
-bool HasHairCardsAtlasQueries()
-{
-	return !FHairCardsBuilder::GCardsAtlasQuery.IsEmpty();
-}
-
-static void AddCardsTextureReadbackPass(
-	FRDGBuilder& GraphBuilder,
-	uint32 BytePerPixel,
-	FRDGTextureRef InputTexture,
-	UTexture2D* OutTexture)
-{
-	AddReadbackTexturePass(
-		GraphBuilder,
-		RDG_EVENT_NAME("CopyRDGToTexture2D"),
-		InputTexture,
-	[InputTexture, OutTexture, BytePerPixel](FRHICommandListImmediate& RHICmdList)
-	{
-		const FIntPoint Resolution = InputTexture->Desc.Extent;
-		check(OutTexture->Source.GetSizeX() == Resolution.X);
-		check(OutTexture->Source.GetSizeY() == Resolution.Y);
-
-		const FRHITextureCreateDesc Desc =
-			FRHITextureCreateDesc::Create2D(TEXT("CardsTextureReadbackPass_StagingTexture"), InputTexture->Desc.Extent, InputTexture->Desc.Format)
-			.SetNumMips(InputTexture->Desc.NumMips)
-			.SetFlags(ETextureCreateFlags::CPUReadback);
-
-		FTextureRHIRef StagingTexture = RHICreateTexture(Desc);
-
-		FRHICopyTextureInfo CopyInfo;
-		CopyInfo.NumMips = InputTexture->Desc.NumMips;
-		RHICmdList.CopyTexture(
-			InputTexture->GetRHI(),
-			StagingTexture,
-			CopyInfo);
-
-		// Flush, to ensure that all texture generation is done
-		RHICmdList.SubmitCommandsAndFlushGPU();
-		RHICmdList.BlockUntilGPUIdle();
-
-		// don't think we can mark package as a dirty in the package build
-#if WITH_EDITORONLY_DATA
-		void* InData = nullptr;
-		int32 Width = 0, Height = 0;
-		RHICmdList.MapStagingSurface(StagingTexture, InData, Width, Height);
-		uint8* InDataRGBA8 = (uint8*)InData;
-
-		uint8 MipIndex = 0;
-		const uint32 PixelCount = Resolution.X * Resolution.Y;
-		const uint32 SizeInBytes = BytePerPixel * PixelCount;
-		uint8* OutData = OutTexture->Source.LockMip(0);
-		// Since the source and target format are different we need to swizzle the channel
-		// RGBA8 -> TSF_BGRA8
-		for (uint32 PixelIt = 0; PixelIt < PixelCount; ++PixelIt)
-		{
-			const uint32 PixelOffset = PixelIt * 4;
-			OutData[PixelOffset + 0] = InDataRGBA8[PixelOffset + 2];
-			OutData[PixelOffset + 1] = InDataRGBA8[PixelOffset + 1];
-			OutData[PixelOffset + 2] = InDataRGBA8[PixelOffset + 0];
-			OutData[PixelOffset + 3] = InDataRGBA8[PixelOffset + 3];
-		}
-		OutTexture->Source.UnlockMip(0);
-
-		RHICmdList.UnmapStagingSurface(StagingTexture);
-
-		OutTexture->DeferCompression = true; // This forces reloading data when the asset is saved
-		OutTexture->MarkPackageDirty();
-#endif // #if WITH_EDITORONLY_DATA
-	});
-}
-
-void RunHairCardsAtlasQueries(
-	FRDGBuilder& GraphBuilder,
-	FGlobalShaderMap* ShaderMap,
-	const FShaderPrintData* DebugShaderData)
-{
-	FHairCardsBuilder::FCardsAtlasQuery Q;
-	while (FHairCardsBuilder::GCardsAtlasQuery.Dequeue(Q))
-	{
-		if (!Q.ProceduralData || !Q.RestResource || !Q.ProceduralResource || !Q.Textures)
-			continue;
-
-		// Allocate resources for generating the atlas texture (for editor mode only)
-		FRDGTextureRef DepthTexture		= GraphBuilder.CreateTexture(FRDGTextureDesc::Create2D(Q.ProceduralData->Atlas.Resolution, PF_R8G8B8A8, FClearValueBinding::Black, TexCreate_RenderTargetable | TexCreate_ShaderResource, 1), TEXT("Hair.CardsDepth"));
-		FRDGTextureRef TangentTexture	= GraphBuilder.CreateTexture(FRDGTextureDesc::Create2D(Q.ProceduralData->Atlas.Resolution, PF_R8G8B8A8, FClearValueBinding::Black, TexCreate_RenderTargetable | TexCreate_ShaderResource, 1), TEXT("Hair.CardTangent"));
-		FRDGTextureRef CoverageTexture	= GraphBuilder.CreateTexture(FRDGTextureDesc::Create2D(Q.ProceduralData->Atlas.Resolution, PF_R8G8B8A8, FClearValueBinding::Black, TexCreate_RenderTargetable | TexCreate_ShaderResource, 1), TEXT("Hair.CardCoverage"));
-		FRDGTextureRef AttributeTexture = GraphBuilder.CreateTexture(FRDGTextureDesc::Create2D(Q.ProceduralData->Atlas.Resolution, PF_R8G8B8A8, FClearValueBinding::Black, TexCreate_RenderTargetable | TexCreate_ShaderResource, 1), TEXT("Hair.CardAttribute"));
-		AddHairCardAtlasTexturePass(
-			GraphBuilder,
-			ShaderMap,
-			Q.ProceduralData->Atlas,
-			DebugShaderData,
-			Q.ProceduralResource->CardsStrandsPositions.Buffer->Desc.NumElements,
-			Q.ProceduralResource->CardsStrandsPositions.SRV,
-			Q.ProceduralResource->CardsStrandsAttributes.SRV,
-			DepthTexture,
-			TangentTexture,
-			CoverageTexture,
-			AttributeTexture);
-
-		AddCardsTextureReadbackPass(GraphBuilder, 4, DepthTexture,		Q.Textures->DepthTexture);
-		AddCardsTextureReadbackPass(GraphBuilder, 4, TangentTexture,	Q.Textures->TangentTexture);
-		AddCardsTextureReadbackPass(GraphBuilder, 4, CoverageTexture,	Q.Textures->CoverageTexture);
-		AddCardsTextureReadbackPass(GraphBuilder, 4, AttributeTexture,	Q.Textures->AttributeTexture);
-
-		Q.RestResource->DepthTexture	 = Q.Textures->DepthTexture->TextureReference.TextureReferenceRHI;
-		Q.RestResource->TangentTexture	 = Q.Textures->TangentTexture->TextureReference.TextureReferenceRHI;
-		Q.RestResource->CoverageTexture  = Q.Textures->CoverageTexture->TextureReference.TextureReferenceRHI;
-		Q.RestResource->AttributeTexture = Q.Textures->AttributeTexture->TextureReference.TextureReferenceRHI;
-		
-		Q.Textures->bNeedToBeSaved = true;
-		Q.ProceduralData->Atlas.bIsDirty = false;
-	}
-
-	if (GHairCardsDynamicAtlasRefresh > 0)
-	{
-		FHairCardsBuilder::GCardsAtlasQuery.Enqueue(Q);
-	}
-}
-
 namespace FHairMeshesBuilder
 {
 FString GetVersion()
@@ -4632,45 +1267,6 @@ void ConvertCardsGeometryToMeshDescription(const FHairCardsGeometry& In, FMeshDe
 	}
 }
 
-void ExportGeometry(const FHairCardsDatas& InCardsData, UStaticMesh* OutStaticMesh)
-{
-	if (!OutStaticMesh)
-	{
-		return;
-	}
-
-	FMeshDescription MeshDescription;
-	FStaticMeshAttributes Attributes(MeshDescription);
-	Attributes.Register();
-
-	// Actual conversion
-	ConvertCardsGeometryToMeshDescription(InCardsData.Cards, MeshDescription);
-
-	FMeshDescription* MeshDesc = OutStaticMesh->CreateMeshDescription(0);
-	*MeshDesc = MeshDescription;
-	if (OutStaticMesh->GetBodySetup() == nullptr)
-	{
-		OutStaticMesh->CreateBodySetup();
-	}
-	if (OutStaticMesh->GetBodySetup() != nullptr)
-	{
-		// enable complex as simple collision to use mesh directly
-		OutStaticMesh->GetBodySetup()->CollisionTraceFlag = ECollisionTraceFlag::CTF_UseComplexAsSimple;
-	}
-
-	// add a material slot. Must always have one material slot.
-	if (OutStaticMesh->GetStaticMaterials().Num() == 0)
-	{
-		OutStaticMesh->GetStaticMaterials().Add(FStaticMaterial());
-	}
-
-	// assuming we have updated the LOD 0 MeshDescription, tell UStaticMesh about this
-	OutStaticMesh->CommitMeshDescription(0);
-
-	// not sure what this does...marks things dirty? updates stuff after modification??
-	OutStaticMesh->PostEditChange();
-}
-
 } // namespace FHairCardsExporter
 
 #endif // WITH_EDITOR 
diff --git a/Engine/Plugins/Runtime/HairStrands/Source/HairStrandsCore/Public/GroomAsset.h b/Engine/Plugins/Runtime/HairStrands/Source/HairStrandsCore/Public/GroomAsset.h
--- a/Engine/Plugins/Runtime/HairStrands/Source/HairStrandsCore/Public/GroomAsset.h
+++ b/Engine/Plugins/Runtime/HairStrands/Source/HairStrandsCore/Public/GroomAsset.h
@@ -699,10 +699,6 @@ public:
 	FOnGroomAsyncLoadFinished OnGroomAsyncLoadFinished;
 
 	void MarkMaterialsHasChanged();
-
-	// Save out a static mesh based on generated cards
-	void SaveProceduralCards(uint32 CardsGroupIndex);
-	
 #endif
 
 	/** Array of user data stored with the asset */
diff --git a/Engine/Plugins/Runtime/HairStrands/Source/HairStrandsCore/Public/GroomAssetCards.h b/Engine/Plugins/Runtime/HairStrands/Source/HairStrandsCore/Public/GroomAssetCards.h
--- a/Engine/Plugins/Runtime/HairStrands/Source/HairStrandsCore/Public/GroomAssetCards.h
+++ b/Engine/Plugins/Runtime/HairStrands/Source/HairStrandsCore/Public/GroomAssetCards.h
@@ -14,136 +14,6 @@ class UTexture2D;
 enum class EHairAtlasTextureType : uint8;
 class UGroomAsset;
 
-UENUM(BlueprintType)
-enum class EHairCardsClusterType : uint8
-{
-	Low,
-	High,
-};
-
-UENUM(BlueprintType)
-enum class EHairCardsGenerationType : uint8
-{
-	CardsCount,
-	UseGuides,
-};
-
-/* Deprecated */
-USTRUCT(BlueprintType)
-struct HAIRSTRANDSCORE_API FHairCardsClusterSettings
-{
-	GENERATED_BODY()
-
-	FHairCardsClusterSettings();
-
-	/** Decimation factor use to initialized cluster (only used when UseGuide is disabled). This changes the number of generated cards */
-	UPROPERTY()
-	float ClusterDecimation;
-
-	/** Quality of clustering when group hair to cluster center. This does not change the number cards, but only how cards are shaped (size/shape) */
-	UPROPERTY()
-	EHairCardsClusterType Type;
-
-	/** Use the simulation guide to generate the cards instead of the decimation parameter. This changes the number of generated cards. */
-	UPROPERTY()
-	bool bUseGuide;
-
-	bool operator==(const FHairCardsClusterSettings& A) const;
-};
-
-
-USTRUCT(BlueprintType)
-struct HAIRSTRANDSCORE_API FHairCardsGeometrySettings
-{
-	GENERATED_BODY()
-
-	FHairCardsGeometrySettings();
-
-	/** Define how cards should be generated. Cards count: define a targeted number of cards. Use guides: use simulation guide as cards. */
-	UPROPERTY(EditAnywhere, BlueprintReadWrite, Category = "ClusterSettings")
-	EHairCardsGenerationType GenerationType;
-
-	/** Define how many cards should be generated. The generated number can be lower, as some cards can be discarded by other options. */
-	UPROPERTY(EditAnywhere, BlueprintReadWrite, Category = "ClusterSettings", meta = (ClampMin = "1", UIMin = "1"))
-	int32 CardsCount;
-
-	/** Quality of clustering when group hair to belong to a similar cards. This does not change the number cards, but only how cards are shaped (size/shape) */
-	UPROPERTY(EditAnywhere, BlueprintReadWrite, Category = "ClusterSettings")
-	EHairCardsClusterType ClusterType;
-
-	/** Minimum size of a card segment */
-	UPROPERTY(EditAnywhere, BlueprintReadWrite, Category = "GeometrySettings", AdvancedDisplay, meta = (ClampMin = "0.1", UIMin = "0.1", UIMax = "8.0", SliderExponent = 6))
-	float MinSegmentLength;
-	
-	/** Max angular difference between adjacents vertices to remove vertices during simplification with MinSegmentLength, in degrees. */
-	UPROPERTY(EditAnywhere, Category = "DecimationSettings", meta = (ClampMin = "0", ClampMax = "45", UIMin = "0", UIMax = "45.0"))
-	float AngularThreshold;
-
-	/** Length below which generated cards are discard, as there are considered too small. (Default:0, which means no trimming) */
-	UPROPERTY(EditAnywhere, Category = "DecimationSettings", meta = (ClampMin = "0", UIMin = "0"))
-	float MinCardsLength;
-
-	/** Length above which generated cards are discard, as there are considered too larger. (Default:0, which means no trimming) */
-	UPROPERTY(EditAnywhere, Category = "DecimationSettings", meta = (ClampMin = "0", UIMin = "0"))
-	float MaxCardsLength;
-
-	bool operator==(const FHairCardsGeometrySettings& A) const;
-};
-
-
-USTRUCT(BlueprintType)
-struct HAIRSTRANDSCORE_API FHairCardsTextureSettings
-{
-	GENERATED_BODY()
-
-	FHairCardsTextureSettings();
-
-	/** Max atlas resolution */
-	UPROPERTY(EditAnywhere, BlueprintReadWrite, Category = "AtlasSettings", meta = (ClampMin = "512", UIMin = "512", UIMax = "8192"))
-	int32 AtlasMaxResolution;
-
-	/** Pixel resolution per centimeters */
-	UPROPERTY(EditAnywhere, BlueprintReadWrite, Category = "AtlasSettings", meta = (ClampMin = "4", UIMin = "2", UIMax = "128"))
-	int32 PixelPerCentimeters;
-
-	/** Number of unique clump textures*/
-	UPROPERTY(EditAnywhere, BlueprintReadWrite, Category = "AtlasSettings", meta = (ClampMin = "1", UIMin = "1", UIMax = "128"))
-	int32 LengthTextureCount;
-
-	/** Number of texture having variation of strands count */
-	//UPROPERTY(EditAnywhere, BlueprintReadWrite, Category = "AtlasSettings", meta = (ClampMin = "1", UIMin = "1", UIMax = "128"))
-	UPROPERTY()
-	int32 DensityTextureCount;
-
-	bool operator==(const FHairCardsTextureSettings& A) const;
-};
-
-USTRUCT(BlueprintType)
-struct HAIRSTRANDSCORE_API FHairGroupsProceduralCards
-{
-	GENERATED_BODY()
-
-	FHairGroupsProceduralCards();
-
-	/* Deprecated */
-	UPROPERTY()
-	FHairCardsClusterSettings ClusterSettings;
-
-	UPROPERTY(EditAnywhere, Category = "GeometrySettings", meta = (ToolTip = "Cards geometry settings"))
-	FHairCardsGeometrySettings GeometrySettings;
-
-	UPROPERTY(EditAnywhere, Category = "TextureSettings", meta = (ToolTip = "Cards texture atlas settings"))
-	FHairCardsTextureSettings TextureSettings;
-
-	/* Use to track when a cards asset need to be regenerated */
-	UPROPERTY()
-	int32 Version;
-
-	bool operator==(const FHairGroupsProceduralCards& A) const;
-
-	void BuildDDCKey(FArchive& Ar);
-};
-
 UENUM(BlueprintType)
 enum class EHairCardsSourceType : uint8
 {
@@ -219,21 +89,18 @@ struct HAIRSTRANDSCORE_API FHairGroupsCardsSourceDescription
 	UPROPERTY()
 	FName MaterialSlotName;
 
-	UPROPERTY(EditAnywhere, Category = "CardsSource")
-	EHairCardsSourceType SourceType = EHairCardsSourceType::Imported;
-
-	UPROPERTY(EditAnywhere, Category = "CardsSource")
-	TObjectPtr<UStaticMesh> ProceduralMesh = nullptr;
+	UPROPERTY()
+	EHairCardsSourceType SourceType_DEPRECATED = EHairCardsSourceType::Imported;
+	
+	UPROPERTY()
+	TObjectPtr<UStaticMesh> ProceduralMesh_DEPRECATED = nullptr;
 
 	UPROPERTY()
-	FString ProceduralMeshKey;
+	bool bInvertUV = false;
 
-	UPROPERTY(EditAnywhere, Category = "CardsSource")
+	UPROPERTY(EditAnywhere, Category = "CardsSource", meta = (DisplayName = "Mesh"))
 	TObjectPtr<UStaticMesh> ImportedMesh = nullptr;
 
-	UPROPERTY(EditAnywhere, Category = "CardsSource")
-	FHairGroupsProceduralCards ProceduralSettings;
-
 	UPROPERTY(EditAnywhere, Category = "CardsSource")
 	FHairGroupCardsTextures Textures;
 
diff --git a/Engine/Plugins/Runtime/HairStrands/Source/HairStrandsCore/Public/GroomResources.h b/Engine/Plugins/Runtime/HairStrands/Source/HairStrandsCore/Public/GroomResources.h
--- a/Engine/Plugins/Runtime/HairStrands/Source/HairStrandsCore/Public/GroomResources.h
+++ b/Engine/Plugins/Runtime/HairStrands/Source/HairStrandsCore/Public/GroomResources.h
@@ -667,46 +667,6 @@ struct FHairCardsRestResource : public FHairCommonResource
 	const FHairCardsBulkData& BulkData;
 };
 
-/* Render buffers that will be used for rendering */
-struct FHairCardsProceduralResource : public FHairCommonResource
-{
-	/** Build the hair strands resource */
-	FHairCardsProceduralResource(const FHairCardsProceduralDatas::FRenderData& HairCardsRenderData, const FIntPoint& AtlasResolution, const FHairCardsVoxel& InVoxel, const FName& OwnerName);
-
-	/* Init/release buffers */
-	virtual void InternalAllocate(FRDGBuilder& GraphBuilder) override;
-	virtual void InternalRelease() override;
-
-	/* Get the resource name */
-	virtual FString GetFriendlyName() const override { return TEXT("FHairCardsResource"); }
-
-	/* Return the memory size for GPU resources */
-	uint32 GetResourcesSize() const
-	{
-		return 0;
-	}
-
-	/* Strand hair rest position buffer */		
-	uint32 CardBoundCount;
-	FIntPoint AtlasResolution;
-
-	FRDGExternalBuffer AtlasRectBuffer;
-	FRDGExternalBuffer LengthBuffer;
-	FRDGExternalBuffer CardItToClusterBuffer;
-	FRDGExternalBuffer ClusterIdToVerticesBuffer;
-	FRDGExternalBuffer ClusterBoundBuffer;
-	FRDGExternalBuffer CardsStrandsPositions;
-	FRDGExternalBuffer CardsStrandsAttributes;
-
-	FHairCardsVoxel CardVoxel;
-
-	/* Position offset as the rest positions are expressed in relative coordinate (16bits) */
-	//FVector PositionOffset = FVector::ZeroVector;
-
-	/* Reference to the hair strands render data */
-	const FHairCardsProceduralDatas::FRenderData& RenderData;
-};
-
 struct FHairCardsDeformedResource : public FHairCommonResource
 {
 	/** Build the hair strands resource */
diff --git a/Engine/Plugins/Runtime/HairStrands/Source/HairStrandsCore/Public/HairCardsBuilder.h b/Engine/Plugins/Runtime/HairStrands/Source/HairStrandsCore/Public/HairCardsBuilder.h
--- a/Engine/Plugins/Runtime/HairStrands/Source/HairStrandsCore/Public/HairCardsBuilder.h
+++ b/Engine/Plugins/Runtime/HairStrands/Source/HairStrandsCore/Public/HairCardsBuilder.h
@@ -20,11 +20,8 @@
 struct FHairCardsBulkData;
 struct FHairCardsDatas;
 struct FHairCardsInterpolationBulkData;
-struct FHairCardsProceduralDatas;
-struct FHairCardsProceduralResource;
 struct FHairCardsRestResource;
 struct FHairGroupCardsTextures;
-struct FHairGroupsProceduralCards;
 struct FHairMeshesBulkData;
 struct FHairStrandsDatas;
 struct FHairStrandsVoxelData;
@@ -48,27 +45,6 @@ namespace FHairCardsBuilder
 		const FHairStrandsDatas& InStrandsData,
 		FHairCardsDatas& Out);
 
-	void ExportGeometry(
-		const FHairCardsDatas& InCardsData, 
-		UStaticMesh* OutStaticMesh);
-
-	void BuildGeometry(
-		const FString& LODName,
-		const FHairStrandsDatas& InRen,
-		const FHairStrandsDatas& InSim,
-		const FHairGroupsProceduralCards& Settings,
-		FHairCardsProceduralDatas& Out,
-		FHairCardsBulkData& OutBulk,
-		FHairStrandsDatas& OutGuides,
-		FHairCardsInterpolationBulkData& OutInterpolationBulk,
-		FHairGroupCardsTextures& OutTextures);
-
-	void BuildTextureAtlas(
-		FHairCardsProceduralDatas* ProceduralData,
-		FHairCardsRestResource* RestResource,
-		FHairCardsProceduralResource* AtlasResource,
-		FHairGroupCardsTextures* Textures);
-
 	FString GetVersion();
 }
 
diff --git a/Engine/Plugins/Runtime/HairStrands/Source/HairStrandsCore/Public/HairCardsDatas.h b/Engine/Plugins/Runtime/HairStrands/Source/HairStrandsCore/Public/HairCardsDatas.h
--- a/Engine/Plugins/Runtime/HairStrands/Source/HairStrandsCore/Public/HairCardsDatas.h
+++ b/Engine/Plugins/Runtime/HairStrands/Source/HairStrandsCore/Public/HairCardsDatas.h
@@ -20,15 +20,6 @@ struct FHairCardsPositionFormat
 	static const EPixelFormat Format = PF_A32B32G32R32F;
 };
 
-struct FHairCardsStrandsAttributeFormat
-{
-	typedef FVector4f Type;
-	static const uint32 ComponentCount = 1;
-	static const uint32 SizeInByte = sizeof(Type);
-	static const EVertexElementType VertexElementType = VET_Float4;
-	static const EPixelFormat Format = PF_A32B32G32R32F;
-};
-
 struct FHairCardsUVFormat
 {
 	// Store atlas UV and (approximated) root UV
@@ -68,69 +59,6 @@ struct FHairCardsIndexFormat
 	static const EPixelFormat Format = PF_R32_UINT;
 };
 
-struct FHairCardsAtlasRectFormat
-{
-	struct Type { uint16_t X; uint16_t Y; uint16_t Z; uint16_t W; };
-	static const uint32 ComponentCount = 1;
-	static const uint32 SizeInByte = sizeof(Type);
-	static const EVertexElementType VertexElementType = VET_UShort4;
-	static const EPixelFormat Format = PF_R16G16B16A16_UINT;
-};
-
-struct FHairCardsDimensionFormat
-{
-	typedef float Type;
-	static const uint32 ComponentCount = 1;
-	static const uint32 SizeInByte = sizeof(Type);
-	static const EVertexElementType VertexElementType = VET_Float1;
-	static const EPixelFormat Format = PF_R32_FLOAT;
-};
-
-struct FHairCardsStrandsPositionFormat
-{
-	typedef FVector4f Type;
-	static const uint32 ComponentCount = 1;
-	static const uint32 SizeInByte = sizeof(Type);
-	static const EVertexElementType VertexElementType = VET_Float4;
-	static const EPixelFormat Format = PF_A32B32G32R32F;
-}; 
-
-struct FHairCardsOffsetAndCount
-{
-	typedef FUintPoint Type;
-	static const uint32 ComponentCount = 1;
-	static const uint32 SizeInByte = sizeof(Type);
-	static const EVertexElementType VertexElementType = VET_None;
-	static const EPixelFormat Format = PF_R32G32_UINT;
-};
-
-struct FHairCardsBoundsFormat
-{
-	typedef FVector4f Type;
-	static const uint32 ComponentCount = 4;
-	static const uint32 SizeInByte = sizeof(Type);
-	static const EVertexElementType VertexElementType = VET_Float4;
-	static const EPixelFormat Format = PF_A32B32G32R32F;
-};
-
-struct FHairCardsVoxelDensityFormat
-{
-	typedef uint32 Type;
-	static const uint32 ComponentCount = 1;
-	static const uint32 SizeInByte = sizeof(Type);
-	static const EVertexElementType VertexElementType = VET_UInt;
-	static const EPixelFormat Format = PF_R32_UINT;
-};
-
-struct FHairCardsVoxelTangentFormat
-{
-	typedef FVector4f Type;
-	static const uint32 ComponentCount = 1;
-	static const uint32 SizeInByte = sizeof(Type);
-	static const EVertexElementType VertexElementType = VET_Float4;
-	static const EPixelFormat Format = PF_A32B32G32R32F;
-};
-
 struct FHairCardsInterpolationVertex
 {
 	uint32 VertexIndex : 24;
@@ -148,30 +76,6 @@ struct FHairCardsInterpolationFormat
 	static const EPixelFormat Format = PF_R32_UINT;
 };
 
-struct FHairOrientedBound
-{
-	FVector3f Center;
-	FVector3f ExtentX;
-	FVector3f ExtentY;
-	FVector3f ExtentZ;
-
-	FVector3f& Extent(uint8 Axis)
-	{
-		if (Axis == 0) return ExtentX;
-		if (Axis == 1) return ExtentY;
-		if (Axis == 2) return ExtentZ;
-		return ExtentX;
-	}
-
-	const FVector3f& Extent(uint8 Axis) const
-	{
-		if (Axis == 0) return ExtentX;
-		if (Axis == 1) return ExtentY;
-		if (Axis == 2) return ExtentZ;
-		return ExtentX;
-	}
-};
-
 ////////////////////////////////////////////////////////////////////////////////////////////////
 // Hair Cards Data/Bulk
 
@@ -287,124 +191,6 @@ struct FHairCardsBulkData
 	void Serialize(FArchive& Ar);
 };
 
-////////////////////////////////////////////////////////////////////////////////////////////////
-// Hair Cards Data (Procedural generation)
-
-// Internal data structure for computing surface flow. 
-// This structured is exposed only for debug purpose
-struct FHairCardsVoxel
-{
-	FRDGExternalBuffer TangentBuffer;
-	FRDGExternalBuffer NormalBuffer;
-	FRDGExternalBuffer DensityBuffer;
-	FIntVector	Resolution;
-	FVector3f	MinBound;
-	FVector3f	MaxBound;
-	float		VoxelSize;
-};
-
-struct FHairCardsProceduralGeometry : FHairCardsGeometry
-{
-	struct Rect
-	{
-		FIntPoint Offset = FIntPoint(0, 0);
-		FIntPoint Resolution = FIntPoint(0, 0);
-	};
-
-	TArray<uint32> CardIndices; // Store the cards index whose a verex belongs to
-
-	TArray<Rect>   Rects;
-	TArray<float>  Lengths;
-	
-	TArray<FUintPoint> CardIndexToClusterOffsetAndCount;
-	TArray<FUintPoint> ClusterIndexToVertexOffsetAndCount;
-
-	TArray<FHairOrientedBound> Bounds;
-
-	void SetNum(uint32 Count)
-	{
-		FHairCardsGeometry::SetNum(Count);
-		CardIndices.Empty();
-
-		Rects.SetNum(Count);
-		Lengths.SetNum(Count);
-		Bounds.SetNum(Count);
-
-		// Cluster infos (editor only, for texture generation)
-		CardIndexToClusterOffsetAndCount.Empty();
-		ClusterIndexToVertexOffsetAndCount.Empty();
-	}
-
-	uint32 GetNum() const
-	{
-		return Bounds.Num();
-	}
-};
-
-struct FHairCardsProceduralAtlas
-{
-	struct Rect
-	{
-		FIntPoint Offset = FIntPoint(0, 0);
-		FIntPoint Resolution = FIntPoint(0, 0);
-		uint32 VertexOffset = 0;
-		uint32 VertexCount = 0;
-		FVector3f MinBound = FVector3f::ZeroVector;
-		FVector3f MaxBound = FVector3f::ZeroVector;
-
-		FVector3f RasterAxisX = FVector3f::ZeroVector;
-		FVector3f RasterAxisY = FVector3f::ZeroVector;
-		FVector3f RasterAxisZ = FVector3f::ZeroVector;
-
-		float CardWidth = 0;
-		float CardLength = 0;
-	};
-
-	FIntPoint Resolution;
-	TArray<Rect> Rects;
-	TArray<FVector4f> StrandsPositions;
-	TArray<FVector4f> StrandsAttributes;
-	bool bIsDirty = true;
-};
-
-
-// Data structure holding cards geometry information and 
-// intermediate data used for generating cards data based 
-// on strands groom
-struct FHairCardsProceduralDatas
-{
-	FHairStrandsDatas Guides;
-	FHairCardsProceduralGeometry Cards;
-	FHairCardsProceduralAtlas Atlas;
-	FHairCardsVoxel Voxels;
-
-	struct FRenderData
-	{
-		TArray<FHairCardsPositionFormat::Type> Positions;
-		TArray<FHairCardsNormalFormat::Type> Normals;
-		TArray<FHairCardsUVFormat::Type> UVs;
-		TArray<FHairCardsIndexFormat::Type> Indices;
-		TArray<FHairCardsAtlasRectFormat::Type> CardsRect;
-		TArray<FHairCardsDimensionFormat::Type> CardsLengths;
-		TArray<FHairCardsStrandsPositionFormat::Type> CardsStrandsPositions;
-		TArray<FHairCardsStrandsAttributeFormat::Type> CardsStrandsAttributes;
-		
-		TArray<FHairCardsOffsetAndCount::Type> CardItToCluster;		// Offset & Count
-		TArray<FHairCardsOffsetAndCount::Type> ClusterIdToVertices; // Offset & count
-		TArray<FHairCardsBoundsFormat::Type> ClusterBounds;
-
-		TArray<FHairCardsVoxelDensityFormat::Type> VoxelDensity;
-		TArray<FHairCardsVoxelTangentFormat::Type> VoxelTangent;
-		TArray<FHairCardsVoxelTangentFormat::Type> VoxelNormal;
-	} RenderData;
-};
-
-struct HAIRSTRANDSCORE_API FHairCardsSourceData
-{
-	FHairCardsDatas				ImportedData;
-	FHairCardsProceduralDatas	ProceduralData;
-};
-
 ////////////////////////////////////////////////////////////////////////////////////////////////
 // Hair Card interpolation Data/Bulk
 
diff --git a/Engine/Plugins/Runtime/HairStrands/Source/HairStrandsEditor/Private/GroomAssetDetails.cpp b/Engine/Plugins/Runtime/HairStrands/Source/HairStrandsEditor/Private/GroomAssetDetails.cpp
--- a/Engine/Plugins/Runtime/HairStrands/Source/HairStrandsEditor/Private/GroomAssetDetails.cpp
+++ b/Engine/Plugins/Runtime/HairStrands/Source/HairStrandsEditor/Private/GroomAssetDetails.cpp
@@ -68,47 +68,6 @@
 
 DEFINE_LOG_CATEGORY_STATIC(LogGroomAssetDetails, Log, All);
 
-static int32 GHairCardsProcerudalResolution = 4096;
-static int32 GHairCardsProcerudalResolution_LOD0 = -1;
-static int32 GHairCardsProcerudalResolution_LOD1 = -1;
-static int32 GHairCardsProcerudalResolution_LOD2 = -1;
-static int32 GHairCardsProcerudalResolution_LOD3 = -1;
-static int32 GHairCardsProcerudalResolution_LOD4 = -1;
-static int32 GHairCardsProcerudalResolution_LOD5 = -1;
-static int32 GHairCardsProcerudalResolution_LOD6 = -1;
-static int32 GHairCardsProcerudalResolution_LOD7 = -1;
-static FAutoConsoleVariableRef CVarHairCardsProcerudalResolution(TEXT("r.HairStrands.CardsAtlas.DefaultResolution"), GHairCardsProcerudalResolution, TEXT("Default cards atlas resolution."));
-static FAutoConsoleVariableRef CVarHairCardsProcerudalResolution_LOD0(TEXT("r.HairStrands.CardsAtlas.DefaultResolution.LOD0"), GHairCardsProcerudalResolution_LOD0, TEXT("Default cards atlas resolution for LOD0."));
-static FAutoConsoleVariableRef CVarHairCardsProcerudalResolution_LOD1(TEXT("r.HairStrands.CardsAtlas.DefaultResolution.LOD1"), GHairCardsProcerudalResolution_LOD1, TEXT("Default cards atlas resolution for LOD1."));
-static FAutoConsoleVariableRef CVarHairCardsProcerudalResolution_LOD2(TEXT("r.HairStrands.CardsAtlas.DefaultResolution.LOD2"), GHairCardsProcerudalResolution_LOD2, TEXT("Default cards atlas resolution for LOD2."));
-static FAutoConsoleVariableRef CVarHairCardsProcerudalResolution_LOD3(TEXT("r.HairStrands.CardsAtlas.DefaultResolution.LOD3"), GHairCardsProcerudalResolution_LOD3, TEXT("Default cards atlas resolution for LOD3."));
-static FAutoConsoleVariableRef CVarHairCardsProcerudalResolution_LOD4(TEXT("r.HairStrands.CardsAtlas.DefaultResolution.LOD4"), GHairCardsProcerudalResolution_LOD4, TEXT("Default cards atlas resolution for LOD4."));
-static FAutoConsoleVariableRef CVarHairCardsProcerudalResolution_LOD5(TEXT("r.HairStrands.CardsAtlas.DefaultResolution.LOD5"), GHairCardsProcerudalResolution_LOD5, TEXT("Default cards atlas resolution for LOD5."));
-static FAutoConsoleVariableRef CVarHairCardsProcerudalResolution_LOD6(TEXT("r.HairStrands.CardsAtlas.DefaultResolution.LOD6"), GHairCardsProcerudalResolution_LOD6, TEXT("Default cards atlas resolution for LOD6."));
-static FAutoConsoleVariableRef CVarHairCardsProcerudalResolution_LOD7(TEXT("r.HairStrands.CardsAtlas.DefaultResolution.LOD7"), GHairCardsProcerudalResolution_LOD7, TEXT("Default cards atlas resolution for LOD7."));
-
-static uint32 GetHairCardsAtlasResolution(int32 InLODIndex, int32 PrevResolution)
-{
-	uint32 OutResolution = GHairCardsProcerudalResolution;
-
-	const uint32 LODIndex = FMath::Clamp(InLODIndex, 0, 7);
-	switch (LODIndex)
-	{
-	case 0: OutResolution = GHairCardsProcerudalResolution_LOD0 >= 0 ? GHairCardsProcerudalResolution_LOD0 : uint32(GHairCardsProcerudalResolution); break;
-	case 1: OutResolution = GHairCardsProcerudalResolution_LOD1 >= 0 ? GHairCardsProcerudalResolution_LOD1 : uint32(PrevResolution * 0.5f); break;
-	case 2: OutResolution = GHairCardsProcerudalResolution_LOD2 >= 0 ? GHairCardsProcerudalResolution_LOD2 : uint32(PrevResolution * 0.5f); break;
-	case 3: OutResolution = GHairCardsProcerudalResolution_LOD3 >= 0 ? GHairCardsProcerudalResolution_LOD3 : uint32(PrevResolution * 0.5f); break;
-	case 4: OutResolution = GHairCardsProcerudalResolution_LOD4 >= 0 ? GHairCardsProcerudalResolution_LOD4 : uint32(PrevResolution * 0.5f); break;
-	case 5: OutResolution = GHairCardsProcerudalResolution_LOD5 >= 0 ? GHairCardsProcerudalResolution_LOD5 : uint32(PrevResolution * 0.5f); break;
-	case 6: OutResolution = GHairCardsProcerudalResolution_LOD6 >= 0 ? GHairCardsProcerudalResolution_LOD5 : uint32(PrevResolution * 0.5f); break;
-	case 7: OutResolution = GHairCardsProcerudalResolution_LOD6 >= 0 ? GHairCardsProcerudalResolution_LOD6 : uint32(PrevResolution * 0.5f); break;
-	}
-
-	const uint32 MinResolution = 128;
-	const uint32 MaxResolution = 16384;
-	return FMath::Clamp(OutResolution, MinResolution, MaxResolution);
-}
-
 FText GetHairAttributeLocText(EHairAttribute In, uint32 InFlags);
 
 ///////////////////////////////////////////////////////////////////////////////////////////////////////////
@@ -645,18 +604,12 @@ FReply FGroomRenderingDetails::OnAddGroup(FProperty* Property)
 			const FHairGroupsCardsSourceDescription& Prev = GroomAsset->GetHairGroupsCards()[LODCount - 2];
 			FHairGroupsCardsSourceDescription& Current = GroomAsset->GetHairGroupsCards()[LODCount - 1];
 
-			Current.SourceType = Prev.SourceType;
 			Current.GroupIndex = Prev.GroupIndex;
 			Current.LODIndex = FMath::Min(Prev.LODIndex + 1, 7);
-
-			// Prefill the LOD setting with basic preset
-			Current.ProceduralSettings.TextureSettings.AtlasMaxResolution  = GetHairCardsAtlasResolution(Current.LODIndex, Prev.ProceduralSettings.TextureSettings.AtlasMaxResolution);
-			Current.ProceduralSettings.TextureSettings.PixelPerCentimeters = Prev.ProceduralSettings.TextureSettings.PixelPerCentimeters * 0.75f;
 		}
 		else
 		{
 			FHairGroupsCardsSourceDescription& Current = GroomAsset->GetHairGroupsCards()[LODCount - 1];
-			Current.ProceduralSettings.TextureSettings.AtlasMaxResolution  = GetHairCardsAtlasResolution(0,0);
 			Current.LODIndex = 0;
 		}
 
@@ -937,15 +890,6 @@ bool FGroomRenderingDetails::CommonResetToDefault(TSharedPtr<IPropertyHandle> Ch
 	// Cards
 	if (bIsCardDescIndexValid)
 	{
-		{
-			FHairGroupsCardsSourceDescription Default;
-			HAIR_RESET0(GetHairGroupsCards(), FHairGroupsCardsSourceDescription, SourceType);
-			HAIR_RESET0(GetHairGroupsCards(), FHairGroupsCardsSourceDescription, ProceduralMesh);
-			HAIR_RESET0(GetHairGroupsCards(), FHairGroupsCardsSourceDescription, ImportedMesh);
-			HAIR_RESET0(GetHairGroupsCards(), FHairGroupsCardsSourceDescription, GroupIndex);
-			HAIR_RESET0(GetHairGroupsCards(), FHairGroupsCardsSourceDescription, LODIndex);
-		}
-
 		{
 			FHairGroupCardsTextures Default;
 			HAIR_RESET1(GetHairGroupsCards(), FHairGroupCardsTextures, Textures, DepthTexture);
@@ -1223,29 +1167,6 @@ FReply FGroomRenderingDetails::OnAddLODClicked(int32 GroupIndex, FProperty* Prop
 	return FReply::Handled();
 }
 
-// Hair_TODO: rename into OnReloadCards
-FReply FGroomRenderingDetails::OnRefreshCards(int32 DescIndex, FProperty* Property)
-{
-	if (DescIndex < GroomAsset->GetHairGroupsCards().Num() && GroomAsset->GetHairGroupsCards()[DescIndex].SourceType == EHairCardsSourceType::Procedural)
-	{
-		FScopedTransaction Transaction(FText::FromString(TEXT("RefreshCards")));
-		
-		FPropertyChangedEvent PropertyChangedEvent(Property);
-		GroomAsset->PostEditChangeProperty(PropertyChangedEvent);
-	}
-	return FReply::Handled();
-}
-
-// Hair_TODO: rename into OnGenerageCards
-FReply FGroomRenderingDetails::OnSaveCards(int32 DescIndex, FProperty* Property)
-{
-	if (DescIndex < GroomAsset->GetHairGroupsCards().Num())
-	{
-		GroomAsset->SaveProceduralCards(DescIndex);
-	}
-	return FReply::Handled();
-}
-
 FReply FGroomRenderingDetails::OnGenerateCardDataUsingPlugin(int32 GroupIndex)
 {
 	if (GroomAsset && GroomAsset->GetHairGroupsCards().IsValidIndex(GroupIndex))
@@ -1589,77 +1510,24 @@ void FGroomRenderingDetails::OnGenerateElementForHairGroup(TSharedRef<IPropertyH
 			else
 			{
 				IDetailPropertyRow& PropertyRow = AddPropertyWithCustomReset(ChildHandle, ChildrenBuilder, GroupIndex, -1);
-
-				auto CustomizeMeshPropertyRow = [](IDetailPropertyRow& PropertyRow, bool bEnable)->FDetailWidgetRow&
+				
+				if (PropertyName == GET_MEMBER_NAME_CHECKED(FHairGroupsCardsSourceDescription, ImportedMesh))
 				{
 					TSharedPtr<SWidget> NameWidget;
 					TSharedPtr<SWidget> ValueWidget;
 					FDetailWidgetRow Row;
 					PropertyRow.GetDefaultWidgets(NameWidget, ValueWidget, Row);
 
-					ValueWidget->SetEnabled(bEnable);
-
-					return PropertyRow.CustomWidget()
-						.NameContent()
-						[
-							PropertyRow.GetPropertyHandle()->CreatePropertyNameWidget(
-								LOCTEXT("HairCardsMeshProperty", "Mesh"),
-								LOCTEXT("HairCardsMeshTooltop",  "")
-							)
-						]
-						.ValueContent()
-						[
-							ValueWidget.ToSharedRef()
-						];
-				};
-
-				TWeakObjectPtr<UGroomAsset> GroomAssetPtr = GroomAsset;
-				auto ShouldShowProceduralProperties = [GroomAssetPtr, GroupIndex]()
-				{
-					return GroomAssetPtr.IsValid() && GroomAssetPtr->GetHairGroupsCards().IsValidIndex(GroupIndex) &&
-						GroomAssetPtr->GetHairGroupsCards()[GroupIndex].SourceType == EHairCardsSourceType::Procedural;
-				};
-				TAttribute<EVisibility> ProceduralPropertyVisibility = TAttribute<EVisibility>::CreateLambda([ShouldShowProceduralProperties]()
-					{
-						return ShouldShowProceduralProperties() ? EVisibility::Visible : EVisibility::Collapsed;
-					}
-				);
-				
-				if (PropertyName == GET_MEMBER_NAME_CHECKED(FHairGroupsCardsSourceDescription, ProceduralMesh))
-				{
-					PropertyRow.Visibility(ProceduralPropertyVisibility);
-					CustomizeMeshPropertyRow(PropertyRow, /*bEnable =*/false);
-				}
-				else if (PropertyName == GET_MEMBER_NAME_CHECKED(FHairGroupsCardsSourceDescription, ProceduralSettings))
-				{
-					PropertyRow.Visibility(TAttribute<EVisibility>::CreateLambda([ShouldShowProceduralProperties]()
-						{
-							TArray<IHairCardGenerator*> HairCardGeneratorPlugins = IModularFeatures::Get().GetModularFeatureImplementations<IHairCardGenerator>(IHairCardGenerator::ModularFeatureName);
-							return HairCardGeneratorPlugins.IsEmpty() && ShouldShowProceduralProperties() ? EVisibility::Visible : EVisibility::Collapsed;
-						}
-					));
-				}
-				else if (PropertyName == GET_MEMBER_NAME_CHECKED(FHairGroupsCardsSourceDescription, ImportedMesh))
-				{
-					TAttribute<EVisibility> NonProceduralPropertyVisibility = TAttribute<EVisibility>::CreateLambda([ShouldShowProceduralProperties]()
+					TArray<IHairCardGenerator*> HairCardGeneratorPlugins = IModularFeatures::Get().GetModularFeatureImplementations<IHairCardGenerator>(IHairCardGenerator::ModularFeatureName);
+					bool bHasProceduralGenerationPlugin = HairCardGeneratorPlugins.Num() > 0;
+					TAttribute<EVisibility> ProceduralPropertyVisibility = TAttribute<EVisibility>::CreateLambda([bHasProceduralGenerationPlugin]()
 						{
-							return ShouldShowProceduralProperties() ? EVisibility::Collapsed : EVisibility::Visible;
+							return bHasProceduralGenerationPlugin ? EVisibility::Visible : EVisibility::Collapsed;
 						}
 					);
-					PropertyRow.Visibility(NonProceduralPropertyVisibility);
-
-					CustomizeMeshPropertyRow(PropertyRow, /*bEnable =*/true);
-				}
-				else if (PropertyName == GET_MEMBER_NAME_CHECKED(FHairGroupsCardsSourceDescription, SourceType))
-				{
-					TSharedPtr<SWidget> NameWidget;
-					TSharedPtr<SWidget> ValueWidget;
-					FDetailWidgetRow Row;
-					PropertyRow.GetDefaultWidgets(NameWidget, ValueWidget, Row);
 
 					TSharedRef<SHorizontalBox> ProceduralGenButtons = SNew(SHorizontalBox).Visibility(ProceduralPropertyVisibility);
-					TArray<IHairCardGenerator*> HairCardGeneratorPlugins = IModularFeatures::Get().GetModularFeatureImplementations<IHairCardGenerator>(IHairCardGenerator::ModularFeatureName);
-					if (HairCardGeneratorPlugins.Num() > 0)
+					if (bHasProceduralGenerationPlugin)
 					{
 						FText RegenToolTipText = LOCTEXT("GenerateNewCardsTooltip", "Generate new card assets (meshes and textures) using the current procedural settings. NOTE: This will overwrite preexisting card assets for this LOD.");
 						ProceduralGenButtons->AddSlot()
@@ -1680,41 +1548,6 @@ void FGroomRenderingDetails::OnGenerateElementForHairGroup(TSharedRef<IPropertyH
 									]
 							];
 					}
-					else
-					{
-						FText ToolTipTextForGeneration(FText::FromString(TEXT("Generate procedural cards data (meshes and textures) based on current procedural settings. Cards generation needs to run prior to the (re)loading of the cards data.")));
-						FText ToolTipTextForReloading(FText::FromString(TEXT("(Re)Load generated cards data (meshes and textures) into the groom asset. The data need to be generated with the save/generated button prior to reloading.")));
-
-						ProceduralGenButtons->AddSlot()
-							.AutoWidth()
-							[
-								SNew(SButton)
-								.VAlign(VAlign_Center)
-								.HAlign(HAlign_Center)
-								.ButtonStyle(FAppStyle::Get(), "HoverHintOnly")
-								.ToolTipText(ToolTipTextForGeneration)
-								.OnClicked(this, &FGroomRenderingDetails::OnSaveCards, GroupIndex, Property)
-								[
-									SNew(SImage)
-									.Image(FAppStyle::GetBrush("AssetEditor.SaveAsset"))
-								]
-							];
-
-						ProceduralGenButtons->AddSlot()
-							.AutoWidth()
-							[
-								SNew(SButton)
-								.VAlign(VAlign_Center)
-								.HAlign(HAlign_Center)
-								.ButtonStyle(FAppStyle::Get(), "HoverHintOnly")
-								.ToolTipText(ToolTipTextForReloading)
-								.OnClicked(this, &FGroomRenderingDetails::OnRefreshCards, GroupIndex, Property)
-								[
-									SNew(SImage)
-									.Image(FAppStyle::GetBrush("Icons.Refresh"))
-								]
-							];
-					}
 
 					PropertyRow.CustomWidget()
 						.NameContent()
