Change Hair card generator icond.

diff --git a/Engine/Plugins/Runtime/HairStrands/Source/HairStrandsEditor/Private/GroomAssetDetails.cpp b/Engine/Plugins/Runtime/HairStrands/Source/HairStrandsEditor/Private/GroomAssetDetails.cpp
--- a/Engine/Plugins/Runtime/HairStrands/Source/HairStrandsEditor/Private/GroomAssetDetails.cpp
+++ b/Engine/Plugins/Runtime/HairStrands/Source/HairStrandsEditor/Private/GroomAssetDetails.cpp
@@ -63,6 +63,7 @@
 #include "IPropertyUtilities.h"
 
 #include "Styling/AppStyle.h"
+#include "GroomEditorStyle.h"
 
 #define LOCTEXT_NAMESPACE "GroomRenderingDetails"
 
@@ -1526,6 +1527,16 @@ void FGroomRenderingDetails::OnGenerateElementForHairGroup(TSharedRef<IPropertyH
 						}
 					);
 
+					const FSlateBrush* Brush = nullptr;
+					if (const ISlateStyle* AppStyle = Toolkit->GetSlateStyle())
+					{
+						Brush = AppStyle->GetBrush("GroomEditor.GroomCardGenerator");
+					}
+					else
+					{
+						Brush = FAppStyle::GetBrush("ContentBrowser.AssetActions.ReimportAsset");
+					}
+
 					TSharedRef<SHorizontalBox> ProceduralGenButtons = SNew(SHorizontalBox).Visibility(ProceduralPropertyVisibility);
 					if (bHasProceduralGenerationPlugin)
 					{
@@ -1541,10 +1552,7 @@ void FGroomRenderingDetails::OnGenerateElementForHairGroup(TSharedRef<IPropertyH
 									.OnClicked(this, &FGroomRenderingDetails::OnGenerateCardDataUsingPlugin, GroupIndex)
 									[
 										SNew(SImage)
-											// @TODO: Need a specialized icon for this?
-											.Image(FAppStyle::GetBrush("ContentBrowser.AssetActions.ReimportAsset"))
-											.RenderTransform(FSlateRenderTransform(FQuat2D(FMath::DegreesToRadians(90.0f))))
-											.RenderTransformPivot(FVector2D(0.5f, 0.5f))
+											.Image(Brush)
 									]
 							];
 					}
diff --git a/Engine/Plugins/Runtime/HairStrands/Source/HairStrandsEditor/Private/GroomCustomAssetEditorToolkit.cpp b/Engine/Plugins/Runtime/HairStrands/Source/HairStrandsEditor/Private/GroomCustomAssetEditorToolkit.cpp
--- a/Engine/Plugins/Runtime/HairStrands/Source/HairStrandsEditor/Private/GroomCustomAssetEditorToolkit.cpp
+++ b/Engine/Plugins/Runtime/HairStrands/Source/HairStrandsEditor/Private/GroomCustomAssetEditorToolkit.cpp
@@ -430,6 +430,7 @@ void FGroomCustomAssetEditorToolkit::InitCustomAssetEditor(const EToolkitMode::T
 
 	ViewportTab = SNew(SGroomEditorViewport);
 	ThumbnailPool = MakeShared<FAssetThumbnailPool>(64);
+	GroomEditorStyle = MakeShareable(new FGroomEditorStyle());
 
 	// Automatically affect the first skelal mesh compatible with the groom asset
 	#if 0
@@ -984,4 +985,9 @@ int32 FGroomCustomAssetEditorToolkit::GetActiveBindingIndex() const
 	return ActiveGroomBindingIndex;
 }
 
+FGroomEditorStyle* FGroomCustomAssetEditorToolkit::GetSlateStyle() const 
+{ 
+	return GroomEditorStyle.Get();
+}
+
 #undef LOCTEXT_NAMESPACE
diff --git a/Engine/Plugins/Runtime/HairStrands/Source/HairStrandsEditor/Public/GroomCustomAssetEditorToolkit.h b/Engine/Plugins/Runtime/HairStrands/Source/HairStrandsEditor/Public/GroomCustomAssetEditorToolkit.h
--- a/Engine/Plugins/Runtime/HairStrands/Source/HairStrandsEditor/Public/GroomCustomAssetEditorToolkit.h
+++ b/Engine/Plugins/Runtime/HairStrands/Source/HairStrandsEditor/Public/GroomCustomAssetEditorToolkit.h
@@ -9,6 +9,7 @@
 #include "Toolkits/AssetEditorToolkit.h"
 #include "Widgets/SGroomEditorViewport.h"
 #include "HairStrandsInterface.h"
+#include "GroomEditorStyle.h"
 
 class FAssetThumbnailPool;
 class UGroomAsset; 
@@ -37,6 +38,8 @@ public:
 	virtual void PreviewBinding(int32 BindingIndex) = 0;
 
 	virtual int32 GetActiveBindingIndex() const =0;
+
+	virtual FGroomEditorStyle* GetSlateStyle() const { return nullptr; };
 };
 
 class HAIRSTRANDSEDITOR_API FGroomCustomAssetEditorToolkit : public IGroomCustomAssetEditorToolkit
@@ -83,6 +86,9 @@ public:
 	/** Return the index of the active binding. */
 	virtual int32 GetActiveBindingIndex() const override;
 
+	/** Return groom editor style. */
+	FGroomEditorStyle* GetSlateStyle() const override;
+
 private:
 
 	// called when The Play simulation button is pressed
@@ -183,5 +189,7 @@ private:
 	TWeakObjectPtr<UAnimationAsset> PreviewSkeletalAnimationAsset;
 
 	TSharedPtr<FAssetThumbnailPool> ThumbnailPool;
+
+	TSharedPtr<FGroomEditorStyle> GroomEditorStyle;
 };
 
diff --git a/Engine/Plugins/Runtime/HairStrands/Source/HairStrandsEditor/Public/GroomEditorStyle.h b/Engine/Plugins/Runtime/HairStrands/Source/HairStrandsEditor/Public/GroomEditorStyle.h
--- a/Engine/Plugins/Runtime/HairStrands/Source/HairStrandsEditor/Public/GroomEditorStyle.h
+++ b/Engine/Plugins/Runtime/HairStrands/Source/HairStrandsEditor/Public/GroomEditorStyle.h
@@ -33,7 +33,7 @@ public:
 //		Set("HairToolsManagerCommands.BeginAddPrimitiveTool", new IMAGE_BRUSH("HairSpray128", Icon40x40));
 		Set("GroomEditor.SimulationOptions", new IMAGE_BRUSH("S_SimulationOptions_40x", Icon40x40));
 		Set("GroomEditor.SimulationOptions.Small", new IMAGE_BRUSH("S_SimulationOptions_40x", Icon20x20));
-
+		Set("GroomEditor.GroomCardGenerator", new IMAGE_BRUSH("S_GroomCardGenerator_128", Icon16x16));
 		FSlateStyleRegistry::RegisterSlateStyle(*this);
 	}
 
