Port back changes from MHT.

diff --git a/Engine/Plugins/Runtime/HairStrands/Source/HairStrandsCore/Private/GroomBuilder.cpp b/Engine/Plugins/Runtime/HairStrands/Source/HairStrandsCore/Private/GroomBuilder.cpp
--- a/Engine/Plugins/Runtime/HairStrands/Source/HairStrandsCore/Private/GroomBuilder.cpp
+++ b/Engine/Plugins/Runtime/HairStrands/Source/HairStrandsCore/Private/GroomBuilder.cpp
@@ -2321,7 +2321,28 @@ void Decimate(
 	}
 	HairStrandsBuilder::BuildInternalData(OutData);
 }
-	
+
+void GetCurveShuffleIndices(TArray<uint32>& CurveRandomizedIndex, uint32 InCurveCount)
+{
+	FRandomStream Random;
+	Random.Initialize(0xdeedbeed);
+
+	// this is the proof of concept version - which just does a shuffle - TODO: improved voxel based version which caters for both long hair and short hair
+	CurveRandomizedIndex.SetNum(InCurveCount);
+
+	//initialize value
+	for (uint32 CurveIndex = 0; CurveIndex < InCurveCount; CurveIndex++)
+	{
+		CurveRandomizedIndex[CurveIndex] = CurveIndex;
+	}
+	//shuffle
+	for (uint32 CurveIndex = 0; CurveIndex < InCurveCount; CurveIndex++)
+	{
+		uint32 RandIndex = Random.RandRange(0, InCurveCount - 1);
+		CurveRandomizedIndex.Swap(CurveIndex, RandIndex);
+	}
+}
+
 void Decimate(
 	const FHairStrandsDatas& InData, 
 	float CurveDecimationPercentage, 
@@ -2338,29 +2359,15 @@ void Decimate(
 	const uint32 OutCurveCount = FMath::Clamp(uint32(InCurveCount * CurveDecimationPercentage), 1u, InCurveCount);
 	const uint32 InAttributes = InData.GetAttributes();
 
-	FRandomStream Random;
-	Random.Initialize(0xdeedbeed);
-
 	TArray<uint32> CurveRandomizedIndex;
-
-	// this is the proof of concept version - which just does a shuffle - TODO: improved voxel based version which caters for both long hair and short hair
 	if (bContinuousDecimationReordering)
 	{
-		CurveRandomizedIndex.SetNum(InCurveCount);
-
-		//initialize value
-		for (uint32 CurveIndex = 0; CurveIndex < InCurveCount; CurveIndex++)
-		{
-			CurveRandomizedIndex[CurveIndex] = CurveIndex;
-		}
-		//shuffle
-		for (uint32 CurveIndex = 0; CurveIndex < InCurveCount; CurveIndex++)
-		{
-			uint32 RandIndex = Random.RandRange(0, InCurveCount - 1);
-			CurveRandomizedIndex.Swap(CurveIndex, RandIndex);
-		}
+		GetCurveShuffleIndices(CurveRandomizedIndex, InCurveCount);
 	}
 
+	FRandomStream Random;
+	Random.Initialize(0xdeedbeed);
+
 	uint32 OutTotalPointCount = InData.GetNumPoints();
 	TArray<uint32> CurveIndices;
 	const bool bNeedDecimation = CurveDecimationPercentage < 1 || VertexDecimationPercentage < 1;
diff --git a/Engine/Plugins/Runtime/HairStrands/Source/HairStrandsCore/Private/GroomRBFDeformer.cpp b/Engine/Plugins/Runtime/HairStrands/Source/HairStrandsCore/Private/GroomRBFDeformer.cpp
--- a/Engine/Plugins/Runtime/HairStrands/Source/HairStrandsCore/Private/GroomRBFDeformer.cpp
+++ b/Engine/Plugins/Runtime/HairStrands/Source/HairStrandsCore/Private/GroomRBFDeformer.cpp
@@ -12,6 +12,13 @@
 #include "Engine/Texture.h"
 #include "GroomBindingBuilder.h"
 
+#define ALLOW_CURVE_SHUFFLING 0
+
+namespace FHairStrandsDecimation
+{
+	void GetCurveShuffleIndices(TArray<uint32>& CurveRandomizedIndex, uint32 InCurveCount);
+}
+
 // HairStrandsMesh.usf
 void InitMeshSamples(
 	uint32 MaxVertexCount,
@@ -105,6 +112,10 @@ void DeformStrands(
 		OutDeformedPositionBuffer[VertexIndex] = DisplacedPosition;
 	});
 
+	// Build curve shuffling indices
+	TArray<uint32> CurveShuffleIndex;
+	FHairStrandsDecimation::GetCurveShuffleIndices(CurveShuffleIndex, HairStandsData.GetNumCurves());
+
 	// Compute correction for snapping the strands back to the surface, as the RBF introduces low frequency offset
 	const uint32 CurveCount = HairStandsData.GetNumCurves();
 	TArray<FVector4f> CorrectionOffsets;
@@ -112,16 +123,33 @@ void DeformStrands(
 	ParallelFor(CurveCount, [&](uint32 CurveIndex)
 	{
 		const uint32 VertexOffset = HairStandsData.StrandsCurves.CurvesOffset[CurveIndex];
-		const uint32 RootIndex = PointToCurveBuffer[VertexOffset];
-		const uint32 TriangleIndex = RootToUniqueTriangleBuffer[RootIndex];
 
 		// Sanity check
+		const uint32 RootIndex = PointToCurveBuffer[VertexOffset];
 		check(RootIndex == CurveIndex);
 
+		// HairStandsData are *not* shuffled (for export purpose), while the root data have been built with shuffled data.
+		// Due to this we have to revert the shuffling prior to looking up binding data.
+	#if ALLOW_CURVE_SHUFFLING
+		const uint32 ShuffleRootIndex = RootIndex;
+	#else
+		uint32 ShuffleRootIndex = RootIndex;
+		for (uint32 It=0, Count=CurveShuffleIndex.Num();It<Count; ++It)
+		{
+			if (CurveShuffleIndex[It] == RootIndex)
+			{
+				ShuffleRootIndex = It;
+				break;
+			}
+		}
+	#endif
+		const uint32 TriangleIndex = RootToUniqueTriangleBuffer[ShuffleRootIndex];
+
+
 		const FVector3f& Rest_Position   = HairStandsData.StrandsPoints.PointsPosition[VertexOffset];
 		const FVector3f& Deform_Position = OutDeformedPositionBuffer[VertexOffset];
 
-		const uint32 PackedBarycentric = RootBarycentricBuffer[RootIndex];
+		const uint32 PackedBarycentric = RootBarycentricBuffer[ShuffleRootIndex];
 		const FVector2f B0 = FVector2f(FHairStrandsRootUtils::UnpackBarycentrics(PackedBarycentric));
 		const FVector3f   B  = FVector3f(B0.X, B0.Y, 1.f - B0.X - B0.Y);
 
@@ -441,8 +469,7 @@ static void ApplyDeformationToGroom(const TArray<FRBFDeformedPositions>& Deforme
 
 				FHairStrandsDatas StrandsData;
 				FHairStrandsDatas GuidesData;
-				FGroomBuilder::BuildData(HairGroup, GroomAsset->GetHairGroupsInterpolation()[GroupIndex], HairGroupsInfo, StrandsData, GuidesData);
-				//GroomAsset->GetHairStrandsDatas(GroupIndex, StrandsData, GuidesData);
+				FGroomBuilder::BuildData(HairGroup, GroomAsset->GetHairGroupsInterpolation()[GroupIndex], HairGroupsInfo, StrandsData, GuidesData, ALLOW_CURVE_SHUFFLING);
 
 				FGroomBuilder::BuildBulkData(HairGroup.Info, GuidesData,  HairGroupsData.Guides.BulkData);
 				FGroomBuilder::BuildBulkData(HairGroup.Info, StrandsData, HairGroupsData.Strands.BulkData);
@@ -455,7 +482,7 @@ static void ApplyDeformationToGroom(const TArray<FRBFDeformedPositions>& Deforme
 			}
 		}
 	}
-	
+
 	GroomAsset->CommitHairDescription(MoveTemp(HairDescription), UGroomAsset::EHairDescriptionType::Source);
 	GroomAsset->UpdateHairGroupsInfo();
 
@@ -468,7 +495,7 @@ static void ExtractSkeletalVertexPosition(
 	const FSkeletalMeshRenderData* SkeletalMeshData,
 	const uint32 MeshLODIndex,
 	TArray<FVector3f>& OutMeshVertexPositionsBuffer)
-{	
+{
 	const uint32 VertexCount = SkeletalMeshData->LODRenderData[MeshLODIndex].StaticVertexBuffers.PositionVertexBuffer.GetNumVertices();
 	OutMeshVertexPositionsBuffer.SetNum(VertexCount);
 	for (uint32 VertexIndex = 0; VertexIndex < VertexCount; ++VertexIndex)
@@ -478,7 +505,7 @@ static void ExtractSkeletalVertexPosition(
 }
 
 void DeformStaticMeshPositions(
-	UStaticMesh* OutMesh, 
+	UStaticMesh* OutMesh,
 	const TArray<FVector3f>& MeshVertexPositionsBuffer_Target,
 	const FHairStrandsRootData::FMeshProjectionLOD& RestLODData)
 {
@@ -627,9 +654,9 @@ void FGroomRBFDeformer::GetRBFDeformedGroomAsset(const UGroomAsset* InGroomAsset
 			}
 		};
 
-		// Note that the GroupID from the HairGroups cannot be used as the GroupIndex since 
+		// Note that the GroupID from the HairGroups cannot be used as the GroupIndex since
 		// the former may not be strictly increasing nor consecutive
-		// but the ordering of the groups does represent the GroupIndex		
+		// but the ordering of the groups does represent the GroupIndex
 		for (const FHairDescriptionGroup& Group : HairDescriptionGroups.HairGroups)
 		{
 			const uint32 GroupIndex = Group.Info.GroupID;
@@ -641,7 +668,7 @@ void FGroomRBFDeformer::GetRBFDeformedGroomAsset(const UGroomAsset* InGroomAsset
 			FHairGroupInfo DummyInfo;
 
 			// Disable explicitely curve reordering as we want to preserve vertex mapping to write deformed postions back to the hair description, after RBF deformation
-			FGroomBuilder::BuildData(Group, InGroomAsset->GetHairGroupsInterpolation()[GroupIndex], DummyInfo, StrandsData, GuidesData, false /*bAllowCurveReordering*/);
+			FGroomBuilder::BuildData(Group, InGroomAsset->GetHairGroupsInterpolation()[GroupIndex], DummyInfo, StrandsData, GuidesData, ALLOW_CURVE_SHUFFLING);
 
 			TArray<uint32> SimRootDataPointToCurveBuffer;
 			TArray<uint32> RenRootDataPointToCurveBuffer;
@@ -674,6 +701,20 @@ void FGroomRBFDeformer::GetRBFDeformedGroomAsset(const UGroomAsset* InGroomAsset
 				FHairStrandsRootData RenRootData;
 				FGroomBindingBuilder::GetRootData(RenRootData, BindingAsset->GetHairGroupsPlatformData()[GroupIndex].RenRootBulkData);
 
+				// Transfer RBF weights from SimRootData to RenRootData, as RenRootData does not hold RBF sample data
+				// This is transient data.
+				{
+					FHairStrandsRootData SimRootData;
+					FGroomBindingBuilder::GetRootData(SimRootData, BindingAsset->GetHairGroupsPlatformData()[GroupIndex].SimRootBulkData);
+
+					FHairStrandsRootData::FMeshProjectionLOD& SimLODData = SimRootData.MeshProjectionLODs[MeshLODIndex];
+					FHairStrandsRootData::FMeshProjectionLOD& RenLODData = RenRootData.MeshProjectionLODs[MeshLODIndex];
+					RenLODData.SampleCount 						= SimLODData.SampleCount;
+					RenLODData.MeshInterpolationWeightsBuffer 	= SimLODData.MeshInterpolationWeightsBuffer;
+					RenLODData.MeshSampleIndicesBuffer 			= SimLODData.MeshSampleIndicesBuffer;
+					RenLODData.RestSamplePositionsBuffer 		= SimLODData.RestSamplePositionsBuffer;
+				}
+
 				DeformedPositions[GroupIndex].RenderStrands = GetDeformedHairStrandsPositions(
 					MeshVertexPositionsBuffer_Target,
 					StrandsData,
@@ -778,10 +819,11 @@ void FGroomRBFDeformer::GetRBFDeformedGroomAsset(const UGroomAsset* InGroomAsset
 				{
 					Mesh->ConditionalPostLoad();
 	
-					FHairStrandsRootData RenRootData;
-					FGroomBindingBuilder::GetRootData(RenRootData, BindingAsset->GetHairGroupsPlatformData()[Desc.GroupIndex].RenRootBulkData);
+					// Load Sim root data, as they contains RBF weights
+					FHairStrandsRootData SimRootData;
+					FGroomBindingBuilder::GetRootData(SimRootData, BindingAsset->GetHairGroupsPlatformData()[Desc.GroupIndex].SimRootBulkData);
 	
-					DeformStaticMeshPositions(Mesh, MeshVertexPositionsBuffer_Target, RenRootData.MeshProjectionLODs[MeshLODIndex]);
+					DeformStaticMeshPositions(Mesh, MeshVertexPositionsBuffer_Target, SimRootData.MeshProjectionLODs[MeshLODIndex]);
 				}
 			}
 		} 
@@ -797,11 +839,12 @@ void FGroomRBFDeformer::GetRBFDeformedGroomAsset(const UGroomAsset* InGroomAsset
 				}
 
 				Mesh->ConditionalPostLoad();
-				
-				FHairStrandsRootData RenRootData;
-				FGroomBindingBuilder::GetRootData(RenRootData, BindingAsset->GetHairGroupsPlatformData()[Desc.GroupIndex].RenRootBulkData);
 
-				DeformStaticMeshPositions(Mesh, MeshVertexPositionsBuffer_Target, RenRootData.MeshProjectionLODs[MeshLODIndex]);
+				// Load Sim root data, as they contains RBF weights
+				FHairStrandsRootData SimRootData;
+				FGroomBindingBuilder::GetRootData(SimRootData, BindingAsset->GetHairGroupsPlatformData()[Desc.GroupIndex].SimRootBulkData);
+
+				DeformStaticMeshPositions(Mesh, MeshVertexPositionsBuffer_Target, SimRootData.MeshProjectionLODs[MeshLODIndex]);
 			}
 		}
 
