Fix groom binding error due to too tight bounding box searching.

diff --git a/Engine/Plugins/Runtime/HairStrands/Source/HairStrandsCore/Private/GroomBindingBuilder.cpp b/Engine/Plugins/Runtime/HairStrands/Source/HairStrandsCore/Private/GroomBindingBuilder.cpp
--- a/Engine/Plugins/Runtime/HairStrands/Source/HairStrandsCore/Private/GroomBindingBuilder.cpp
+++ b/Engine/Plugins/Runtime/HairStrands/Source/HairStrandsCore/Private/GroomBindingBuilder.cpp
@@ -1202,7 +1202,7 @@ namespace GroomBinding_RootProjection
 				// By nature, it is possible that coarser LOD have positions which ressemble only very coarsly to 
 				// LOD0. In this case we increase the hair bound to ensure that skel. mesh triangles will be intersect 
 				// the groom bound to be correctly inserted.
-				if (bHasTransferredPosition && ClosestTrianglePoint < FLT_MAX)
+				if (ClosestTrianglePoint < FLT_MAX)
 				{
 					GridMin -= FVector(ClosestTrianglePoint * 1.25f);
 					GridMax += FVector(ClosestTrianglePoint * 1.25f);
