Fix missing DDC value for groom binding data due to DDC eviction.

diff --git a/Engine/Plugins/Runtime/HairStrands/Source/HairStrandsCore/Private/GroomBindingAsset.cpp b/Engine/Plugins/Runtime/HairStrands/Source/HairStrandsCore/Private/GroomBindingAsset.cpp
--- a/Engine/Plugins/Runtime/HairStrands/Source/HairStrandsCore/Private/GroomBindingAsset.cpp
+++ b/Engine/Plugins/Runtime/HairStrands/Source/HairStrandsCore/Private/GroomBindingAsset.cpp
@@ -41,7 +41,7 @@ static void InternalSerializeGuides(FArchive& Ar, UObject* Owner, FHairStrandsRo
 	}
 }
 
-static void InternalSerializeStrands(FArchive& Ar, UObject* Owner, FHairStrandsRootBulkData& Data, uint32 Flags, bool bHeader, bool bData)
+static void InternalSerializeStrands(FArchive& Ar, UObject* Owner, FHairStrandsRootBulkData& Data, uint32 Flags, bool bHeader, bool bData, bool* bOutHasDataInCache = nullptr)
 {
 	Ar.UsingCustomVersion(FAnimObjectVersion::GUID);
 
@@ -73,10 +73,13 @@ static void InternalSerializeStrands(FArchive& Ar, UObject* Owner, FHairStrandsR
 		const bool bPreWarmCache = IsLoading() && bHeader && !bData;
 		if (bPreWarmCache)
 		{
+			bool bHasDataInCache = true;
 			for (int32 LODIndex = 0, LODCount = Data.GetLODCount(); LODIndex < LODCount; ++LODIndex)
 			{
-				FHairStreamingRequest R; R.WarmCache(HAIR_MAX_NUM_CURVE_PER_GROUP, HAIR_MAX_NUM_POINT_PER_GROUP, LODIndex, Data);
+				FHairStreamingRequest R; bHasDataInCache &= R.WarmCache(HAIR_MAX_NUM_CURVE_PER_GROUP, HAIR_MAX_NUM_POINT_PER_GROUP, LODIndex, Data);
 			}
+
+			if (bOutHasDataInCache) { *bOutHasDataInCache = bHasDataInCache; }
 		}
 		#endif
 	}
@@ -100,7 +103,7 @@ static void InternalSerializeCards(FArchive& Ar, UObject* Owner, TArray<FHairStr
 	}
 }
 
-static void InternalSerializePlatformData(FArchive& Ar, UObject* Owner, UGroomBindingAsset::FHairGroupPlatformData& GroupData, uint32 Flags, bool bHeader, bool bData)
+static void InternalSerializePlatformData(FArchive& Ar, UObject* Owner, UGroomBindingAsset::FHairGroupPlatformData& GroupData, uint32 Flags, bool bHeader, bool bData, bool* bOutHasDataInCache = nullptr)
 {
 	Ar.UsingCustomVersion(FAnimObjectVersion::GUID);
 
@@ -108,7 +111,7 @@ static void InternalSerializePlatformData(FArchive& Ar, UObject* Owner, UGroomBi
 	InternalSerializeGuides(Ar, Owner, GroupData.SimRootBulkData);
 
 	// Strands
-	InternalSerializeStrands(Ar, Owner, GroupData.RenRootBulkData, Flags, bHeader, bData);
+	InternalSerializeStrands(Ar, Owner, GroupData.RenRootBulkData, Flags, bHeader, bData, bOutHasDataInCache);
 
 	// Cards
 	InternalSerializeCards(Ar, Owner, GroupData.CardsRootBulkData);
@@ -875,6 +878,7 @@ void UGroomBindingAsset::CacheDerivedDatas(uint32 InGroupIndex, const FString Ke
 
 		FGroomComponentRecreateRenderStateContext RecreateRenderContext(GetGroom());
 
+		bool bSuccess = false;
 		if (Data)
 		{
 			UE_CLOG(IsHairStrandsDDCLogEnable(), LogHairStrands, Log, TEXT("[GroomBinding/DDC] Found (GroomBinding:%s)."), *GetName());
@@ -883,11 +887,11 @@ void UGroomBindingAsset::CacheDerivedDatas(uint32 InGroupIndex, const FString Ke
 
 			// Header
 			FMemoryReaderView Ar(Data, /*bIsPersistent*/ true);
-			InternalSerializePlatformData(Ar, this, PlatformData, 0 /*Flags*/, true /*bHeader*/, false /*bData*/);
+			InternalSerializePlatformData(Ar, this, PlatformData, 0 /*Flags*/, true /*bHeader*/, false /*bData*/, &bSuccess);
 
 			bOutValid = true;
 		}
-		else
+		if (!bSuccess)
 		{
 			UE_CLOG(IsHairStrandsDDCLogEnable(), LogHairStrands, Log, TEXT("[GroomBinding/DDC] Not found (GroomBinding:%s)."), *GetName());
 
