Add possibility to use strands guide as card guides for card groom having non-conforming topology (i.e. triangles strip).

diff --git a/Engine/Plugins/Runtime/HairStrands/Source/HairStrandsCore/Private/GroomAsset.cpp b/Engine/Plugins/Runtime/HairStrands/Source/HairStrandsCore/Private/GroomAsset.cpp
--- a/Engine/Plugins/Runtime/HairStrands/Source/HairStrandsCore/Private/GroomAsset.cpp
+++ b/Engine/Plugins/Runtime/HairStrands/Source/HairStrandsCore/Private/GroomAsset.cpp
@@ -2175,7 +2175,8 @@ namespace GroomDerivedDataCacheUtils
 		{
 			FString Key = Desc.GetMeshKey();
 			Ar << Key;
-		}		
+		}
+		Ar << Desc.GuideType;
 
 		FSHAHash Hash;
 		FSHA1::HashBuffer(TempBytes.GetData(), TempBytes.Num(), Hash.Hash);
@@ -2865,6 +2866,66 @@ inline FString GetLODName(const UGroomAsset* Asset, uint32 LODIndex)
 	return Asset->GetOutermost()->GetName() + FString::Printf(TEXT("_LOD%d"), LODIndex);
 }
 
+static bool InternalImportCardGeometry(FHairGroupsCardsSourceDescription* InDesc, const FHairDescriptionGroups& InHairDescriptionGroups, FHairCardsBulkData& OutBulkData, FHairStrandsDatas& OutGuidesData, FHairCardsInterpolationBulkData& OutInterpolationBulkData)
+{
+	UStaticMesh* CardsMesh = InDesc->ImportedMesh;
+	check(CardsMesh);
+	check(InHairDescriptionGroups.IsValid());
+	check(InHairDescriptionGroups.HairGroups.IsValidIndex(InDesc->GroupIndex));
+
+	// * Create a transient FHairStrandsData in order to extract RootUV and transfer them to cards data
+	// * Voxelize hair strands data (all group), to transfer group index from strands to cards
+	//FHairGroupInfo DummyInfo;
+
+	// Guides
+	FHairStrandsDatas TempHairGuidesData = InHairDescriptionGroups.HairGroups[InDesc->GroupIndex].Guides;
+	FGroomBuilder::BuildData(TempHairGuidesData);
+
+	// Strands
+	FHairStrandsVoxelData TempHairStrandsVoxelData;
+	FHairStrandsDatas TempHairStrandsData = InHairDescriptionGroups.HairGroups[InDesc->GroupIndex].Strands;
+	FGroomBuilder::BuildData(TempHairStrandsData);
+	FGroomBuilder::VoxelizeGroupIndex(InHairDescriptionGroups, TempHairStrandsVoxelData);
+
+	// Transfer RootUV from strands to guides
+	const bool bNeedRootUVTransfer = InDesc->GuideType == EHairCardsGuideType::GuideBased;
+	if (bNeedRootUVTransfer)
+	{
+		struct FRoot
+		{
+			FVector3f Position;
+			FVector2f RootUV;
+		};
+
+		// Extra all roots from strands
+		TArray<FRoot> StrandsRoots;
+		StrandsRoots.SetNum(TempHairStrandsData.GetNumCurves());
+		for (uint32 CurveIt = 0, CurveCount = TempHairStrandsData.GetNumCurves(); CurveIt < CurveCount; ++CurveIt)
+		{
+			StrandsRoots[CurveIt].Position = TempHairStrandsData.StrandsPoints.PointsPosition[TempHairStrandsData.StrandsCurves.CurvesOffset[CurveIt]];
+			StrandsRoots[CurveIt].RootUV   = TempHairStrandsData.StrandsCurves.CurvesRootUV[CurveIt];
+		}
+
+		// Find closest strands root for each guide roots, and assign RootUV
+		for (uint32 CurveIt = 0, CurveCount = TempHairGuidesData.GetNumCurves(); CurveIt < CurveCount; ++CurveIt)
+		{
+			float MaxDistance = FLT_MAX;
+			const FVector3f Position = TempHairGuidesData.StrandsPoints.PointsPosition[TempHairGuidesData.StrandsCurves.CurvesOffset[CurveIt]];
+			for (const FRoot& Root : StrandsRoots)
+			{
+				const float Dist = (Position - Root.Position).Length();
+				if (Dist < MaxDistance)
+				{
+					TempHairGuidesData.StrandsCurves.CurvesRootUV[CurveIt] = Root.RootUV;
+					MaxDistance = Dist;
+				}
+			}
+		}
+	}
+
+	return FHairCardsBuilder::ImportGeometry(CardsMesh, TempHairGuidesData, TempHairStrandsData, TempHairStrandsVoxelData, InDesc->GuideType == EHairCardsGuideType::Generated, OutBulkData, OutGuidesData, OutInterpolationBulkData);
+}
+
 bool UGroomAsset::BuildCardsData(uint32 GroupIndex)
 {
 	LLM_SCOPE_BYTAG(Groom);
@@ -2965,20 +3026,8 @@ bool UGroomAsset::BuildCardsData(uint32 GroupIndex)
 			{
 				CardsMesh->ConditionalPostLoad();
 
-				// * Create a transient FHairStrandsData in order to extract RootUV and transfer them to cards data
-				// * Voxelize hair strands data (all group), to transfer group index from strands to cards
-				FHairStrandsDatas TempHairStrandsData;
-				FHairStrandsVoxelData TempHairStrandsVoxelData;
-				{
-					const FHairDescriptionGroups& LocalHairDescriptionGroups = GetHairDescriptionGroups();
-					FHairGroupInfo DummyInfo;
-					check(LocalHairDescriptionGroups.IsValid());
-					TempHairStrandsData = LocalHairDescriptionGroups.HairGroups[GroupIndex].Strands;
-					FGroomBuilder::BuildData(TempHairStrandsData);
-					FGroomBuilder::VoxelizeGroupIndex(LocalHairDescriptionGroups, TempHairStrandsVoxelData);
-				}
-				
-				bInitResources = FHairCardsBuilder::ImportGeometry(CardsMesh, TempHairStrandsData, TempHairStrandsVoxelData, LOD.BulkData, LODGuidesData, LOD.InterpolationBulkData);
+				const FHairDescriptionGroups& LocalHairDescriptionGroups = GetHairDescriptionGroups();
+				bInitResources = InternalImportCardGeometry(Desc, LocalHairDescriptionGroups, LOD.BulkData, LODGuidesData, LOD.InterpolationBulkData);
 				if (!bInitResources)
 				{
 					UE_LOG(LogHairStrands, Warning, TEXT("Failed to import cards from %s for Group %d LOD %d."), *CardsMesh->GetName(), GroupIndex, LODIt);
@@ -3905,11 +3954,10 @@ bool UGroomAsset::GetHairCardsGuidesDatas(
 		{
 			CardsMesh->ConditionalPostLoad();
 
-			FHairStrandsVoxelData				DummyVoxelData;
-			FHairCardsBulkData					DummyBulkData;
-			FHairCardsInterpolationBulkData		DummyInterpolationBulkData;
-			FHairStrandsDatas					DummyHairStrandsData;
-			FHairCardsBuilder::ImportGeometry(CardsMesh, DummyHairStrandsData, DummyVoxelData, DummyBulkData, OutCardsGuidesData, DummyInterpolationBulkData);
+			FHairCardsInterpolationBulkData	OutDummyInterpolationBulkData;
+			FHairCardsBulkData OutDummyLODBulkData;
+			const FHairDescriptionGroups& LocalHairDescriptionGroups = GetHairDescriptionGroups();
+			InternalImportCardGeometry(Desc, LocalHairDescriptionGroups, OutDummyLODBulkData, OutCardsGuidesData, OutDummyInterpolationBulkData);
 			return true;
 		}
 	}
diff --git a/Engine/Plugins/Runtime/HairStrands/Source/HairStrandsCore/Private/GroomAssetCards.cpp b/Engine/Plugins/Runtime/HairStrands/Source/HairStrandsCore/Private/GroomAssetCards.cpp
--- a/Engine/Plugins/Runtime/HairStrands/Source/HairStrandsCore/Private/GroomAssetCards.cpp
+++ b/Engine/Plugins/Runtime/HairStrands/Source/HairStrandsCore/Private/GroomAssetCards.cpp
@@ -11,6 +11,7 @@ FHairGroupsCardsSourceDescription::FHairGroupsCardsSourceDescription()
 {
 	MaterialSlotName = NAME_None;
 	SourceType_DEPRECATED = EHairCardsSourceType::Imported;
+	GuideType = EHairCardsGuideType::Generated;
 	GroupIndex = 0;
 	LODIndex = -1;
 }
diff --git a/Engine/Plugins/Runtime/HairStrands/Source/HairStrandsCore/Private/HairCardsBuilder.cpp b/Engine/Plugins/Runtime/HairStrands/Source/HairStrandsCore/Private/HairCardsBuilder.cpp
--- a/Engine/Plugins/Runtime/HairStrands/Source/HairStrandsCore/Private/HairCardsBuilder.cpp
+++ b/Engine/Plugins/Runtime/HairStrands/Source/HairStrandsCore/Private/HairCardsBuilder.cpp
@@ -66,7 +66,7 @@ namespace FHairCardsBuilder
 FString GetVersion()
 {
 	// Important to update the version when cards building or importing changes
-	return TEXT("9g");
+	return TEXT("9m");
 }
 
 bool InternalCreateCardsGuides(
@@ -352,8 +352,7 @@ bool InternalCreateCardsGuides(
 void InternalCreateCardsInterpolation(
 	const FHairCardsGeometry& InCards,
 	const FHairStrandsDatas& InGuides,
-	FHairCardsInterpolationDatas& Out,
-	const TArray<float>& CardLengths)
+	FHairCardsInterpolationDatas& Out)
 {		
 	const uint32 TotalVertexCount = InCards.Positions.Num();
 	Out.SetNum(TotalVertexCount);
@@ -371,8 +370,6 @@ void InternalCreateCardsInterpolation(
 		const uint32 IndexOffset = InCards.IndexOffsets[CardIt];
 		const uint32 IndexCount  = InCards.IndexCounts[CardIt];
 
-		const float CardLength = CardLengths[CardIt];
-
 		for (uint32 IndexIt = 0; IndexIt < IndexCount; ++IndexIt)
 		{
 			const uint32 VertexIndex = InCards.Indices[IndexOffset + IndexIt];
@@ -432,7 +429,24 @@ void SanitizeMeshDescription(FMeshDescription* MeshDescription)
 	}
 }
 
-bool InternalImportGeometry(
+static float PackCardLengthAndGroupIndex(float InCardLength, uint32 InCardGroupIndex)
+{
+	// Encode cards length & group index into the .W component of position
+	const uint32 EncodedW = FFloat16(InCardLength).Encoded | (InCardGroupIndex << 16u);
+	return *(float*)&EncodedW;
+}
+
+static uint32 PackMaterialAttribute(const FVector3f& InBaseColor, float InRoughness)
+{
+	// Encode the base color in (cheap) sRGB in XYZ. The W component remains unused
+	return
+		(uint32(FMath::Sqrt(InBaseColor.X) * 255.f)    )|
+		(uint32(FMath::Sqrt(InBaseColor.Y) * 255.f)<<8 )|
+		(uint32(FMath::Sqrt(InBaseColor.Z) * 255.f)<<16)|
+		(uint32(InRoughness                * 255.f)<<24);
+}
+
+static bool InternalImportGeometry_WithGeneratedGuides(
 	const UStaticMesh* StaticMesh,
 	const FHairStrandsDatas& InStrandsData,			// Used for extracting & assigning root UV to cards data
 	const FHairStrandsVoxelData& InStrandsVoxelData,// Used for transfering & assigning group index to cards data
@@ -593,7 +607,7 @@ bool InternalImportGeometry(
 	if (bSuccess)
 	{
 		FHairCardsInterpolationDatas InterpolationData;
-		InternalCreateCardsInterpolation(Out.Cards, OutGuides, InterpolationData, CardLengths);
+		InternalCreateCardsInterpolation(Out.Cards, OutGuides, InterpolationData);
 
 		// Fill out the interpolation data
 		OutInterpolationBulkData.Interpolation.SetNum(PointCount);
@@ -666,23 +680,13 @@ bool InternalImportGeometry(
 		{
 			const uint32 VertexIndex = Out.Cards.Indices[CardsIndexOffset + IndexIt];
 			const float CoordU = Out.Cards.CoordU[VertexIndex];
-			const float InterpolatedCardLength = FMath::Lerp(0.f, CardLength, CoordU);
 
 			// Instead of storing the interpolated card length, store the actual max length of the card, 
 			// as reconstructing the strands length, based on interpolated CardLength will be too prone to numerical issue.
 			// This means that the strand length retrieves in shader will be an over estimate of the actual length
-			const FFloat16 hCardLength = CardLength;  // InterpolatedCardLength;
-
-			// Encode cards length & group index into the .W component of position
-			const uint32 EncodedW = hCardLength.Encoded | (CardGroupIndices[CardIt] << 16u);
-			OutBulk.Positions[VertexIndex].W = *(float*)&EncodedW;
-
-			// Encode the base color in (cheap) sRGB in XYZ. The W component remains unused
-			OutBulk.Materials[VertexIndex] =
-				(uint32(FMath::Sqrt(CardBaseColor[VertexIndex].X) * 255.f)    )|
-				(uint32(FMath::Sqrt(CardBaseColor[VertexIndex].Y) * 255.f)<<8 )|
-				(uint32(FMath::Sqrt(CardBaseColor[VertexIndex].Z) * 255.f)<<16)|
-				(uint32(CardsRoughness[VertexIndex]               * 255.f)<<24);
+			OutBulk.Positions[VertexIndex].W = PackCardLengthAndGroupIndex(CardLength, CardGroupIndices[CardIt]);
+
+			OutBulk.Materials[VertexIndex] = PackMaterialAttribute(CardBaseColor[VertexIndex], CardsRoughness[VertexIndex]);
 		}
 	}
 
@@ -831,23 +835,250 @@ bool InternalImportGeometry(
 	return bSuccess;
 }
 
+
+static bool InternalImportGeometry_WithImportedGuides(
+	const UStaticMesh* StaticMesh,
+	const FHairStrandsDatas& InGuides,
+	const FHairStrandsDatas& InStrandsData,			// Used for extracting & assigning root UV to cards data
+	const FHairStrandsVoxelData& InStrandsVoxelData,// Used for transfering & assigning group index to cards data
+	FHairCardsBulkData& OutBulk,
+	FHairStrandsDatas& OutGuides,
+	FHairCardsInterpolationBulkData& OutInterpolationBulkData)
+{
+	// Note: if there are multiple section we only import the first one. Support for multiple section could be added later on. 
+	FMeshDescription* MeshDescription = StaticMesh->GetMeshDescription(0);
+	const uint32 PointCount = MeshDescription->Vertices().Num();
+	const uint32 IndexCount  = MeshDescription->Triangles().Num() * 3;
+
+	const uint32 MeshLODIndex = 0;
+
+	SanitizeMeshDescription(MeshDescription);
+	const TVertexAttributesRef<const FVector3f> VertexPositions					= MeshDescription->VertexAttributes().GetAttributesRef<FVector3f>(MeshAttribute::Vertex::Position);
+	const TVertexInstanceAttributesRef<const FVector3f> VertexInstanceNormals	= MeshDescription->VertexInstanceAttributes().GetAttributesRef<FVector3f>(MeshAttribute::VertexInstance::Normal);
+	const TVertexInstanceAttributesRef<const FVector3f> VertexInstanceTangents	= MeshDescription->VertexInstanceAttributes().GetAttributesRef<FVector3f>(MeshAttribute::VertexInstance::Tangent);
+	const TVertexInstanceAttributesRef<const float> VertexInstanceBinormalSigns	= MeshDescription->VertexInstanceAttributes().GetAttributesRef<float>(MeshAttribute::VertexInstance::BinormalSign);
+	const TVertexInstanceAttributesRef<const FVector2f> VertexInstanceUVs		= MeshDescription->VertexInstanceAttributes().GetAttributesRef<FVector2f>(MeshAttribute::VertexInstance::TextureCoordinate);
+
+	// Write out BulkData
+	OutBulk.Positions.SetNum(PointCount);
+	OutBulk.Normals.SetNum(PointCount * FHairCardsNormalFormat::ComponentCount);
+	OutBulk.UVs.SetNum(PointCount);
+	OutBulk.Materials.SetNum(PointCount);
+	FBox3f BoundingBox = FBox3f(EForceInit::ForceInit);
+	for (const FVertexID VertexId : MeshDescription->Vertices().GetElementIDs())
+	{
+		TArrayView<const FVertexInstanceID> VertexInstanceIds = MeshDescription->GetVertexVertexInstanceIDs(VertexId);
+		if (VertexInstanceIds.Num() == 0)
+		{
+			continue;
+		}
+
+		FVertexInstanceID VertexInstanceId0 = VertexInstanceIds[0]; // Assume no actual duplicated data.
+
+		const uint32 VertexIndex = VertexId.GetValue();
+		check(VertexIndex < PointCount);
+		OutBulk.Positions[VertexIndex] 		= FVector4f(VertexPositions[VertexId], 0);
+		OutBulk.UVs[VertexIndex] 			= FVector4f(VertexInstanceUVs[VertexInstanceId0].Component(0), VertexInstanceUVs[VertexInstanceId0].Component(1), 0, 0); // RootUV are not set here, but will be 'patched' later once guides & interpolation data are built
+		OutBulk.Normals[VertexIndex * 2]	= FVector4f(VertexInstanceTangents[VertexInstanceId0], 0);
+		OutBulk.Normals[VertexIndex * 2 + 1]= FVector4f(VertexInstanceNormals[VertexInstanceId0], VertexInstanceBinormalSigns[VertexInstanceId0] /*TangentFrameSigns*/);
+		OutBulk.Materials[VertexIndex] 		= 0;
+
+		BoundingBox += OutBulk.Positions[VertexIndex];
+	}
+	OutBulk.BoundingBox = ToFBox3d(BoundingBox);
+
+	TArray<FVector3f> CardBaseColor;
+	TArray<float> CardsRoughness;
+	TArray<uint8> CardGroupIndices;
+	TArray<float> CardLengths;
+	CardGroupIndices.Init(0u, PointCount);					// Per-vertex
+	CardBaseColor.Init(FVector3f::ZeroVector, PointCount);	// Per-vertex
+	CardsRoughness.Init(0.f, PointCount);					// Per-vertex
+	CardLengths.Init(0.f, PointCount);						// Per-vertex
+
+	// Compute material properties per vertex
+	if (InStrandsVoxelData.IsValid())
+	{
+		for (uint32 PointIt = 0; PointIt < PointCount; ++PointIt)
+		{
+			const FVector4f& P = OutBulk.Positions[PointIt];
+
+			FHairStrandsVoxelData::FData VoxelData = InStrandsVoxelData.GetData(P);
+			if (VoxelData.GroupIndex != FHairStrandsVoxelData::InvalidGroupIndex)
+			{
+				CardBaseColor[PointIt] = VoxelData.BaseColor;
+				CardsRoughness[PointIt] = VoxelData.Roughness;
+				CardGroupIndices[PointIt] = VoxelData.GroupIndex;
+			}
+		}
+	}
+
+	// Build interpolation data
+	{		
+		struct FClosestGuide
+		{
+			uint32 GuideIndex = ~0u;
+			uint32 PointIndex = ~0u;
+			float U = 0;
+			float Distance = FLT_MAX;
+		};
+
+		TArray<FClosestGuide> ClosestGuides;
+		ClosestGuides.Init(FClosestGuide(), PointCount);
+
+		FHairCardsInterpolationDatas InterpolationData;
+		InterpolationData.SetNum(PointCount);
+		OutInterpolationBulkData.Interpolation.SetNum(PointCount);
+
+		TArray<float> CoordUs;
+		CoordUs.Init(0, PointCount);
+
+		// For each cards, and for each cards vertex,
+		// Compute the closest guide points (two guide points to interpolation in-between), 
+		// and compute their indices and lerping value
+		const uint32 GuideCount = InGuides.GetNumCurves();
+		for (uint32 GuideIt = 0; GuideIt < GuideCount; ++GuideIt)
+		{
+			const uint32 GuidePointOffset = InGuides.StrandsCurves.CurvesOffset[GuideIt];
+			const uint32 GuidePointCount = InGuides.StrandsCurves.CurvesCount[GuideIt];
+			check(GuidePointCount >= 2);
+
+			uint32 GuideIndex0 = GuidePointOffset + 0;
+			uint32 GuideIndex1 = GuidePointOffset + 1;
+			float  GuideLerp   = 0;
+			bool bFoundMatch = false;
+			for (uint32 GuidePointIt = 0; GuidePointIt < GuidePointCount-1; ++GuidePointIt)
+			{
+				const uint32 I0 = GuidePointOffset + GuidePointIt;
+				const uint32 I1 = I0+1u;
+
+				const FVector3f& GPoint0 = InGuides.StrandsPoints.PointsPosition[I0];
+				const FVector3f& GPoint1 = InGuides.StrandsPoints.PointsPosition[I1];
+				const float GuideCoordU0 = InGuides.StrandsPoints.PointsCoordU[I0];
+				const float GuideCoordU1 = InGuides.StrandsPoints.PointsCoordU[I1];
+
+				// For each point of the mesh, compute its distance to the guide's segment
+				for (uint32 PointIt = 0; PointIt < PointCount; ++PointIt)
+				{
+					// Project P onto the segment.
+					const FVector3f P = OutBulk.Positions[PointIt];
+					float U = FVector3f::DotProduct(P-GPoint0, GPoint1-GPoint0);
+					U = FMath::Clamp(U, 0.f, 1.f);
+					const FVector3f PP = U * (GPoint1-GPoint0) + GPoint0;
+					const float Dist = (P-PP).Length();
+
+					// If new projection is closer than the older one, change guides
+					FClosestGuide& ClosestGuide = ClosestGuides[PointIt];
+					if (Dist < ClosestGuide.Distance)
+					{
+						ClosestGuide.GuideIndex = GuideIt;
+						ClosestGuide.PointIndex = GuidePointIt;
+						ClosestGuide.U = U;
+						ClosestGuide.Distance = Dist;
+						CoordUs[PointIt] = FMath::Lerp(GuideCoordU0, GuideCoordU1, ClosestGuide.U);
+						CardLengths[PointIt] = InGuides.StrandsCurves.CurvesLength[GuideIt];
+					}
+				}
+			}
+		}
+
+		// Compute the interpolation data
+		for (uint32 PointIt = 0; PointIt < PointCount; ++PointIt)
+		{
+			const FClosestGuide& ClosestGuide = ClosestGuides[PointIt];
+
+			// Interpolation data
+			{
+				check(ClosestGuide.Distance < FLT_MAX);
+				InterpolationData.PointsSimCurvesIndex[PointIt] = ClosestGuide.GuideIndex;
+				InterpolationData.PointsSimCurvesVertexIndex[PointIt] = InGuides.StrandsCurves.CurvesOffset[ClosestGuide.GuideIndex] + ClosestGuide.PointIndex;
+				InterpolationData.PointsSimCurvesVertexLerp[PointIt] = ClosestGuide.U;
+			}
+			// Interpolation bulk data
+			{
+				FHairCardsInterpolationVertex PackedData;
+				PackedData.VertexIndex = InterpolationData.PointsSimCurvesVertexIndex[PointIt];
+				PackedData.VertexLerp = FMath::Clamp(uint32(InterpolationData.PointsSimCurvesVertexLerp[PointIt] * 0xFF), 0u, 0xFFu);
+				OutInterpolationBulkData.Interpolation[PointIt] = PackedData;
+			}
+			// Root UV bulk data
+			{
+				const FVector2f RootUV = InGuides.StrandsCurves.CurvesRootUV[ClosestGuide.GuideIndex];
+				OutBulk.UVs[PointIt].Z = RootUV.X;
+				OutBulk.UVs[PointIt].W = RootUV.Y;
+			}
+		}
+	}
+
+	// Encode material bulk data
+	for (uint32 PointIt = 0; PointIt < PointCount; ++PointIt)
+	{
+		OutBulk.Positions[PointIt].W = PackCardLengthAndGroupIndex(CardLengths[PointIt], CardGroupIndices[PointIt]);
+		OutBulk.Materials[PointIt] = PackMaterialAttribute(CardBaseColor[PointIt], CardsRoughness[PointIt]);
+	}
+
+	// Fill in vertex indices and the cards indices offset/count
+	OutBulk.Indices.Reserve(IndexCount);
+	for (const FTriangleID& TriangleId : MeshDescription->Triangles().GetElementIDs())
+	{
+		TArrayView<const FVertexInstanceID> VertexInstanceIDs = MeshDescription->GetTriangleVertexInstances(TriangleId);
+		check(VertexInstanceIDs.Num() == 3);
+		FVertexInstanceID VI0 = VertexInstanceIDs[0];
+		FVertexInstanceID VI1 = VertexInstanceIDs[1];
+		FVertexInstanceID VI2 = VertexInstanceIDs[2];
+
+		FVertexID V0 = MeshDescription->GetVertexInstanceVertex(VI0);
+		FVertexID V1 = MeshDescription->GetVertexInstanceVertex(VI1);
+		FVertexID V2 = MeshDescription->GetVertexInstanceVertex(VI2);
+
+		OutBulk.Indices.Add(V0.GetValue());
+		OutBulk.Indices.Add(V1.GetValue());
+		OutBulk.Indices.Add(V2.GetValue());
+	}
+
+	OutGuides = InGuides;
+
+	OutBulk.DepthTexture = nullptr;
+	OutBulk.TangentTexture = nullptr;
+	OutBulk.CoverageTexture = nullptr;
+	OutBulk.AttributeTexture = nullptr;
+	
+	return true;
+}
+
 bool ImportGeometry(
 	const UStaticMesh* StaticMesh,
+	const FHairStrandsDatas& InGuidesData,
 	const FHairStrandsDatas& InStrandsData,
 	const FHairStrandsVoxelData& InStrandsVoxelData,
+	const bool bGenerateGuidesFromCardGeometry,
 	FHairCardsBulkData& OutBulk,
 	FHairStrandsDatas& OutGuides,
 	FHairCardsInterpolationBulkData& OutInterpolationBulkData)
 {
-	FHairCardsDatas CardData;
-	return InternalImportGeometry(
-		StaticMesh,
-		InStrandsData,
-		InStrandsVoxelData,
-		CardData,
-		OutBulk,
-		OutGuides,
-		OutInterpolationBulkData);
+	if (bGenerateGuidesFromCardGeometry)
+	{
+		FHairCardsDatas CardData;
+		return InternalImportGeometry_WithGeneratedGuides(
+			StaticMesh,
+			InStrandsData,
+			InStrandsVoxelData,
+			CardData,
+			OutBulk,
+			OutGuides,
+			OutInterpolationBulkData);
+	}
+	else
+	{
+		return InternalImportGeometry_WithImportedGuides(
+			StaticMesh,
+			InGuidesData,
+			InStrandsData,
+			InStrandsVoxelData,
+			OutBulk,
+			OutGuides,
+			OutInterpolationBulkData);
+	}
 }
 
 bool ExtractCardsData(const UStaticMesh* StaticMesh, const FHairStrandsDatas& InStrandsData, FHairCardsDatas& Out)
@@ -856,7 +1087,7 @@ bool ExtractCardsData(const UStaticMesh* StaticMesh, const FHairStrandsDatas& In
 	FHairCardsBulkData OutBulk;
 	FHairStrandsDatas OutGuides;
 	FHairCardsInterpolationBulkData OutInterpolationBulkData;
-	return InternalImportGeometry(
+	return InternalImportGeometry_WithGeneratedGuides(
 		StaticMesh,
 		InStrandsData,
 		DummyStrandsVoxelData,
diff --git a/Engine/Plugins/Runtime/HairStrands/Source/HairStrandsCore/Public/GroomAssetCards.h b/Engine/Plugins/Runtime/HairStrands/Source/HairStrandsCore/Public/GroomAssetCards.h
--- a/Engine/Plugins/Runtime/HairStrands/Source/HairStrandsCore/Public/GroomAssetCards.h
+++ b/Engine/Plugins/Runtime/HairStrands/Source/HairStrandsCore/Public/GroomAssetCards.h
@@ -21,6 +21,13 @@ enum class EHairCardsSourceType : uint8
 	Imported UMETA(DisplayName = "Imported"),
 };
 
+UENUM(BlueprintType)
+enum class EHairCardsGuideType : uint8
+{
+	Generated  UMETA(DisplayName = "Generated"),
+	GuideBased UMETA(DisplayName = "Guide-Based"),
+};
+
 USTRUCT(BlueprintType)
 struct HAIRSTRANDSCORE_API FHairGroupCardsInfo
 {
@@ -98,6 +105,9 @@ struct HAIRSTRANDSCORE_API FHairGroupsCardsSourceDescription
 	UPROPERTY()
 	bool bInvertUV = false;
 
+	UPROPERTY(EditAnywhere, Category = "CardsSource")
+	EHairCardsGuideType GuideType;
+
 	UPROPERTY(EditAnywhere, Category = "CardsSource", meta = (DisplayName = "Mesh"))
 	TObjectPtr<UStaticMesh> ImportedMesh = nullptr;
 
diff --git a/Engine/Plugins/Runtime/HairStrands/Source/HairStrandsCore/Public/HairCardsBuilder.h b/Engine/Plugins/Runtime/HairStrands/Source/HairStrandsCore/Public/HairCardsBuilder.h
--- a/Engine/Plugins/Runtime/HairStrands/Source/HairStrandsCore/Public/HairCardsBuilder.h
+++ b/Engine/Plugins/Runtime/HairStrands/Source/HairStrandsCore/Public/HairCardsBuilder.h
@@ -34,10 +34,12 @@ namespace FHairCardsBuilder
 {
 	bool ImportGeometry(
 		const UStaticMesh* StaticMesh,
+		const FHairStrandsDatas& InGuidesData,
 		const FHairStrandsDatas& InStrandsData,
 		const FHairStrandsVoxelData& InStrandsVoxelData,
+		const bool bGenerateGuidesFromCardGeometry,
 		FHairCardsBulkData& OutBulk,
-		FHairStrandsDatas& OutGuides,
+		FHairStrandsDatas& OutCardGuides,
 		FHairCardsInterpolationBulkData& OutInterpolationBulkData);
 
 	HAIRSTRANDSCORE_API bool ExtractCardsData(
@@ -60,4 +62,4 @@ namespace FHairMeshesBuilder
 
 	FString GetVersion();
 }
-#endif // WITH_EDITOR
\ No newline at end of file
+#endif // WITH_EDITOR
