Groom binding builder now uses explicit pointer to skel. mesh data to ease debugging/tracking potential race condition with skel. mesh data building.

diff --git a/Engine/Plugins/Runtime/HairStrands/Source/HairStrandsCore/Private/GroomBindingBuilder.cpp b/Engine/Plugins/Runtime/HairStrands/Source/HairStrandsCore/Private/GroomBindingBuilder.cpp
--- a/Engine/Plugins/Runtime/HairStrands/Source/HairStrandsCore/Private/GroomBindingBuilder.cpp
+++ b/Engine/Plugins/Runtime/HairStrands/Source/HairStrandsCore/Private/GroomBindingBuilder.cpp
@@ -114,7 +114,7 @@ struct FHairRootGroupData
 	TArray<FHairStrandsRootData>	CardsRootData;
 };
 
-namespace
+namespace GroomBinding_Mesh
 {
 //////////////////////////////////////////////////////////////////////////
 // Interfaces to query mesh data from different sources
@@ -162,58 +162,85 @@ public:
 class FSkeletalMeshSection : public IMeshSectionData
 {
 public:
-	FSkeletalMeshSection(FSkelMeshRenderSection& InSection)
-		: Section(InSection)
+	FSkeletalMeshSection(const FSkeletalMeshRenderData* InMeshData=nullptr, int32 InLODIndex=-1, int32 InSectionIndex=-1)
 	{
+		LODIndex 	 = InLODIndex;
+		SectionIndex = InSectionIndex;
+		MeshData 	 = InMeshData; 
 	}
 
 	virtual uint32 GetNumVertices() const override
 	{
-		return Section.NumVertices;
+		check(MeshData->LODRenderData.IsValidIndex(LODIndex));
+		check(MeshData->LODRenderData[LODIndex].RenderSections.IsValidIndex(SectionIndex));
+		return MeshData->LODRenderData[LODIndex].RenderSections[SectionIndex].NumVertices;
 	}
 
 	virtual uint32 GetNumTriangles() const override
 	{
-		return Section.NumTriangles;
+		check(MeshData->LODRenderData.IsValidIndex(LODIndex));
+		check(MeshData->LODRenderData[LODIndex].RenderSections.IsValidIndex(SectionIndex));
+		return MeshData->LODRenderData[LODIndex].RenderSections[SectionIndex].NumTriangles;
 	}
 
 	virtual uint32 GetBaseIndex() const override
 	{
-		return Section.BaseIndex;
+		check(MeshData->LODRenderData.IsValidIndex(LODIndex));
+		check(MeshData->LODRenderData[LODIndex].RenderSections.IsValidIndex(SectionIndex));
+		return MeshData->LODRenderData[LODIndex].RenderSections[SectionIndex].BaseIndex;
 	}
 
 	virtual uint32 GetBaseVertexIndex() const override
 	{
-		return Section.BaseVertexIndex;
+		check(MeshData->LODRenderData.IsValidIndex(LODIndex));
+		check(MeshData->LODRenderData[LODIndex].RenderSections.IsValidIndex(SectionIndex));
+		return MeshData->LODRenderData[LODIndex].RenderSections[SectionIndex].BaseVertexIndex;
 	}
 
 private:
-	FSkelMeshRenderSection& Section;
+	int32 LODIndex = -1;
+	int32 SectionIndex = -1;
+	const FSkeletalMeshRenderData* MeshData = nullptr;
 };
 
 class FSkeletalMeshLODData : public IMeshLODData
 {
 public:
-	FSkeletalMeshLODData(FSkeletalMeshLODRenderData& InMeshLODData) 
-		: MeshLODData(InMeshLODData)
+	FSkeletalMeshLODData(const FSkeletalMeshRenderData* InMeshData=nullptr, int32 InLODIndex=-1)
 	{
-		IndexBuffer.SetNum(MeshLODData.MultiSizeIndexContainer.GetIndexBuffer()->Num());
-		MeshLODData.MultiSizeIndexContainer.GetIndexBuffer(IndexBuffer);
+		LODIndex = InLODIndex;
+		MeshData = InMeshData;
 
-		for (FSkelMeshRenderSection& MeshSection : MeshLODData.RenderSections)
+		if (MeshData && MeshData->LODRenderData.IsValidIndex(LODIndex))
 		{
-			Sections.Add(FSkeletalMeshSection(MeshSection));
+			const uint32 SectionCount = MeshData->LODRenderData[LODIndex].RenderSections.Num();
+			Sections.Reserve(SectionCount);
+			for (uint32 SectionIt=0;SectionIt<SectionCount; ++SectionIt)
+			{
+				Sections.Add(FSkeletalMeshSection(MeshData, LODIndex, SectionIt));
+			}
+
+			if (IndexBuffer.IsEmpty())
+			{
+				check(MeshData->LODRenderData.IsValidIndex(LODIndex));
+				IndexBuffer.SetNum(MeshData->LODRenderData[LODIndex].MultiSizeIndexContainer.GetIndexBuffer()->Num());
+				MeshData->LODRenderData[LODIndex].MultiSizeIndexContainer.GetIndexBuffer(IndexBuffer);
+			}
 		}
 	}
 
 	virtual const FVector3f* GetVerticesBuffer() const override
 	{
-		return static_cast<FVector3f*>(MeshLODData.StaticVertexBuffers.PositionVertexBuffer.GetVertexData());
+		check(MeshData);
+		check(MeshData->LODRenderData.IsValidIndex(LODIndex));
+		return static_cast<const FVector3f*>(MeshData->LODRenderData[LODIndex].StaticVertexBuffers.PositionVertexBuffer.GetVertexData());
 	}
 
 	virtual uint32 GetNumVertices() const override
 	{
-		return MeshLODData.StaticVertexBuffers.PositionVertexBuffer.GetNumVertices();
+		check(MeshData);
+		check(MeshData->LODRenderData.IsValidIndex(LODIndex));
+		return MeshData->LODRenderData[LODIndex].StaticVertexBuffers.PositionVertexBuffer.GetNumVertices();
 	}
 
 	virtual const TArray<uint32>& GetIndexBuffer() const override
@@ -226,29 +253,39 @@ public:
 		return Sections.Num();
 	}
 
-	virtual const IMeshSectionData& GetSection(uint32 SectionIndex) const override
+	virtual const IMeshSectionData& GetSection(uint32 InSectionIndex) const override
 	{
-		return Sections[SectionIndex];
+		check(MeshData);
+		check(Sections.IsValidIndex(InSectionIndex));
+		return Sections[InSectionIndex];
 	}
 
-	virtual const FVector3f& GetVertexPosition(uint32 VertexIndex) const override
+	virtual const FVector3f& GetVertexPosition(uint32 InVertexIndex) const override
 	{
-		return MeshLODData.StaticVertexBuffers.PositionVertexBuffer.VertexPosition(VertexIndex);
+		check(MeshData);
+		check(MeshData->LODRenderData.IsValidIndex(LODIndex));
+		return MeshData->LODRenderData[LODIndex].StaticVertexBuffers.PositionVertexBuffer.VertexPosition(InVertexIndex);
 	}
 
-	virtual FVector2D GetVertexUV(uint32 VertexIndex, uint32 ChannelIndex) const override
+	virtual FVector2D GetVertexUV(uint32 InVertexIndex, uint32 InChannelIndex) const override
 	{
-		return FVector2D(MeshLODData.StaticVertexBuffers.StaticMeshVertexBuffer.GetVertexUV(VertexIndex, ChannelIndex));
+		check(MeshData);
+		check(MeshData->LODRenderData.IsValidIndex(LODIndex));
+		return FVector2D(MeshData->LODRenderData[LODIndex].StaticVertexBuffers.StaticMeshVertexBuffer.GetVertexUV(InVertexIndex, InChannelIndex));
 	}
 
 	virtual void GetSectionFromVertexIndex(uint32 InVertIndex, int32& OutSectionIndex) const override
 	{
+		check(MeshData);
+		check(MeshData->LODRenderData.IsValidIndex(LODIndex));
 		int32 OutVertIndex = 0;
-		return MeshLODData.GetSectionFromVertexIndex(InVertIndex, OutSectionIndex, OutVertIndex);
+		return MeshData->LODRenderData[LODIndex].GetSectionFromVertexIndex(InVertIndex, OutSectionIndex, OutVertIndex);
 	}
 
 private:
-	FSkeletalMeshLODRenderData& MeshLODData;
+	int32 LODIndex = -1;
+	const FSkeletalMeshRenderData* MeshData = nullptr;
+
 	TArray<uint32> IndexBuffer;
 	TArray<FSkeletalMeshSection> Sections;
 };
@@ -257,16 +294,18 @@ class FSkeletalMeshData : public IMeshData
 {
 public:
 	FSkeletalMeshData(USkeletalMesh* InSkeletalMesh) 
-		: SkeletalMesh(InSkeletalMesh)
 	{
+		SkeletalMesh = InSkeletalMesh;
 		if (SkeletalMesh)
 		{
 			MeshData = SkeletalMesh->GetResourceForRendering();
 			if (MeshData)
 			{
-				for (FSkeletalMeshLODRenderData& MeshLODData : MeshData->LODRenderData)
+				const uint32 LODCount = SkeletalMesh->GetLODNum();
+				MeshesLODData.Reserve(LODCount);
+				for (uint32 LODIt=0; LODIt<LODCount; ++LODIt)
 				{
-					MeshesLODData.Add(FSkeletalMeshLODData(MeshLODData));
+					MeshesLODData.Add(FSkeletalMeshLODData(MeshData, LODIt));
 				}
 			}
 			else
@@ -286,14 +325,15 @@ public:
 		return SkeletalMesh->GetLODNum();
 	}
 
-	virtual const IMeshLODData& GetMeshLODData(uint32 LODIndex) const override
+	virtual const IMeshLODData& GetMeshLODData(uint32 InLODIndex) const override
 	{
-		return MeshesLODData[LODIndex];
+		check(MeshesLODData.IsValidIndex(InLODIndex));
+		return MeshesLODData[InLODIndex];
 	}
 
 private:
-	USkeletalMesh* SkeletalMesh;
-	FSkeletalMeshRenderData* MeshData;
+	const USkeletalMesh* SkeletalMesh = nullptr;
+	const FSkeletalMeshRenderData* MeshData = nullptr;
 	TArray<FSkeletalMeshLODData> MeshesLODData;
 };
 
@@ -659,7 +699,7 @@ namespace GroomBinding_RBFWeighting
 		}
 	}
 
-	void FillLocalValidPoints(const IMeshLODData& MeshLODData, const int32 TargetSection,
+	void FillLocalValidPoints(const GroomBinding_Mesh::IMeshLODData& MeshLODData, const int32 TargetSection,
 		const FHairStrandsRootData::FMeshProjectionLOD& ProjectionLOD, TArray<bool>& ValidPoints)
 	{
 		const TArray<uint32>& TriangleIndices = MeshLODData.GetIndexBuffer();
@@ -675,7 +715,7 @@ namespace GroomBinding_RBFWeighting
 			FHairStrandsRootUtils::UnpackTriangleIndex(EncodedTriangleId, TriangleIndex, SectionIndex);
 			if (!ValidSection || (ValidSection && (SectionIndex == TargetSection)))
 			{
-				const IMeshSectionData& Section = MeshLODData.GetSection(SectionIndex);
+				const GroomBinding_Mesh::IMeshSectionData& Section = MeshLODData.GetSection(SectionIndex);
 				for (uint32 VertexIt = 0; VertexIt < 3; ++VertexIt)
 				{
 					const uint32 VertexIndex = TriangleIndices[Section.GetBaseIndex() + 3 * TriangleIndex + VertexIt];
@@ -685,14 +725,14 @@ namespace GroomBinding_RBFWeighting
 		}
 	}
 
-	void FillGlobalValidPoints(const IMeshLODData& MeshLODData, const int32 TargetSection, TArray<bool>& ValidPoints)
+	void FillGlobalValidPoints(const GroomBinding_Mesh::IMeshLODData& MeshLODData, const int32 TargetSection, TArray<bool>& ValidPoints)
 	{
 		const TArray<uint32>& TriangleIndices = MeshLODData.GetIndexBuffer(); 
 		if (TargetSection >= 0 && TargetSection < MeshLODData.GetNumSections())
 		{
 			ValidPoints.Init(false, MeshLODData.GetNumVertices());
 
-			const IMeshSectionData& Section = MeshLODData.GetSection(TargetSection);
+			const GroomBinding_Mesh::IMeshSectionData& Section = MeshLODData.GetSection(TargetSection);
 			for (uint32 TriangleIt = 0; TriangleIt < Section.GetNumTriangles(); ++TriangleIt)
 			{
 				for (uint32 VertexIt = 0; VertexIt < 3; ++VertexIt)
@@ -722,7 +762,7 @@ namespace GroomBinding_RBFWeighting
 		const bool bNeedStrandsRoot,
 		const uint32 NumInterpolationPoints, 
 		const int32 MatchingSection, 
-		const IMeshData* MeshData, 
+		const GroomBinding_Mesh::IMeshData* MeshData, 
 		const TArray<TArray<FVector3f>>& TransferedPositions)
 	{
 		const uint32 MeshLODCount= MeshData->GetNumLODs();
@@ -730,7 +770,7 @@ namespace GroomBinding_RBFWeighting
 
 		for (uint32 LODIndex = 0; LODIndex < MeshLODCount; ++LODIndex)
 		{
-			const IMeshLODData& MeshLODData = MeshData->GetMeshLODData(LODIndex);
+			const GroomBinding_Mesh::IMeshLODData& MeshLODData = MeshData->GetMeshLODData(LODIndex);
 
 			int32 TargetSection = -1;
 			bool GlobalSamples = false;
@@ -1094,7 +1134,7 @@ namespace GroomBinding_RootProjection
 
 	static bool Project(
 		const FHairStrandsDatas& InStrandsData,
-		const IMeshData* InMeshData,
+		const GroomBinding_Mesh::IMeshData* InMeshData,
 		const TArray<TArray<FVector3f>>& InTransferredPositions,
 		FHairStrandsRootData& OutRootData)
 	{
@@ -1116,7 +1156,7 @@ namespace GroomBinding_RootProjection
 			check(LODIt == OutRootData.MeshProjectionLODs[LODIt].LODIndex);
 
 			// 2.1. Build a grid around the hair AABB
-			const IMeshLODData& MeshLODData = InMeshData->GetMeshLODData(LODIt);
+			const GroomBinding_Mesh::IMeshLODData& MeshLODData = InMeshData->GetMeshLODData(LODIt);
 			const TArray<uint32>& IndexBuffer = MeshLODData.GetIndexBuffer();
 
 			const uint32 MaxSectionCount = GetHairStrandsMaxSectionCount();
@@ -1137,7 +1177,7 @@ namespace GroomBinding_RootProjection
 			for (uint32 SectionIt = 0; SectionIt < SectionCount; ++SectionIt)
 			{
 				// 2.2.1 Compute the bounding box of the skeletal mesh
-				const IMeshSectionData& Section = MeshLODData.GetSection(SectionIt);
+				const GroomBinding_Mesh::IMeshSectionData& Section = MeshLODData.GetSection(SectionIt);
 				const uint32 TriangleCount = Section.GetNumTriangles();
 				const uint32 SectionBaseIndex = Section.GetBaseIndex();
 
@@ -1214,7 +1254,7 @@ namespace GroomBinding_RootProjection
 			for (uint32 SectionIt = 0; SectionIt < SectionCount; ++SectionIt)
 			{
 				// 2.2.2 Insert all triangle within the grid
-				const IMeshSectionData& Section = MeshLODData.GetSection(SectionIt);
+				const GroomBinding_Mesh::IMeshSectionData& Section = MeshLODData.GetSection(SectionIt);
 				const uint32 TriangleCount = Section.GetNumTriangles();
 				const uint32 SectionBaseIndex = Section.GetBaseIndex();
 
@@ -1677,8 +1717,8 @@ namespace GroomBinding_Transfer
 	}
 
 	bool Transfer(
-		const IMeshData* InSourceMeshData,
-		const IMeshData* InTargetMeshData,
+		const GroomBinding_Mesh::IMeshData* InSourceMeshData,
+		const GroomBinding_Mesh::IMeshData* InTargetMeshData,
 		TArray<TArray<FVector3f>>& OutTransferredPositions, const int32 MatchingSection)
 	{
 
@@ -1693,7 +1733,7 @@ namespace GroomBinding_Transfer
 			// Notes:
 			// LODs are transfered using the LOD0 of the source mesh, as the LOD count can mismatch between source and target meshes.
 			// Assume that the section 0 contains the head section, which is where the hair/facial hair should be projected on
-			const IMeshLODData& MeshLODData = InSourceMeshData->GetMeshLODData(InSourceLODIndex);
+			const GroomBinding_Mesh::IMeshLODData& MeshLODData = InSourceMeshData->GetMeshLODData(InSourceLODIndex);
 			const uint32 SourceTriangleCount = MeshLODData.GetSection(InSourceSectionId).GetNumTriangles();
 			const uint32 SourceSectionBaseIndex = MeshLODData.GetSection(InSourceSectionId).GetBaseIndex();
 
@@ -1730,7 +1770,7 @@ namespace GroomBinding_Transfer
 		// 1. Insert triangles into a 2D UV grid
 		const uint32 ChannelIndex = 0;
 		const uint32 SourceLODIndex = 0;
-		const IMeshLODData& SourceMeshLODData = InSourceMeshData->GetMeshLODData(SourceLODIndex);
+		const GroomBinding_Mesh::IMeshLODData& SourceMeshLODData = InSourceMeshData->GetMeshLODData(SourceLODIndex);
 		const bool bIsMatchingSectionValid = MatchingSection < SourceMeshLODData.GetNumSections();
 		const int32 SourceSectionId = bIsMatchingSectionValid ? MatchingSection : 0;
 		if (!bIsMatchingSectionValid && GHairStrandsBindingBuilderWarningEnable > 0)
@@ -1758,7 +1798,7 @@ namespace GroomBinding_Transfer
 			// If this is not the case, then fall back to section 0 and rebuild the source triangle grid to match the same section ID (1.)
 			int32 LocalSourceSectionId = SourceSectionId;
 			int32 LocalTargetSectionId = TargetSectionId;
-			const IMeshLODData& TargetMeshLODData = InTargetMeshData->GetMeshLODData(TargetLODIndex);
+			const GroomBinding_Mesh::IMeshLODData& TargetMeshLODData = InTargetMeshData->GetMeshLODData(TargetLODIndex);
 
 			if (LocalTargetSectionId >= TargetMeshLODData.GetNumSections())
 			{
@@ -2076,18 +2116,22 @@ static bool InternalBuildBinding_CPU(UGroomBindingAsset* BindingAsset, uint32 In
 		const int32 NumInterpolationPoints = BindingAsset->GetNumInterpolationPoints();
 		UGroomAsset* GroomAsset = BindingAsset->GetGroom();
 
-		TUniquePtr<IMeshData> SourceMeshData;
-		TUniquePtr<IMeshData> TargetMeshData;
+		TUniquePtr<GroomBinding_Mesh::IMeshData> SourceMeshData;
+		TUniquePtr<GroomBinding_Mesh::IMeshData> TargetMeshData;
 		if (BindingAsset->GetGroomBindingType() == EGroomBindingMeshType::SkeletalMesh)
 		{
+			check(!BindingAsset->GetTargetSkeletalMesh()->IsCompiling());
+			check(BindingAsset->GetTargetSkeletalMesh()->IsAsyncTaskComplete());
 			BindingAsset->GetTargetSkeletalMesh()->ConditionalPostLoad();
 			if (BindingAsset->GetSourceSkeletalMesh())
 			{
+				check(!BindingAsset->GetSourceSkeletalMesh()->IsCompiling());
+				check(BindingAsset->GetSourceSkeletalMesh()->IsAsyncTaskComplete());
 				BindingAsset->GetSourceSkeletalMesh()->ConditionalPostLoad();
 			}
 
-			SourceMeshData = TUniquePtr<FSkeletalMeshData, TDefaultDelete<IMeshData>>(new FSkeletalMeshData(BindingAsset->GetSourceSkeletalMesh()));
-			TargetMeshData = TUniquePtr<FSkeletalMeshData, TDefaultDelete<IMeshData>>(new FSkeletalMeshData(BindingAsset->GetTargetSkeletalMesh()));
+			SourceMeshData = TUniquePtr<GroomBinding_Mesh::FSkeletalMeshData, TDefaultDelete<GroomBinding_Mesh::IMeshData>>(new GroomBinding_Mesh::FSkeletalMeshData(BindingAsset->GetSourceSkeletalMesh()));
+			TargetMeshData = TUniquePtr<GroomBinding_Mesh::FSkeletalMeshData, TDefaultDelete<GroomBinding_Mesh::IMeshData>>(new GroomBinding_Mesh::FSkeletalMeshData(BindingAsset->GetTargetSkeletalMesh()));
 		}
 		else
 		{
@@ -2097,8 +2141,8 @@ static bool InternalBuildBinding_CPU(UGroomBindingAsset* BindingAsset, uint32 In
 				BindingAsset->GetSourceGeometryCache()->ConditionalPostLoad();
 			}
 
-			SourceMeshData = TUniquePtr<FGeometryCacheData, TDefaultDelete<IMeshData>>(new FGeometryCacheData(BindingAsset->GetSourceGeometryCache()));
-			TargetMeshData = TUniquePtr<FGeometryCacheData, TDefaultDelete<IMeshData>>(new FGeometryCacheData(BindingAsset->GetTargetGeometryCache()));
+			SourceMeshData = TUniquePtr<GroomBinding_Mesh::FGeometryCacheData, TDefaultDelete<GroomBinding_Mesh::IMeshData>>(new GroomBinding_Mesh::FGeometryCacheData(BindingAsset->GetSourceGeometryCache()));
+			TargetMeshData = TUniquePtr<GroomBinding_Mesh::FGeometryCacheData, TDefaultDelete<GroomBinding_Mesh::IMeshData>>(new GroomBinding_Mesh::FGeometryCacheData(BindingAsset->GetTargetGeometryCache()));
 		}
 
 		if (!TargetMeshData->IsValid())
