Fix missing cards data on multi-group /multi-geometry type grooms.

For LOD0/LOD1, grooms have both strands & cards data. The LOD settings are set to strands geometry type for LOD0/LOD1 . When building the groom data, the cards data were discarded due to the mismatch of geometry type. It works for group0 as a matter of facts but all groups >0 would have not data being built.

diff --git a/Engine/Plugins/Runtime/HairStrands/Source/HairStrandsCore/Private/GroomAsset.cpp b/Engine/Plugins/Runtime/HairStrands/Source/HairStrandsCore/Private/GroomAsset.cpp
--- a/Engine/Plugins/Runtime/HairStrands/Source/HairStrandsCore/Private/GroomAsset.cpp
+++ b/Engine/Plugins/Runtime/HairStrands/Source/HairStrandsCore/Private/GroomAsset.cpp
@@ -2997,7 +2997,7 @@ bool UGroomAsset::BuildCardsData(uint32 GroupIndex)
 			// #hair_todo: add finer culling granularity to rebuild only what is necessary
 			const FHairGroupsCardsSourceDescription* CachedDesc = SourceIt < CachedHairGroupsCards.Num() ? &CachedHairGroupsCards[SourceIt] : nullptr;
 			const bool bLODHasChanged = CachedDesc == nullptr || !(*CachedDesc == *Desc);
-			if (!bLODHasChanged && (bIsAlreadyBuilt[LODIt] || GroupsLOD.LODs[LODIt].GeometryType != EGroomGeometryType::Cards)) // build only if it's Cards type
+			if (!bLODHasChanged && bIsAlreadyBuilt[LODIt])
 			{
 				bDataBuilt |= bIsAlreadyBuilt[LODIt];
 				continue;
@@ -3239,7 +3239,7 @@ bool UGroomAsset::BuildMeshesData(uint32 GroupIndex)
 		{
 			const FHairGroupsMeshesSourceDescription* CachedDesc = SourceIt < CachedHairGroupsMeshes.Num() ? &CachedHairGroupsMeshes[SourceIt] : nullptr;
 			const bool bLODHasChanged = CachedDesc == nullptr || !(*CachedDesc == *Desc);
-			if (!bLODHasChanged && (bIsAlreadyBuilt[LODIt] || GroupsLOD.LODs[LODIt].GeometryType != EGroomGeometryType::Meshes)) // build only if it's Meshes type
+			if (!bLODHasChanged && bIsAlreadyBuilt[LODIt])
 			{
 				bDataBuilt |= bIsAlreadyBuilt[LODIt];
 				continue;
diff --git a/Engine/Plugins/Runtime/HairStrands/Source/HairStrandsCore/Private/GroomBuilder.cpp b/Engine/Plugins/Runtime/HairStrands/Source/HairStrandsCore/Private/GroomBuilder.cpp
--- a/Engine/Plugins/Runtime/HairStrands/Source/HairStrandsCore/Private/GroomBuilder.cpp
+++ b/Engine/Plugins/Runtime/HairStrands/Source/HairStrandsCore/Private/GroomBuilder.cpp
@@ -40,7 +40,7 @@ static FAutoConsoleVariableRef CVarHairGroupIndexBuilder_MaxVoxelResolution(TEXT
 
 FString FGroomBuilder::GetVersion()
 {
-	return TEXT("v8r45");
+	return TEXT("v8r46");
 }
 
 namespace FHairStrandsDecimation
