Fix groom cluster visualizazation

diff --git a/Engine/Plugins/Runtime/HairStrands/Source/HairStrandsCore/Private/GroomBuilder.cpp b/Engine/Plugins/Runtime/HairStrands/Source/HairStrandsCore/Private/GroomBuilder.cpp
--- a/Engine/Plugins/Runtime/HairStrands/Source/HairStrandsCore/Private/GroomBuilder.cpp
+++ b/Engine/Plugins/Runtime/HairStrands/Source/HairStrandsCore/Private/GroomBuilder.cpp
@@ -40,7 +40,7 @@ static FAutoConsoleVariableRef CVarHairGroupIndexBuilder_MaxVoxelResolution(TEXT
 
 FString FGroomBuilder::GetVersion()
 {
-	return TEXT("v8r46");
+	return TEXT("v8r50");
 }
 
 namespace FHairStrandsDecimation
@@ -3176,17 +3176,13 @@ static void BuildClusterData(
 	Out.CurveToClusterIds.SetNum(InRenStrandsData.GetNumCurves());
 	Out.ClusterInfos.SetNum(Out.ClusterCount);
 	const uint32 CurvePerCluster = FMath::DivideAndRoundUp(InRenStrandsData.GetNumCurves(), Out.ClusterCount);
+	// Every 'CurvePerCluster' is assigned to the current cluster
+	for (uint32 CurveIt = 0, CurveCount=InRenStrandsData.GetNumCurves(); CurveIt < CurveCount; ++CurveIt)
+	{
+		Out.CurveToClusterIds[CurveIt] = CurveIt % CurvePerCluster;
+	}
 	for (uint32 ClusterIt = 0; ClusterIt < Out.ClusterCount; ++ClusterIt)
 	{
-		// Every 'CurvePerCluster' is assigned to the current cluster
-		for (uint32 CurveIt = 0; CurveIt < CurvePerCluster; ++CurveIt)
-		{
-			const uint32 CurveIndex = CurveIt * CurvePerCluster;
-			if (CurveIndex < InRenStrandsData.GetNumCurves())
-			{
-				Out.CurveToClusterIds[CurveIndex] = ClusterIt;
-			}
-		}
 		Out.ClusterInfos[ClusterIt].LODCount = LODCount;
 		for (uint8 LODIt = 0; LODIt < LODCount; ++LODIt)
 		{
