Fix missing groom cluster visualization in cook build.

diff --git a/Engine/Plugins/Runtime/HairStrands/Source/HairStrandsCore/Private/GroomComponent.cpp b/Engine/Plugins/Runtime/HairStrands/Source/HairStrandsCore/Private/GroomComponent.cpp
--- a/Engine/Plugins/Runtime/HairStrands/Source/HairStrandsCore/Private/GroomComponent.cpp
+++ b/Engine/Plugins/Runtime/HairStrands/Source/HairStrandsCore/Private/GroomComponent.cpp
@@ -112,10 +112,7 @@ FHairGroupInstance::~FHairGroupInstance()
 			InternalResourceRelease(Strands.RenRaytracingResource);
 		}
 #endif
-
-#if WITH_EDITOR
 		Strands.DebugCurveAttributeBuffer.Release();
-#endif
 
 		InternalResourceRelease(Strands.VertexFactory);
 	}
@@ -739,7 +736,7 @@ public:
 					check(HairGroupInstances[GroupIt]->GetRefCount() > 0);
 
 					FMaterialRenderProxy* Debug_MaterialProxy = nullptr;
-					const EGroomViewMode ViewMode = AllowDebugViewmodes() ? GetGroomViewMode(*View) : EGroomViewMode::None;
+					const EGroomViewMode ViewMode = GetGroomViewMode(*View);
 					bool bNeedDebugMaterial = false;
 					switch(ViewMode)
 					{
diff --git a/Engine/Plugins/Runtime/HairStrands/Source/HairStrandsCore/Private/HairStrandsInterpolation.cpp b/Engine/Plugins/Runtime/HairStrands/Source/HairStrandsCore/Private/HairStrandsInterpolation.cpp
--- a/Engine/Plugins/Runtime/HairStrands/Source/HairStrandsCore/Private/HairStrandsInterpolation.cpp
+++ b/Engine/Plugins/Runtime/HairStrands/Source/HairStrandsCore/Private/HairStrandsInterpolation.cpp
@@ -746,7 +746,7 @@ class FHairPatchAttributeCS : public FGlobalShader
 	END_SHADER_PARAMETER_STRUCT()
 
 public:
-	static bool ShouldCompilePermutation(const FGlobalShaderPermutationParameters& Parameters) { return IsHairStrandsSupported(EHairStrandsShaderType::All, Parameters.Platform); }
+	static bool ShouldCompilePermutation(const FGlobalShaderPermutationParameters& Parameters) { return IsHairStrandsSupported(EHairStrandsShaderType::Strands, Parameters.Platform); }
 	static void ModifyCompilationEnvironment(const FGlobalShaderPermutationParameters& Parameters, FShaderCompilerEnvironment& OutEnvironment)
 	{
 		FGlobalShader::ModifyCompilationEnvironment(Parameters, OutEnvironment);
@@ -1768,10 +1768,8 @@ void ComputeHairStrandsInterpolation(
 
 			// "WITH_EDITOR && PatchMode == EHairPatchAttribute::GuideInflucence" is a special path when visualizing groom within the groom editor, so that we can diplay guides even when there is no simulation or global interpolation enabled
 			EHairPatchAttribute PatchMode = EHairPatchAttribute::None;
-			#if WITH_EDITOR
 			if (ViewMode == EGroomViewMode::RenderHairStrands && bNeedDeformation)				{ PatchMode = EHairPatchAttribute::GuideInflucence; }
 			if (ViewMode == EGroomViewMode::Cluster || ViewMode == EGroomViewMode::ClusterAABB)	{ PatchMode = EHairPatchAttribute::ClusterInfluence; }
-			#endif
 
 			// 2. Deform hair if needed (e.g.: skinning, simulation, RBF) and recompute tangent
 			FRDGBufferSRVRef Strands_PositionSRV = nullptr;
@@ -2023,7 +2021,6 @@ void ComputeHairStrandsInterpolation(
 			}
 
 			// 2.1 Patch attribute for debug visualization (guide influence or clusters visualization)
-			#if WITH_EDITOR
 			if (PatchMode != EHairPatchAttribute::None)
 			{
 				// Create an debug buffer for storing cluster visualization data. This is only used for debug purpose, hence only enable in editor build.
@@ -2049,7 +2046,6 @@ void ComputeHairStrandsInterpolation(
 					bValidGuide ? RegisterAsSRV(GraphBuilder, Instance->Strands.InterpolationResource->InterpolationBuffer) : nullptr,
 					OutRenCurveAttributeBuffer);
 			}
-			#endif
 
 			// 2.2 Update the VF input with the update resources
 			Instance->HairGroupPublicData->VFInput = InternalComputeHairStrandsVertexInputData(&GraphBuilder, Instance, ViewMode);
