Fix typo: GroupResrouces->GroupResources

diff --git a/Engine/Plugins/Runtime/HairStrands/Source/HairStrandsCore/Private/GroomBindingBuilder.cpp b/Engine/Plugins/Runtime/HairStrands/Source/HairStrandsCore/Private/GroomBindingBuilder.cpp
--- a/Engine/Plugins/Runtime/HairStrands/Source/HairStrandsCore/Private/GroomBindingBuilder.cpp
+++ b/Engine/Plugins/Runtime/HairStrands/Source/HairStrandsCore/Private/GroomBindingBuilder.cpp
@@ -2300,9 +2300,9 @@ static bool InternalBuildBinding_CPU(UGroomBindingAsset* BindingAsset, uint32 In
 	UGroomBindingAsset::FHairGroupResources& OutHairGroupResources = BindingAsset->GetHairGroupResources();
 	if (OutHairGroupResources.Num() > 0)
 	{
-		for (UGroomBindingAsset::FHairGroupResource& GroupResrouces : OutHairGroupResources)
+		for (UGroomBindingAsset::FHairGroupResource& GroupResources : OutHairGroupResources)
 		{
-			BindingAsset->AddHairGroupResourcesToDelete(GroupResrouces);
+			BindingAsset->AddHairGroupResourcesToDelete(GroupResources);
 		}
 		OutHairGroupResources.Empty();
 	}
