Fix typo causing groom binding to be rebuilt on CitySample.

diff --git a/Engine/Plugins/Runtime/HairStrands/Source/HairStrandsCore/Private/GroomBindingAsset.cpp b/Engine/Plugins/Runtime/HairStrands/Source/HairStrandsCore/Private/GroomBindingAsset.cpp
--- a/Engine/Plugins/Runtime/HairStrands/Source/HairStrandsCore/Private/GroomBindingAsset.cpp
+++ b/Engine/Plugins/Runtime/HairStrands/Source/HairStrandsCore/Private/GroomBindingAsset.cpp
@@ -70,7 +70,7 @@ static void InternalSerializeStrands(FArchive& Ar, UObject* Owner, FHairStrandsR
 
 		// Pre-warm DDC cache
 		#if WITH_EDITORONLY_DATA
-		const bool bPreWarmCache = IsLoading() && bHeader && !bData;
+		const bool bPreWarmCache = Ar.IsLoading() && bHeader && !bData;
 		if (bPreWarmCache)
 		{
 			bool bHasDataInCache = true;
