Fix similar typo issue on GroomAsset.

diff --git a/Engine/Plugins/Runtime/HairStrands/Source/HairStrandsCore/Private/GroomAsset.cpp b/Engine/Plugins/Runtime/HairStrands/Source/HairStrandsCore/Private/GroomAsset.cpp
--- a/Engine/Plugins/Runtime/HairStrands/Source/HairStrandsCore/Private/GroomAsset.cpp
+++ b/Engine/Plugins/Runtime/HairStrands/Source/HairStrandsCore/Private/GroomAsset.cpp
@@ -1879,7 +1879,7 @@ static void InternalSerializeStrand(FArchive& Ar, UObject* Owner, FHairGroupPlat
 
 		#if WITH_EDITORONLY_DATA
 		// Pre-warm DDC cache
-		const bool bPreWarmCache = IsLoading() && bHeader && !bData;
+		const bool bPreWarmCache = Ar.IsLoading() && bHeader && !bData;
 		if (bPreWarmCache)
 		{
 			bool bHasDataInCache = true;
