Optimize groom asset & groom binding asset building.

* Made certain costly operation run paralelle.
* Remove unnecessary double computation
* Made certain data processing more cache friendly

Test on Citysample timing 47s -> 10s.

diff --git a/Engine/Plugins/Runtime/HairStrands/Source/HairStrandsCore/Private/GroomBindingBuilder.cpp b/Engine/Plugins/Runtime/HairStrands/Source/HairStrandsCore/Private/GroomBindingBuilder.cpp
--- a/Engine/Plugins/Runtime/HairStrands/Source/HairStrandsCore/Private/GroomBindingBuilder.cpp
+++ b/Engine/Plugins/Runtime/HairStrands/Source/HairStrandsCore/Private/GroomBindingBuilder.cpp
@@ -141,7 +141,7 @@ public:
 	virtual const int32 GetNumSections() const = 0;
 	virtual const IMeshSectionData& GetSection(uint32 SectionIndex) const = 0;
 	virtual const FVector3f& GetVertexPosition(uint32 VertexIndex) const = 0;
-	virtual FVector2D GetVertexUV(uint32 VertexIndex, uint32 ChannelIndex) const = 0;
+	virtual FVector2f GetVertexUV(uint32 VertexIndex, uint32 ChannelIndex) const = 0;
 	virtual void GetSectionFromVertexIndex(uint32 InVertIndex, int32& OutSectionIndex) const = 0;
 	virtual ~IMeshLODData() {}
 };
@@ -267,11 +267,11 @@ public:
 		return MeshData->LODRenderData[LODIndex].StaticVertexBuffers.PositionVertexBuffer.VertexPosition(InVertexIndex);
 	}
 
-	virtual FVector2D GetVertexUV(uint32 InVertexIndex, uint32 InChannelIndex) const override
+	virtual FVector2f GetVertexUV(uint32 InVertexIndex, uint32 InChannelIndex) const override
 	{
 		check(MeshData);
 		check(MeshData->LODRenderData.IsValidIndex(LODIndex));
-		return FVector2D(MeshData->LODRenderData[LODIndex].StaticVertexBuffers.StaticMeshVertexBuffer.GetVertexUV(InVertexIndex, InChannelIndex));
+		return FVector2f(MeshData->LODRenderData[LODIndex].StaticVertexBuffers.StaticMeshVertexBuffer.GetVertexUV(InVertexIndex, InChannelIndex));
 	}
 
 	virtual void GetSectionFromVertexIndex(uint32 InVertIndex, int32& OutSectionIndex) const override
@@ -465,9 +465,9 @@ public:
 		return MeshData.Positions[VertexIndex];
 	}
 
-	virtual FVector2D GetVertexUV(uint32 VertexIndex, uint32 ChannelIndex) const override
+	virtual FVector2f GetVertexUV(uint32 VertexIndex, uint32 ChannelIndex) const override
 	{
-		return FVector2D(MeshData.TextureCoordinates[VertexIndex]);
+		return FVector2f(MeshData.TextureCoordinates[VertexIndex]);
 	}
 
 	virtual void GetSectionFromVertexIndex(uint32 InVertIndex, int32& OutSectionIndex) const override
@@ -852,13 +852,13 @@ namespace GroomBinding_RootProjection
 			uint32  I1;
 			uint32  I2;
 
-			FVector P0;
-			FVector P1;
-			FVector P2;
+			FVector3f P0;
+			FVector3f P1;
+			FVector3f P2;
 
-			FVector2D UV0;
-			FVector2D UV1;
-			FVector2D UV2;
+			FVector2f UV0;
+			FVector2f UV1;
+			FVector2f UV2;
 		};
 
 		struct FCell
@@ -867,16 +867,16 @@ namespace GroomBinding_RootProjection
 		};
 		typedef TArray<const FCell*> FCells;
 
-		FTriangleGrid(const FVector& InMinBound, const FVector& InMaxBound, float InVoxelWorldSize)
+		FTriangleGrid(const FVector3f& InMinBound, const FVector3f& InMaxBound, float InVoxelWorldSize)
 		{
 			MinBound = InMinBound;
 			MaxBound = InMaxBound;
 
 			// Compute the voxel volume resolution, and snap the max bound to the voxel grid
 			GridResolution = FIntVector::ZeroValue;
-			FVector VoxelResolutionF = (MaxBound - MinBound) / InVoxelWorldSize;
+			FVector3f VoxelResolutionF = (MaxBound - MinBound) / InVoxelWorldSize;
 			GridResolution = FIntVector(FMath::CeilToInt(VoxelResolutionF.X), FMath::CeilToInt(VoxelResolutionF.Y), FMath::CeilToInt(VoxelResolutionF.Z));
-			MaxBound = MinBound + FVector(GridResolution) * InVoxelWorldSize;
+			MaxBound = MinBound + FVector3f(GridResolution) * InVoxelWorldSize;
 
 			Cells.SetNum(GridResolution.X * GridResolution.Y * GridResolution.Z);
 		}
@@ -889,7 +889,7 @@ namespace GroomBinding_RootProjection
 				0 <= P.Z && P.Z < GridResolution.Z;
 		}
 
-		FORCEINLINE bool IsOutside(const FVector& MinP, const FVector& MaxP) const
+		FORCEINLINE bool IsOutside(const FVector3f& MinP, const FVector3f& MaxP) const
 		{
 			return
 				(MaxP.X <= MinBound.X || MaxP.Y <= MinBound.Y || MaxP.Z <= MinBound.Z) ||
@@ -905,10 +905,10 @@ namespace GroomBinding_RootProjection
 				FMath::Clamp(CellCoord.Z, 0, GridResolution.Z - 1));
 		}
 
-		FORCEINLINE FIntVector ToCellCoord(const FVector& P) const
+		FORCEINLINE FIntVector ToCellCoord(const FVector3f& P) const
 		{
 			bool bIsValid = false;
-			const FVector F = ((P - MinBound) / (MaxBound - MinBound));
+			const FVector3f F = ((P - MinBound) / (MaxBound - MinBound));
 			const FIntVector CellCoord = FIntVector(FMath::FloorToInt(F.X * GridResolution.X), FMath::FloorToInt(F.Y * GridResolution.Y), FMath::FloorToInt(F.Z * GridResolution.Z));
 			return ClampToVolume(CellCoord, bIsValid);
 		}
@@ -920,7 +920,7 @@ namespace GroomBinding_RootProjection
 			return CellIndex;
 		}
 
-		FCells ToCells(const FVector& P)
+		FCells ToCells(const FVector3f& P)
 		{
 			FCells Out;
 
@@ -973,14 +973,14 @@ namespace GroomBinding_RootProjection
 
 		bool IsTriangleValid(const FTriangle& T) const
 		{
-			const FVector A = T.P0;
-			const FVector B = T.P1;
-			const FVector C = T.P2;
+			const FVector3f A = T.P0;
+			const FVector3f B = T.P1;
+			const FVector3f C = T.P2;
 
-			const FVector AB = B - A;
-			const FVector AC = C - A;
-			const FVector BC = B - C;
-			return FVector::DotProduct(AB, AB) > 0 && FVector::DotProduct(AC, AC) > 0 && FVector::DotProduct(BC, BC) > 0;
+			const FVector3f AB = B - A;
+			const FVector3f AC = C - A;
+			const FVector3f BC = B - C;
+			return FVector3f::DotProduct(AB, AB) > 0 && FVector3f::DotProduct(AC, AC) > 0 && FVector3f::DotProduct(BC, BC) > 0;
 		}
 
 		bool Insert(const FTriangle& T)
@@ -990,12 +990,12 @@ namespace GroomBinding_RootProjection
 				return false;
 			}
 
-			FVector TriMinBound;
+			FVector3f TriMinBound;
 			TriMinBound.X = FMath::Min(T.P0.X, FMath::Min(T.P1.X, T.P2.X));
 			TriMinBound.Y = FMath::Min(T.P0.Y, FMath::Min(T.P1.Y, T.P2.Y));
 			TriMinBound.Z = FMath::Min(T.P0.Z, FMath::Min(T.P1.Z, T.P2.Z));
 
-			FVector TriMaxBound;
+			FVector3f TriMaxBound;
 			TriMaxBound.X = FMath::Max(T.P0.X, FMath::Max(T.P1.X, T.P2.X));
 			TriMaxBound.Y = FMath::Max(T.P0.Y, FMath::Max(T.P1.Y, T.P2.Y));
 			TriMaxBound.Z = FMath::Max(T.P0.Z, FMath::Max(T.P1.Z, T.P2.Z));
@@ -1029,8 +1029,8 @@ namespace GroomBinding_RootProjection
 			return bInserted;
 		}
 
-		FVector MinBound;
-		FVector MaxBound;
+		FVector3f MinBound;
+		FVector3f MaxBound;
 		FIntVector GridResolution;
 		TArray<FCell> Cells;
 	};
@@ -1039,39 +1039,39 @@ namespace GroomBinding_RootProjection
 	// Code from the book "Real-Time Collision Detection" by Christer Ericson
 	struct FTrianglePoint
 	{
-		FVector P;
-		FVector Barycentric;
+		FVector3f P;
+		FVector3f Barycentric;
 	};
 
-	static FTrianglePoint ComputeClosestPoint(const FTriangleGrid::FTriangle& Tri, const FVector& P)
+	static FTrianglePoint ComputeClosestPoint(const FTriangleGrid::FTriangle& Tri, const FVector3f& P)
 	{
-		const FVector A = Tri.P0;
-		const FVector B = Tri.P1;
-		const FVector C = Tri.P2;
+		const FVector3f A = Tri.P0;
+		const FVector3f B = Tri.P1;
+		const FVector3f C = Tri.P2;
 
 		// Check if P is in vertex region outside A.
-		FVector AB = B - A;
-		FVector AC = C - A;
-		FVector AP = P - A;
-		float D1 = FVector::DotProduct(AB, AP);
-		float D2 = FVector::DotProduct(AC, AP);
+		FVector3f AB = B - A;
+		FVector3f AC = C - A;
+		FVector3f AP = P - A;
+		float D1 = FVector3f::DotProduct(AB, AP);
+		float D2 = FVector3f::DotProduct(AC, AP);
 		if (D1 <= 0.f && D2 <= 0.f)
 		{
 			FTrianglePoint Out;
 			Out.P = A;
-			Out.Barycentric = FVector(1, 0, 0);
+			Out.Barycentric = FVector3f(1, 0, 0);
 			return Out;
 		}
 
 		// Check if P is in vertex region outside B.
-		FVector BP = P - B;
-		float D3 = FVector::DotProduct(AB, BP);
-		float D4 = FVector::DotProduct(AC, BP);
+		FVector3f BP = P - B;
+		float D3 = FVector3f::DotProduct(AB, BP);
+		float D4 = FVector3f::DotProduct(AC, BP);
 		if (D3 >= 0.f && D4 <= D3)
 		{
 			FTrianglePoint Out;
 			Out.P = B;
-			Out.Barycentric = FVector(0, 1, 0);
+			Out.Barycentric = FVector3f(0, 1, 0);
 			return Out;
 		}
 
@@ -1083,19 +1083,19 @@ namespace GroomBinding_RootProjection
 
 			FTrianglePoint Out;
 			Out.P = A + V * AB;
-			Out.Barycentric = FVector(1 - V, V, 0);
+			Out.Barycentric = FVector3f(1 - V, V, 0);
 			return Out;
 		}
 
 		// Check if P is in vertex region outside C.
-		FVector CP = P - C;
-		float D5 = FVector::DotProduct(AB, CP);
-		float D6 = FVector::DotProduct(AC, CP);
+		FVector3f CP = P - C;
+		float D5 = FVector3f::DotProduct(AB, CP);
+		float D6 = FVector3f::DotProduct(AC, CP);
 		if (D6 >= 0.f && D5 <= D6)
 		{
 			FTrianglePoint Out;
 			Out.P = C;
-			Out.Barycentric = FVector(0, 0, 1);
+			Out.Barycentric = FVector3f(0, 0, 1);
 			return Out;
 		}
 
@@ -1106,7 +1106,7 @@ namespace GroomBinding_RootProjection
 			float W = D2 / (D2 - D6);
 			FTrianglePoint Out;
 			Out.P = A + W * AC;
-			Out.Barycentric = FVector(1 - W, 0, W);
+			Out.Barycentric = FVector3f(1 - W, 0, W);
 			return Out;
 		}
 
@@ -1117,7 +1117,7 @@ namespace GroomBinding_RootProjection
 			float W = (D4 - D3) / (D4 - D3 + D5 - D6);
 			FTrianglePoint Out;
 			Out.P = B + W * (C - B);
-			Out.Barycentric = FVector(0, 1 - W, W);
+			Out.Barycentric = FVector3f(0, 1 - W, W);
 			return Out;
 		}
 
@@ -1128,7 +1128,7 @@ namespace GroomBinding_RootProjection
 
 		FTrianglePoint Out;
 		Out.P = A + AB * V + AC * W;
-		Out.Barycentric = FVector(1 - V - W, V, W);
+		Out.Barycentric = FVector3f(1 - V - W, V, W);
 		return Out;
 	}
 
@@ -1162,7 +1162,7 @@ namespace GroomBinding_RootProjection
 			const uint32 MaxSectionCount = GetHairStrandsMaxSectionCount();
 			const uint32 MaxTriangleCount = GetHairStrandsMaxTriangleCount();
 
-			FBox MeshBound;
+			FBox3f MeshBound;
 			MeshBound.Init();
 			const uint32 SectionCount = MeshLODData.GetNumSections();
 			
@@ -1198,15 +1198,15 @@ namespace GroomBinding_RootProjection
 
 					if (bHasTransferredPosition)
 					{
-						T.P0 = (FVector)InTransferredPositions[LODIt][T.I0];
-						T.P1 = (FVector)InTransferredPositions[LODIt][T.I1];
-						T.P2 = (FVector)InTransferredPositions[LODIt][T.I2];
+						T.P0 = InTransferredPositions[LODIt][T.I0];
+						T.P1 = InTransferredPositions[LODIt][T.I1];
+						T.P2 = InTransferredPositions[LODIt][T.I2];
 					}
 					else
 					{
-						T.P0 = (FVector)MeshLODData.GetVertexPosition(T.I0);
-						T.P1 = (FVector)MeshLODData.GetVertexPosition(T.I1);
-						T.P2 = (FVector)MeshLODData.GetVertexPosition(T.I2);
+						T.P0 = MeshLODData.GetVertexPosition(T.I0);
+						T.P1 = MeshLODData.GetVertexPosition(T.I1);
+						T.P2 = MeshLODData.GetVertexPosition(T.I2);
 					}
 
 					T.UV0 = MeshLODData.GetVertexUV(T.I0, ChannelIndex);
@@ -1225,10 +1225,10 @@ namespace GroomBinding_RootProjection
 			}
 
 			// Take the smallest bounding box between the groom and the skeletal mesh
-			const FVector MeshExtent = MeshBound.Max - MeshBound.Min;
-			const FVector HairExtent = InStrandsData.BoundingBox.Max - InStrandsData.BoundingBox.Min;
-			FVector GridMin;
-			FVector GridMax;
+			const FVector3f MeshExtent = MeshBound.Max - MeshBound.Min;
+			const FVector3f HairExtent = InStrandsData.BoundingBox.Max - InStrandsData.BoundingBox.Min;
+			FVector3f GridMin;
+			FVector3f GridMax;
 			if (MeshExtent.Size() < HairExtent.Size())
 			{
 				GridMin = MeshBound.Min;
@@ -1244,8 +1244,8 @@ namespace GroomBinding_RootProjection
 				// the groom bound to be correctly inserted.
 				if (ClosestTrianglePoint < FLT_MAX)
 				{
-					GridMin -= FVector(ClosestTrianglePoint * 1.25f);
-					GridMax += FVector(ClosestTrianglePoint * 1.25f);
+					GridMin -= FVector3f(ClosestTrianglePoint * 1.25f);
+					GridMax += FVector3f(ClosestTrianglePoint * 1.25f);
 				}
 			}
 
@@ -1275,15 +1275,15 @@ namespace GroomBinding_RootProjection
 
 					if (bHasTransferredPosition)
 					{
-						T.P0 = (FVector)InTransferredPositions[LODIt][T.I0];
-						T.P1 = (FVector)InTransferredPositions[LODIt][T.I1];
-						T.P2 = (FVector)InTransferredPositions[LODIt][T.I2];
+						T.P0 = InTransferredPositions[LODIt][T.I0];
+						T.P1 = InTransferredPositions[LODIt][T.I1];
+						T.P2 = InTransferredPositions[LODIt][T.I2];
 					}
 					else
 					{
-						T.P0 = (FVector)MeshLODData.GetVertexPosition(T.I0);
-						T.P1 = (FVector)MeshLODData.GetVertexPosition(T.I1);
-						T.P2 = (FVector)MeshLODData.GetVertexPosition(T.I2);
+						T.P0 = MeshLODData.GetVertexPosition(T.I0);
+						T.P1 = MeshLODData.GetVertexPosition(T.I1);
+						T.P2 = MeshLODData.GetVertexPosition(T.I2);
 					}
 
 					T.UV0 = MeshLODData.GetVertexUV(T.I0, ChannelIndex);
@@ -1329,7 +1329,7 @@ namespace GroomBinding_RootProjection
 		#endif
 			{
 				const uint32 Offset = InStrandsData.StrandsCurves.CurvesOffset[CurveIndex];
-				const FVector& RootP = (FVector)InStrandsData.StrandsPoints.PointsPosition[Offset];
+				const FVector3f& RootP = InStrandsData.StrandsPoints.PointsPosition[Offset];
 				const FTriangleGrid::FCells Cells = Grid.ToCells(RootP);
 
 				if (Cells.Num() == 0)
@@ -1343,18 +1343,18 @@ namespace GroomBinding_RootProjection
 
 				float ClosestDistance = FLT_MAX;
 				FTriangleGrid::FTriangle ClosestTriangle;
-				FVector2D ClosestBarycentrics;
+				FVector2f ClosestBarycentrics;
 				for (const FTriangleGrid::FCell* Cell : Cells)
 				{
 					for (const FTriangleGrid::FTriangle& CellTriangle : Cell->Triangles)
 					{
 						const FTrianglePoint Tri = ComputeClosestPoint(CellTriangle, RootP);
-						const float Distance = FVector::Distance(Tri.P, RootP);
+						const float Distance = FVector3f::Distance(Tri.P, RootP);
 						if (Distance < ClosestDistance)
 						{
 							ClosestDistance = Distance;
 							ClosestTriangle = CellTriangle;
-							ClosestBarycentrics = FVector2D(Tri.Barycentric.X, Tri.Barycentric.Y);
+							ClosestBarycentrics = FVector2f(Tri.Barycentric.X, Tri.Barycentric.Y);
 						}
 					}
 				}
@@ -1463,13 +1463,13 @@ namespace GroomBinding_Transfer
 			uint32  I1;
 			uint32  I2;
 
-			FVector P0;
-			FVector P1;
-			FVector P2;
+			FVector3f P0;
+			FVector3f P1;
+			FVector3f P2;
 
-			FVector2D UV0;
-			FVector2D UV1;
-			FVector2D UV2;
+			FVector2f UV0;
+			FVector2f UV1;
+			FVector2f UV2;
 		};
 
 		struct FCell
@@ -1482,8 +1482,8 @@ namespace GroomBinding_Transfer
 		{
 			GridResolution.X = Resolution;
 			GridResolution.Y = Resolution;
-			MinBound = FVector2D(0,0);
-			MaxBound = FVector2D(1,1);
+			MinBound = FVector2f(0,0);
+			MaxBound = FVector2f(1,1);
 
 			Cells.SetNum(GridResolution.X * GridResolution.Y);
 		}
@@ -1501,7 +1501,7 @@ namespace GroomBinding_Transfer
 				0 <= P.Y && P.Y < GridResolution.Y;
 		}
 
-		FORCEINLINE bool IsOutside(const FVector2D& MinP, const FVector2D& MaxP) const
+		FORCEINLINE bool IsOutside(const FVector2f& MinP, const FVector2f& MaxP) const
 		{
 			return
 				(MaxP.X <= MinBound.X || MaxP.Y <= MinBound.Y) ||
@@ -1516,10 +1516,10 @@ namespace GroomBinding_Transfer
 				FMath::Clamp(CellCoord.Y, 0, GridResolution.Y - 1));
 		}
 
-		FORCEINLINE FIntPoint ToCellCoord(const FVector2D& P) const
+		FORCEINLINE FIntPoint ToCellCoord(const FVector2f& P) const
 		{
 			bool bIsValid = false;
-			FVector2D PP;
+			FVector2f PP;
 			PP.X = FMath::Clamp(P.X, 0.f, 1.f);
 			PP.Y = FMath::Clamp(P.Y, 0.f, 1.f);
 			const FIntPoint CellCoord = FIntPoint(FMath::FloorToInt(PP.X * GridResolution.X), FMath::FloorToInt(PP.Y * GridResolution.Y));
@@ -1533,7 +1533,7 @@ namespace GroomBinding_Transfer
 			return CellIndex;
 		}
 
-		FCells ToCells(const FVector2D& P)
+		FCells ToCells(const FVector2f& P)
 		{
 			FCells Out;
 
@@ -1577,11 +1577,11 @@ namespace GroomBinding_Transfer
 
 		bool Insert(const FTriangle& T)
 		{
-			FVector2D TriMinBound;
+			FVector2f TriMinBound;
 			TriMinBound.X = FMath::Min(T.UV0.X, FMath::Min(T.UV1.X, T.UV2.X));
 			TriMinBound.Y = FMath::Min(T.UV0.Y, FMath::Min(T.UV1.Y, T.UV2.Y));
 
-			FVector2D TriMaxBound;
+			FVector2f TriMaxBound;
 			TriMaxBound.X = FMath::Max(T.UV0.X, FMath::Max(T.UV1.X, T.UV2.X));
 			TriMaxBound.Y = FMath::Max(T.UV0.Y, FMath::Max(T.UV1.Y, T.UV2.Y));
 
@@ -1611,8 +1611,8 @@ namespace GroomBinding_Transfer
 			return bInserted;
 		}
 
-		FVector2D MinBound;
-		FVector2D MaxBound;
+		FVector2f MinBound;
+		FVector2f MaxBound;
 		FIntPoint GridResolution;
 		TArray<FCell> Cells;
 	};
@@ -1622,40 +1622,40 @@ namespace GroomBinding_Transfer
 	// Code from the book "Real-Time Collision Detection" by Christer Ericson
 	struct FTrianglePoint
 	{
-		FVector P;
-		FVector Barycentric;
+		FVector3f P;
+		FVector3f Barycentric;
 	};
 
-	FTrianglePoint ComputeClosestPoint(const FVector2D& TriUV0, const FVector2D& TriUV1, const FVector2D& TriUV2, const FVector2D& UVs)
+	FTrianglePoint ComputeClosestPoint(const FVector2f& TriUV0, const FVector2f& TriUV1, const FVector2f& TriUV2, const FVector2f& UVs)
 	{
-		const FVector A = FVector(TriUV0, 0);
-		const FVector B = FVector(TriUV1, 0);
-		const FVector C = FVector(TriUV2, 0);
-		const FVector P = FVector(UVs, 0);
+		const FVector3f A = FVector3f(TriUV0, 0);
+		const FVector3f B = FVector3f(TriUV1, 0);
+		const FVector3f C = FVector3f(TriUV2, 0);
+		const FVector3f P = FVector3f(UVs, 0);
 
 		// Check if P is in vertex region outside A.
-		FVector AB = B - A;
-		FVector AC = C - A;
-		FVector AP = P - A;
-		float D1 = FVector::DotProduct(AB, AP);
-		float D2 = FVector::DotProduct(AC, AP);
+		FVector3f AB = B - A;
+		FVector3f AC = C - A;
+		FVector3f AP = P - A;
+		float D1 = FVector3f::DotProduct(AB, AP);
+		float D2 = FVector3f::DotProduct(AC, AP);
 		if (D1 <= 0.f && D2 <= 0.f)
 		{
 			FTrianglePoint Out;
 			Out.P = A;
-			Out.Barycentric = FVector(1, 0, 0);
+			Out.Barycentric = FVector3f(1, 0, 0);
 			return Out;
 		}
 
 		// Check if P is in vertex region outside B.
-		FVector BP = P - B;
-		float D3 = FVector::DotProduct(AB, BP);
-		float D4 = FVector::DotProduct(AC, BP);
+		FVector3f BP = P - B;
+		float D3 = FVector3f::DotProduct(AB, BP);
+		float D4 = FVector3f::DotProduct(AC, BP);
 		if (D3 >= 0.f && D4 <= D3)
 		{
 			FTrianglePoint Out;
 			Out.P = B;
-			Out.Barycentric = FVector(0, 1, 0);
+			Out.Barycentric = FVector3f(0, 1, 0);
 			return Out;
 		}
 
@@ -1667,19 +1667,19 @@ namespace GroomBinding_Transfer
 
 			FTrianglePoint Out;
 			Out.P = A + V * AB;
-			Out.Barycentric = FVector(1 - V, V, 0);
+			Out.Barycentric = FVector3f(1 - V, V, 0);
 			return Out;
 		}
 
 		// Check if P is in vertex region outside C.
-		FVector CP = P - C;
-		float D5 = FVector::DotProduct(AB, CP);
-		float D6 = FVector::DotProduct(AC, CP);
+		FVector3f CP = P - C;
+		float D5 = FVector3f::DotProduct(AB, CP);
+		float D6 = FVector3f::DotProduct(AC, CP);
 		if (D6 >= 0.f && D5 <= D6)
 		{
 			FTrianglePoint Out;
 			Out.P = C;
-			Out.Barycentric = FVector(0, 0, 1);
+			Out.Barycentric = FVector3f(0, 0, 1);
 			return Out;
 		}
 
@@ -1690,7 +1690,7 @@ namespace GroomBinding_Transfer
 			float W = D2 / (D2 - D6);
 			FTrianglePoint Out;
 			Out.P = A + W * AC;
-			Out.Barycentric = FVector(1 - W, 0, W);
+			Out.Barycentric = FVector3f(1 - W, 0, W);
 			return Out;
 		}
 
@@ -1701,7 +1701,7 @@ namespace GroomBinding_Transfer
 			float W = (D4 - D3) / (D4 - D3 + D5 - D6);
 			FTrianglePoint Out;
 			Out.P = B + W * (C - B);
-			Out.Barycentric = FVector(0, 1 - W, W);
+			Out.Barycentric = FVector3f(0, 1 - W, W);
 			return Out;
 		}
 
@@ -1712,7 +1712,7 @@ namespace GroomBinding_Transfer
 
 		FTrianglePoint Out;
 		Out.P = A + AB * V + AC * W;
-		Out.Barycentric = FVector(1 - V - W, V, W);
+		Out.Barycentric = FVector3f(1 - V - W, V, W);
 		return Out;
 	}
 
@@ -1753,9 +1753,9 @@ namespace GroomBinding_Transfer
 				T.I1 = SourceIndexBuffer[T.SectionBaseIndex + SourceTriangleIt * 3 + 1];
 				T.I2 = SourceIndexBuffer[T.SectionBaseIndex + SourceTriangleIt * 3 + 2];
 
-				T.P0 = (FVector)MeshLODData.GetVertexPosition(T.I0);
-				T.P1 = (FVector)MeshLODData.GetVertexPosition(T.I1);
-				T.P2 = (FVector)MeshLODData.GetVertexPosition(T.I2);
+				T.P0 = MeshLODData.GetVertexPosition(T.I0);
+				T.P1 = MeshLODData.GetVertexPosition(T.I1);
+				T.P2 = MeshLODData.GetVertexPosition(T.I2);
 
 				T.UV0 = MeshLODData.GetVertexUV(T.I0, InChannelIndex);
 				T.UV1 = MeshLODData.GetVertexUV(T.I1, InChannelIndex);
@@ -1826,10 +1826,10 @@ namespace GroomBinding_Transfer
 			const uint32 TargetTriangleCount = TargetMeshLODData.GetSection(LocalTargetSectionId).GetNumTriangles();
 			const uint32 TargetVertexCount = TargetMeshLODData.GetNumVertices();
 
-			TSet<FVector2D> UVs;
+			TSet<FVector2f> UVs;
 			for (uint32 TargetVertexIt = 0; TargetVertexIt < TargetVertexCount; ++TargetVertexIt)
 			{
-				const FVector2D Target_UV = TargetMeshLODData.GetVertexUV(TargetVertexIt, ChannelIndex);
+				const FVector2f Target_UV = TargetMeshLODData.GetVertexUV(TargetVertexIt, ChannelIndex);
 				UVs.Add(Target_UV);
 			}
 
@@ -1870,7 +1870,7 @@ namespace GroomBinding_Transfer
 				}
 
 				const FVector3f Target_P    = TargetMeshLODData.GetVertexPosition(TargetVertexIt);
-				const FVector2D Target_UV = TargetMeshLODData.GetVertexUV(TargetVertexIt, ChannelIndex);
+				const FVector2f Target_UV = TargetMeshLODData.GetVertexUV(TargetVertexIt, ChannelIndex);
 
 				// 2.1 Query closest triangles
 				FVector3f RetargetedVertexPosition = Target_P;
@@ -1883,7 +1883,7 @@ namespace GroomBinding_Transfer
 					for (const FTriangleGrid2D::FTriangle& CellTriangle : Cell->Triangles)
 					{
 						const FTrianglePoint ClosestPoint = ComputeClosestPoint(CellTriangle.UV0, CellTriangle.UV1, CellTriangle.UV2, Target_UV);
-						const float UVDistanceToTriangle = FVector2D::Distance(FVector2D(ClosestPoint.P.X, ClosestPoint.P.Y), Target_UV);
+						const float UVDistanceToTriangle = FVector2f::Distance(FVector2f(ClosestPoint.P.X, ClosestPoint.P.Y), Target_UV);
 						if (UVDistanceToTriangle < ClosestUVDistance)
 						{
 							RetargetedVertexPosition = FVector3f(
diff --git a/Engine/Plugins/Runtime/HairStrands/Source/HairStrandsCore/Private/GroomBuilder.cpp b/Engine/Plugins/Runtime/HairStrands/Source/HairStrandsCore/Private/GroomBuilder.cpp
--- a/Engine/Plugins/Runtime/HairStrands/Source/HairStrandsCore/Private/GroomBuilder.cpp
+++ b/Engine/Plugins/Runtime/HairStrands/Source/HairStrandsCore/Private/GroomBuilder.cpp
@@ -40,7 +40,7 @@ static FAutoConsoleVariableRef CVarHairGroupIndexBuilder_MaxVoxelResolution(TEXT
 
 FString FGroomBuilder::GetVersion()
 {
-	return TEXT("v8r50");
+	return TEXT("v8r54");
 }
 
 namespace FHairStrandsDecimation
@@ -94,17 +94,17 @@ namespace HairStrandsBuilder
 				*OffsetIterator = StrandOffset;
 
 				float StrandLength = 0.0;
-				FVector PreviousPosition(0.0, 0.0, 0.0);
+				FVector3f PreviousPosition(0.0, 0.0, 0.0);
 				for (uint32 PointIndex = 0; PointIndex < StrandCount; ++PointIndex, ++PositionIterator, ++RadiusIterator, ++CoordUIterator)
 				{
-					HairStrands.BoundingBox += (FVector)*PositionIterator;
+					HairStrands.BoundingBox += *PositionIterator;
 					HairStrands.BoundingBox.IsValid = 1;
 					if (PointIndex > 0)
 					{
-						StrandLength += ((FVector)*PositionIterator - (FVector)PreviousPosition).Size();
+						StrandLength += (*PositionIterator - PreviousPosition).Size();
 					}
 					*CoordUIterator = StrandLength;
-					PreviousPosition = (FVector)*PositionIterator;
+					PreviousPosition = *PositionIterator;
 				}
 				*LengthIterator = StrandLength;
 			}
@@ -129,7 +129,7 @@ namespace HairStrandsBuilder
 		}
 	}
 
-	inline void CopyVectorToPosition(const FVector& InVector, FHairStrandsPositionFormat::Type& OutPosition)
+	inline void CopyVectorToPosition(const FVector3f& InVector, FHairStrandsPositionFormat::Type& OutPosition)
 	{
 		OutPosition.X = InVector.X;
 		OutPosition.Y = InVector.Y;
@@ -215,7 +215,7 @@ namespace HairStrandsBuilder
 		if (HasHairAttribute(Attributes, EHairAttribute::Roughness)){ AttributeRoughness.SetNum(FMath::DivideAndRoundUp(NumPoints * Stride_Roughness, 	4u) * 4u); }
 		if (HasHairAttribute(Attributes, EHairAttribute::AO))		{ AttributeAO		.SetNum(FMath::DivideAndRoundUp(NumPoints * Stride_AO, 			4u) * 4u); }
 
-		const FVector HairBoxCenter = HairStrands.BoundingBox.GetCenter();
+		const FVector3f HairBoxCenter = HairStrands.BoundingBox.GetCenter();
 
 		const FHairStrandsCurves& Curves = HairStrands.StrandsCurves;
 		const FHairStrandsPoints& Points = HairStrands.StrandsPoints;
@@ -255,7 +255,7 @@ namespace HairStrandsBuilder
 			{
 				const uint32 PrevIndex = FMath::Max(0, PointIndex - 1);
 				const uint32 NextIndex = FMath::Min(PointCount + 1, PointCount - 1);
-				const FVector& PointPosition = (FVector)Points.PointsPosition[PointIndex + IndexOffset];
+				const FVector3f& PointPosition = Points.PointsPosition[PointIndex + IndexOffset];
 
 				// Packed Position|CoordU|NormalizedRadius
 				{
@@ -376,7 +376,8 @@ namespace HairStrandsBuilder
 			}
 		}
 
-		OutBulkData.Header.BoundingBox = HairStrands.BoundingBox;
+		OutBulkData.Header.BoundingBox.Min = (FVector)HairStrands.BoundingBox.Min;
+		OutBulkData.Header.BoundingBox.Max = (FVector)HairStrands.BoundingBox.Max;
 		OutBulkData.Header.CurveCount = HairStrands.GetNumCurves();
 		OutBulkData.Header.PointCount = HairStrands.GetNumPoints();
 		OutBulkData.Header.MinPointPerCurve = MinPointPerCurve;
@@ -592,11 +593,11 @@ namespace HairInterpolationBuilder
 {
 	struct FHairRoot
 	{
-		FVector Position;
-		uint32  VertexCount;
-		FVector Normal;
-		uint32  Index;
-		float   Length;
+		FVector3f Position;
+		uint32    VertexCount;
+		FVector3f Normal;
+		uint32    Index;
+		float     Length;
 	};
 
 	struct FHairInterpolationMetric
@@ -619,8 +620,8 @@ namespace HairInterpolationBuilder
 	inline FHairInterpolationMetric ComputeInterpolationMetric(const FHairRoot& RenderRoot, const FHairRoot& GuideRoot)
 	{
 		FHairInterpolationMetric Out;
-		Out.Distance = FVector::Dist(RenderRoot.Position, GuideRoot.Position);
-		Out.CosAngle = FVector::DotProduct(RenderRoot.Normal, GuideRoot.Normal);
+		Out.Distance = FVector3f::Dist(RenderRoot.Position, GuideRoot.Position);
+		Out.CosAngle = FVector3f::DotProduct(RenderRoot.Normal, GuideRoot.Normal);
 		Out.GuideLength = GuideRoot.Length;
 		Out.RenderLength = RenderRoot.Length;
 
@@ -639,42 +640,48 @@ namespace HairInterpolationBuilder
 	}
 
 	template<uint32 NumSamples>
-	inline FVector GetCurvePosition(const FHairStrandsDatas& CurvesDatas, const uint32 CurveIndex, const uint32 SampleIndex)
+	inline void GetCurvePositions(const FHairStrandsDatas& CurvesDatas, const uint32 CurveIndex, TStaticArray<FVector3f,NumSamples>& Out)
 	{
-		const float PointCount = CurvesDatas.StrandsCurves.CurvesCount[CurveIndex]-1.0;
+		const float  PointCount = CurvesDatas.StrandsCurves.CurvesCount[CurveIndex]-1.0;
 		const uint32 PointOffset = CurvesDatas.StrandsCurves.CurvesOffset[CurveIndex];
 
-		const float CurvePoint = static_cast<float>(SampleIndex) * PointCount / (static_cast<float>(NumSamples)-1.0f);
-		const uint32 PointPrev = (SampleIndex == 0) ? 0 : (SampleIndex == (NumSamples-1)) ? PointCount - 1 : floor(CurvePoint);
-		const uint32 PointNext = PointPrev + 1;
+		for (uint32 SampleIndex = 0; SampleIndex < NumSamples; ++SampleIndex)
+		{	
+			const float CurvePoint = static_cast<float>(SampleIndex) * PointCount / (static_cast<float>(NumSamples)-1.0f);
+			const uint32 PointPrev = (SampleIndex == 0) ? 0 : (SampleIndex == (NumSamples-1)) ? PointCount - 1 : floor(CurvePoint);
+			const uint32 PointNext = PointPrev + 1;	
+			const float PointAlpha = CurvePoint - static_cast<float>(PointPrev);
 
-		const float PointAlpha = CurvePoint - static_cast<float>(PointPrev);
-		return FVector(CurvesDatas.StrandsPoints.PointsPosition[PointOffset+PointPrev]*(1.0f-PointAlpha) +
-			CurvesDatas.StrandsPoints.PointsPosition[PointOffset+PointNext]*PointAlpha);
+			Out[SampleIndex] = CurvesDatas.StrandsPoints.PointsPosition[PointOffset + PointPrev]*(1.0f-PointAlpha) + CurvesDatas.StrandsPoints.PointsPosition[PointOffset + PointNext]*PointAlpha;
+		}
 	}
-
+	
 	template<uint32 NumSamples>
-	inline float ComputeCurvesMetric(const FHairStrandsDatas& RenderCurvesDatas, const uint32 RenderCurveIndex, 
-		const FHairStrandsDatas& GuideCurvesDatas, const uint32 GuideCurveIndex, const float RootImportance, 
-		const float ShapeImportance, const float ProximityImportance)
+	inline float ComputeCurvesMetric(const FHairStrandsDatas& RenCurvesDatas, const uint32 RenCurveIndex, 
+	                                 const FHairStrandsDatas& SimCurvesDatas, const uint32 SimCurveIndex, 
+	                                 const TStaticArray<FVector3f,NumSamples>& RenSamplePoints,
+	                                 const TStaticArray<FVector3f,NumSamples>& SimSamplePoints,
+	                                 const float RootImportance, 
+	                                 const float ShapeImportance, 
+	                                 const float ProximityImportance)
 	{
-		const float AverageLength = FMath::Max(0.5f * (RenderCurvesDatas.StrandsCurves.CurvesLength[RenderCurveIndex] + GuideCurvesDatas.StrandsCurves.CurvesLength[GuideCurveIndex]), SMALL_NUMBER);
+		const float AverageLength = FMath::Max(0.5f * (RenCurvesDatas.StrandsCurves.CurvesLength[RenCurveIndex] + SimCurvesDatas.StrandsCurves.CurvesLength[SimCurveIndex]), SMALL_NUMBER);
 
 		static const float DeltaCoord = 1.0f / static_cast<float>(NumSamples-1);
 
-		const FVector& RenderRoot = (FVector)RenderCurvesDatas.StrandsPoints.PointsPosition[RenderCurvesDatas.StrandsCurves.CurvesOffset[RenderCurveIndex]];
-		const FVector& GuideRoot = (FVector)GuideCurvesDatas.StrandsPoints.PointsPosition[GuideCurvesDatas.StrandsCurves.CurvesOffset[GuideCurveIndex]];
+		const FVector3f& RenRoot = RenCurvesDatas.StrandsPoints.PointsPosition[RenCurvesDatas.StrandsCurves.CurvesOffset[RenCurveIndex]];
+		const FVector3f& SimRoot = SimCurvesDatas.StrandsPoints.PointsPosition[SimCurvesDatas.StrandsCurves.CurvesOffset[SimCurveIndex]];
 
-		float CurveProximityMetric = 0.0;
-		float CurveShapeMetric = 0.0;
+		float CurveProximityMetric = 0.0f;
+		float CurveShapeMetric = 0.0f;
 		for (uint32 SampleIndex = 0; SampleIndex < NumSamples; ++SampleIndex)
 		{
-			const FVector GuidePosition = GetCurvePosition<NumSamples>(GuideCurvesDatas, GuideCurveIndex, SampleIndex);
-			const FVector RenderPosition = GetCurvePosition<NumSamples>(RenderCurvesDatas, RenderCurveIndex, SampleIndex);
+			const FVector3f& SimPosition  = SimSamplePoints[SampleIndex];
+			const FVector3f& RenPosition = RenSamplePoints[SampleIndex];
 			const float RootWeight = FMath::Exp(-RootImportance*SampleIndex*DeltaCoord);
 
-			CurveProximityMetric += (GuidePosition - RenderPosition).Size() * RootWeight;
-			CurveShapeMetric += (GuidePosition - GuideRoot - RenderPosition + RenderRoot).Size() * RootWeight;
+			CurveProximityMetric += (SimPosition - RenPosition).Size() * RootWeight;
+			CurveShapeMetric += (SimPosition - SimRoot - RenPosition + RenRoot).Size() * RootWeight;
 		}
 		CurveShapeMetric *= DeltaCoord / AverageLength;
 		CurveProximityMetric *= DeltaCoord / AverageLength;
@@ -682,6 +689,8 @@ namespace HairInterpolationBuilder
 		return FMath::Exp(-ShapeImportance*CurveShapeMetric) * FMath::Exp(-ProximityImportance * CurveProximityMetric);
 	}
 
+	typedef TStaticArray<FVector3f,16> FCurvePositions;
+
 	inline void PrintInterpolationMetric(const FHairInterpolationMetric& In)
 	{
 		UE_LOG(LogGroomBuilder, Log, TEXT("------------------------------------------------------------------------------------------"));
@@ -805,7 +814,7 @@ namespace HairInterpolationBuilder
 		const int32 RenCurveIndex,
 		const int32 SimCurveIndex)
 	{
-		const float Metric = FVector::Dist(GuideRoot.Position, RenRoot.Position);
+		const float Metric = FVector3f::Dist(GuideRoot.Position, RenRoot.Position);
 		if (Metric < Metrics1.KMinMetrics[FMetrics::Count - 1])
 		{
 			int32 LastGuideIndex = SimCurveIndex;
@@ -826,9 +835,14 @@ namespace HairInterpolationBuilder
 		const FHairStrandsDatas& RenStrandsData,
 		const FHairStrandsDatas& SimStrandsData,
 		const int32 RenCurveIndex,
-		const int32 SimCurveIndex)
+		const int32 SimCurveIndex,
+		const FCurvePositions& RenCurvePositions)
 	{
-		const float Metric = 1.0 - ComputeCurvesMetric<16>(RenStrandsData, RenCurveIndex, SimStrandsData, SimCurveIndex, 0.0f, 1.0f, 1.0f);
+		// Sample sim strands at fix rate posisions (similarly is done render strands)
+		FCurvePositions SimCurvePositions;
+		GetCurvePositions(SimStrandsData, SimCurveIndex, SimCurvePositions);
+
+		const float Metric = 1.0 - ComputeCurvesMetric<16>(RenStrandsData, RenCurveIndex, SimStrandsData, SimCurveIndex, RenCurvePositions, SimCurvePositions, 0.0f, 1.0f, 1.0f);
 		if (Metric < Metrics0.KMinMetrics[FMetrics::Count - 1])
 		{
 			int32 LastGuideIndex = SimCurveIndex;
@@ -846,8 +860,8 @@ namespace HairInterpolationBuilder
 
 	struct FRootsGrid
 	{
-		FVector MinBound = FVector::ZeroVector;
-		FVector MaxBound = FVector::ZeroVector;
+		FVector3f MinBound = FVector3f::ZeroVector;
+		FVector3f MaxBound = FVector3f::ZeroVector;
 		
 		const uint32 MaxLookupDistance = 31;
 		const FIntVector GridResolution = FIntVector(32, 32, 32);
@@ -871,10 +885,10 @@ namespace HairInterpolationBuilder
 				FMath::Clamp(CellCoord.Z, 0, GridResolution.Z - 1));
 		}
 
-		FORCEINLINE FIntVector ToCellCoord(const FVector& P) const
+		FORCEINLINE FIntVector ToCellCoord(const FVector3f& P) const
 		{
 			bool bIsValid = false;
-			const FVector F = ((P - MinBound) / (MaxBound - MinBound));
+			const FVector3f F = ((P - MinBound) / (MaxBound - MinBound));
 			const FIntVector CellCoord = FIntVector(FMath::FloorToInt(F.X * GridResolution.X), FMath::FloorToInt(F.Y * GridResolution.Y), FMath::FloorToInt(F.Z * GridResolution.Z));			
 			return ClampToVolume(CellCoord, bIsValid);
 		}
@@ -886,7 +900,7 @@ namespace HairInterpolationBuilder
 			return CellIndex;
 		}
 
-		void InsertRoots(TArray<FHairRoot>& Roots, const FVector& InMinBound, const FVector& InMaxBound)
+		void InsertRoots(TArray<FHairRoot>& Roots, const FVector3f& InMinBound, const FVector3f& InMaxBound)
 		{
 			MinBound = InMinBound;
 			MaxBound = InMaxBound;
@@ -1024,6 +1038,7 @@ namespace HairInterpolationBuilder
 			const TArray<FHairRoot>& SimRoots,
 			const FHairStrandsDatas& RenStrandsData,
 			const FHairStrandsDatas& SimStrandsData,
+			const FCurvePositions& RenCurvePositions,
 			FMetrics& Metrics0,
 			FMetrics& Metrics1) const
 		{
@@ -1039,7 +1054,7 @@ namespace HairInterpolationBuilder
 				const FHairRoot& GuideRoot = SimRoots[SimCurveIndex];
 				{
 					ComputeSimpleMetric(Metrics1, RenRoot, GuideRoot, RenCurveIndex, SimCurveIndex);
-					ComputeAdvandedMetric(Metrics0, RenStrandsData, SimStrandsData, RenCurveIndex, SimCurveIndex);
+					ComputeAdvandedMetric(Metrics0, RenStrandsData, SimStrandsData, RenCurveIndex, SimCurveIndex, RenCurvePositions);
 				}
 			}
 		}
@@ -1050,6 +1065,7 @@ namespace HairInterpolationBuilder
 			const TArray<FHairRoot>& SimRoots,
 			const FHairStrandsDatas& RenStrandsData,
 			const FHairStrandsDatas& SimStrandsData,
+			const FCurvePositions& RenCurvePositions,
 			const bool bRandomized,
 			const bool bUnique,
 			const FIntVector& RandomIndices) const
@@ -1073,7 +1089,7 @@ namespace HairInterpolationBuilder
 				{
 					bool bIsValid = false;
 					const FIntVector CellCoord = ClampToVolume(PointCoord, bIsValid);
-					if (bIsValid) SearchCell(CellCoord, RenCurveIndex, RenRoot, SimRoots, RenStrandsData, SimStrandsData, Metrics0, Metrics1);
+					if (bIsValid) SearchCell(CellCoord, RenCurveIndex, RenRoot, SimRoots, RenStrandsData, SimStrandsData, RenCurvePositions, Metrics0, Metrics1);
 				}
 
 				// Top & Bottom
@@ -1083,8 +1099,8 @@ namespace HairInterpolationBuilder
 					bool bIsValid0 = false, bIsValid1 = false;
 					const FIntVector CellCoord0 = ClampToVolume(PointCoord + FIntVector(X, Y, Offset), bIsValid0);
 					const FIntVector CellCoord1 = ClampToVolume(PointCoord + FIntVector(X, Y,-Offset), bIsValid1);
-					if (bIsValid0) SearchCell(CellCoord0, RenCurveIndex, RenRoot, SimRoots, RenStrandsData, SimStrandsData, Metrics0, Metrics1);
-					if (bIsValid1) SearchCell(CellCoord1, RenCurveIndex, RenRoot, SimRoots, RenStrandsData, SimStrandsData, Metrics0, Metrics1);
+					if (bIsValid0) SearchCell(CellCoord0, RenCurveIndex, RenRoot, SimRoots, RenStrandsData, SimStrandsData, RenCurvePositions, Metrics0, Metrics1);
+					if (bIsValid1) SearchCell(CellCoord1, RenCurveIndex, RenRoot, SimRoots, RenStrandsData, SimStrandsData, RenCurvePositions, Metrics0, Metrics1);
 				}
 
 				const int32 OffsetMinusOne = Offset - 1;
@@ -1095,8 +1111,8 @@ namespace HairInterpolationBuilder
 					bool bIsValid0 = false, bIsValid1 = false;
 					const FIntVector CellCoord0 = ClampToVolume(PointCoord + FIntVector(X, Offset, Z), bIsValid0);
 					const FIntVector CellCoord1 = ClampToVolume(PointCoord + FIntVector(X, -Offset, Z), bIsValid1);
-					if (bIsValid0) SearchCell(CellCoord0, RenCurveIndex, RenRoot, SimRoots, RenStrandsData, SimStrandsData, Metrics0, Metrics1);
-					if (bIsValid1) SearchCell(CellCoord1, RenCurveIndex, RenRoot, SimRoots, RenStrandsData, SimStrandsData, Metrics0, Metrics1);
+					if (bIsValid0) SearchCell(CellCoord0, RenCurveIndex, RenRoot, SimRoots, RenStrandsData, SimStrandsData, RenCurvePositions, Metrics0, Metrics1);
+					if (bIsValid1) SearchCell(CellCoord1, RenCurveIndex, RenRoot, SimRoots, RenStrandsData, SimStrandsData, RenCurvePositions, Metrics0, Metrics1);
 				}
 				
 				// Left & Right
@@ -1106,8 +1122,8 @@ namespace HairInterpolationBuilder
 					bool bIsValid0 = false, bIsValid1 = false;
 					const FIntVector CellCoord0 = ClampToVolume(PointCoord + FIntVector( Offset, Y, Z), bIsValid0);
 					const FIntVector CellCoord1 = ClampToVolume(PointCoord + FIntVector(-Offset, Y, Z), bIsValid1);
-					if (bIsValid0) SearchCell(CellCoord0, RenCurveIndex, RenRoot, SimRoots, RenStrandsData, SimStrandsData, Metrics0, Metrics1);
-					if (bIsValid1) SearchCell(CellCoord1, RenCurveIndex, RenRoot, SimRoots, RenStrandsData, SimStrandsData, Metrics0, Metrics1);
+					if (bIsValid0) SearchCell(CellCoord0, RenCurveIndex, RenRoot, SimRoots, RenStrandsData, SimStrandsData, RenCurvePositions, Metrics0, Metrics1);
+					if (bIsValid1) SearchCell(CellCoord1, RenCurveIndex, RenRoot, SimRoots, RenStrandsData, SimStrandsData, RenCurvePositions, Metrics0, Metrics1);
 				}
 
 				// Early out if we have found all closest guide during a ring/layer step.
@@ -1143,6 +1159,7 @@ namespace HairInterpolationBuilder
 		const TArray<FHairRoot>& SimRoots,
 		const FHairStrandsDatas& RenStrandsData,
 		const FHairStrandsDatas& SimStrandsData,
+		const FCurvePositions& RenCurvePositions,
 		const bool bRandomized,
 		const bool bUnique,
 		const FIntVector& RandomIndices)
@@ -1157,7 +1174,7 @@ namespace HairInterpolationBuilder
 		const uint32 SimRootsCount = SimRoots.Num();
 		for (uint32 SimCurveIndex =0; SimCurveIndex<SimRootsCount; ++SimCurveIndex)
 		{
-			ComputeAdvandedMetric(Metrics, RenStrandsData, SimStrandsData, RenCurveIndex, SimCurveIndex);
+			ComputeAdvandedMetric(Metrics, RenStrandsData, SimStrandsData, RenCurveIndex, SimCurveIndex, RenCurvePositions);
 		}
 			
 		FClosestGuides ClosestGuides;
@@ -1171,10 +1188,10 @@ namespace HairInterpolationBuilder
 	}
 
 	// Extract strand roots
-	static void ExtractRoots(const FHairStrandsDatas& InData, TArray<FHairRoot>& OutRoots, FVector& MinBound, FVector& MaxBound)
+	static void ExtractRoots(const FHairStrandsDatas& InData, TArray<FHairRoot>& OutRoots, FVector3f& MinBound, FVector3f& MaxBound)
 	{
-		MinBound = FVector(FLT_MAX, FLT_MAX, FLT_MAX);
-		MaxBound = FVector(-FLT_MAX, -FLT_MAX, -FLT_MAX);
+		MinBound = FVector3f(FLT_MAX, FLT_MAX, FLT_MAX);
+		MaxBound = FVector3f(-FLT_MAX, -FLT_MAX, -FLT_MAX);
 		const uint32 CurveCount = InData.StrandsCurves.Num();
 		OutRoots.Reserve(CurveCount);
 		for (uint32 CurveIndex = 0; CurveIndex < CurveCount; ++CurveIndex)
@@ -1183,14 +1200,14 @@ namespace HairInterpolationBuilder
 			const uint32 PointCount = InData.StrandsCurves.CurvesCount[CurveIndex];
 			const float  CurveLength = InData.StrandsCurves.CurvesLength[CurveIndex];
 			check(PointCount > 1);
-			const FVector& P0 = (FVector)InData.StrandsPoints.PointsPosition[PointOffset];
-			const FVector& P1 = (FVector)InData.StrandsPoints.PointsPosition[PointOffset + 1];
-			FVector N = (P1 - P0).GetSafeNormal();
+			const FVector3f& P0 = InData.StrandsPoints.PointsPosition[PointOffset];
+			const FVector3f& P1 = InData.StrandsPoints.PointsPosition[PointOffset + 1];
+			FVector3f N = (P1 - P0).GetSafeNormal();
 
 			// Fallback in case the initial points are too close (this happens on certain assets)
-			if (FVector::DotProduct(N, N) == 0)
+			if (FVector3f::DotProduct(N, N) == 0)
 			{
-				N = FVector(0, 0, 1);
+				N = FVector3f(0, 0, 1);
 			}
 			OutRoots.Add({ P0, PointCount, N, PointOffset, CurveLength });
 
@@ -1296,8 +1313,8 @@ namespace HairInterpolationBuilder
 		FRoots RenRoots, SimRoots;
 		FRootsGrid RootsGrid;
 		{
-			FVector RenMinBound, RenMaxBound;
-			FVector SimMinBound, SimMaxBound;
+			FVector3f RenMinBound, RenMaxBound;
+			FVector3f SimMinBound, SimMaxBound;
 			ExtractRoots(RenStrandsData, RenRoots, RenMinBound, RenMaxBound);
 			ExtractRoots(SimStrandsData, SimRoots, SimMinBound, SimMaxBound);
 
@@ -1305,8 +1322,8 @@ namespace HairInterpolationBuilder
 			{
 				// Build a conservative bound, to insure all queries will fall 
 				// into the grid volume.
-				const FVector MinBound = RenMinBound.ComponentMin(SimMinBound);
-				const FVector MaxBound = RenMaxBound.ComponentMax(SimMaxBound);
+				const FVector3f MinBound = RenMinBound.ComponentMin(SimMinBound);
+				const FVector3f MaxBound = RenMaxBound.ComponentMax(SimMaxBound);
 				RootsGrid.InsertRoots(SimRoots, MinBound, MaxBound);
 			}
 		}
@@ -1363,20 +1380,24 @@ namespace HairInterpolationBuilder
 			}
 			else if (Settings.InterpolationQuality == EHairInterpolationQuality::Medium)
 			{
-				ClosestGuides = RootsGrid.FindBestClosestRoots(RenCurveIndex, RenRoots, SimRoots, RenStrandsData, SimStrandsData, Settings.bRandomizeGuide, Settings.bUseUniqueGuide, RandomGuideIndices[RenCurveIndex]);
+				FCurvePositions RenCurvePositions; 
+				GetCurvePositions(RenStrandsData, RenCurveIndex, RenCurvePositions);
+				ClosestGuides = RootsGrid.FindBestClosestRoots(RenCurveIndex, RenRoots, SimRoots, RenStrandsData, SimStrandsData, RenCurvePositions, Settings.bRandomizeGuide, Settings.bUseUniqueGuide, RandomGuideIndices[RenCurveIndex]);
 			}
 			else // (Settings.Quality == EHairInterpolationQuality::High)
 			{
-				ClosestGuides = FindBestRoots(RenCurveIndex, RenRoots, SimRoots, RenStrandsData, SimStrandsData, Settings.bRandomizeGuide, Settings.bUseUniqueGuide, RandomGuideIndices[RenCurveIndex]);
+				FCurvePositions RenCurvePositions; 
+				GetCurvePositions(RenStrandsData, RenCurveIndex, RenCurvePositions);
+				ClosestGuides = FindBestRoots(RenCurveIndex, RenRoots, SimRoots, RenStrandsData, SimStrandsData, RenCurvePositions, Settings.bRandomizeGuide, Settings.bUseUniqueGuide, RandomGuideIndices[RenCurveIndex]);
 			}
 
 			const uint32 RendPointCount	= RenStrandsData.StrandsCurves.CurvesCount[RenCurveIndex];
 			const uint32 RenOffset		= RenStrandsData.StrandsCurves.CurvesOffset[RenCurveIndex];
-			const FVector& RenPointPositionRoot = (FVector)RenStrandsData.StrandsPoints.PointsPosition[RenOffset];
+			const FVector3f& RenPointPositionRoot = RenStrandsData.StrandsPoints.PointsPosition[RenOffset];
 			for (uint32 RenPointIndex = 0; RenPointIndex < RendPointCount; ++RenPointIndex)
 			{
 				const uint32 PointGlobalIndex = RenPointIndex + RenOffset;
-				const FVector& RenPointPosition = (FVector)RenStrandsData.StrandsPoints.PointsPosition[PointGlobalIndex];
+				const FVector3f& RenPointPosition = RenStrandsData.StrandsPoints.PointsPosition[PointGlobalIndex];
 				const float RenPointDistance = RenStrandsData.StrandsPoints.PointsCoordU[PointGlobalIndex] * RenStrandsData.StrandsCurves.CurvesLength[RenCurveIndex];
 
 				float TotalWeight = 0;
@@ -1391,9 +1412,9 @@ namespace HairInterpolationBuilder
 							FindMatchingVertex(RenPointDistance, SimStrandsData, SimCurveIndex) :
 							FindMatchingVertex(RenStrandsData.StrandsPoints.PointsPosition[PointGlobalIndex], SimStrandsData, SimCurveIndex);
 						
-						const FVector& SimPointPosition0 = (FVector)SimStrandsData.StrandsPoints.PointsPosition[Desc.Index0 + SimOffset];
-						const FVector& SimPointPosition1 = (FVector)SimStrandsData.StrandsPoints.PointsPosition[Desc.Index1 + SimOffset];
-						const float Weight = 1.0f / FMath::Max(MinWeightDistance, FVector::Distance(RenPointPosition, FMath::Lerp(SimPointPosition0, SimPointPosition1, Desc.T)));
+						const FVector3f& SimPointPosition0 = SimStrandsData.StrandsPoints.PointsPosition[Desc.Index0 + SimOffset];
+						const FVector3f& SimPointPosition1 = SimStrandsData.StrandsPoints.PointsPosition[Desc.Index1 + SimOffset];
+						const float Weight = 1.0f / FMath::Max(MinWeightDistance, FVector3f::Distance(RenPointPosition, FMath::Lerp(SimPointPosition0, SimPointPosition1, Desc.T)));
 
 						InterpolationData.PointsSimCurvesIndex[PointGlobalIndex][KIndex] = SimCurveIndex;
 						InterpolationData.PointsSimCurvesVertexIndex[PointGlobalIndex][KIndex] = Desc.Index0 + SimOffset;
@@ -1407,9 +1428,9 @@ namespace HairInterpolationBuilder
 					{
 						const uint32 SimCurveIndex = ClosestGuides.Indices[KIndex];
 						const uint32 SimOffset = SimStrandsData.StrandsCurves.CurvesOffset[SimCurveIndex];
-						const FVector& SimRootPointPosition = (FVector)SimStrandsData.StrandsPoints.PointsPosition[SimOffset];
-						const FVector& RenRootPointPosition = (FVector)RenStrandsData.StrandsPoints.PointsPosition[RenOffset];
-						const float Weight = 1.0f / FMath::Max(MinWeightDistance, FVector::Distance(RenRootPointPosition, SimRootPointPosition));
+						const FVector3f& SimRootPointPosition = SimStrandsData.StrandsPoints.PointsPosition[SimOffset];
+						const FVector3f& RenRootPointPosition = RenStrandsData.StrandsPoints.PointsPosition[RenOffset];
+						const float Weight = 1.0f / FMath::Max(MinWeightDistance, FVector3f::Distance(RenRootPointPosition, SimRootPointPosition));
 						const FVertexInterpolationDesc Desc = FindMatchingVertex(RenPointDistance, SimStrandsData, SimCurveIndex);
 
 						InterpolationData.PointsSimCurvesIndex[PointGlobalIndex][KIndex] = SimCurveIndex;
@@ -1425,8 +1446,8 @@ namespace HairInterpolationBuilder
 						const uint32 SimOffset = SimStrandsData.StrandsCurves.CurvesOffset[SimCurveIndex];
 						const uint32 SimPointCount = SimStrandsData.StrandsCurves.CurvesCount[SimCurveIndex];
 						const uint32 SimPointIndex = FMath::Clamp(RenPointIndex, 0u, SimPointCount - 1);
-						const FVector& SimPointPosition = (FVector)SimStrandsData.StrandsPoints.PointsPosition[SimPointIndex + SimOffset];
-						const float Weight = 1.0f / FMath::Max(MinWeightDistance, FVector::Distance(RenPointPosition, SimPointPosition));
+						const FVector3f& SimPointPosition = SimStrandsData.StrandsPoints.PointsPosition[SimPointIndex + SimOffset];
+						const float Weight = 1.0f / FMath::Max(MinWeightDistance, FVector3f::Distance(RenPointPosition, SimPointPosition));
 
 						InterpolationData.PointsSimCurvesIndex[PointGlobalIndex][KIndex] = SimCurveIndex;
 						InterpolationData.PointsSimCurvesVertexIndex[PointGlobalIndex][KIndex] = SimPointIndex + SimOffset;
@@ -1487,8 +1508,8 @@ namespace HairInterpolationBuilder
 			for (uint32 PointIndex = 0; PointIndex < PointCount; ++PointIndex)
 			{
 				const FIntVector& Indices 	= HairInterpolation.PointsSimCurvesVertexIndex[PointIndex];
-				const FVector3f& Weights 	= (FVector3f)HairInterpolation.PointsSimCurvesVertexWeights[PointIndex];
-				const FVector3f& S 			= (FVector3f)HairInterpolation.PointsSimCurvesVertexLerp[PointIndex];
+				const FVector3f& Weights 	= HairInterpolation.PointsSimCurvesVertexWeights[PointIndex];
+				const FVector3f& S 			= HairInterpolation.PointsSimCurvesVertexLerp[PointIndex];
 
 				FHairInterpolation1GuideVertex& OutInterp = Out[PointIndex];
 				OutInterp.VertexGuideIndex0 = LowerPart(Indices[0]);
@@ -1593,11 +1614,11 @@ namespace HairInterpolationBuilder
 			const uint16 CurveNumVertices = RenData.StrandsCurves.CurvesCount[CurveIndex];
 
 			const FIntVector StrandClosestGuides = RenData.StrandsCurves.CurvesClosestGuideIDs[CurveIndex];
-			const FVector StrandGuideWeights  = RenData.StrandsCurves.CurvesClosestGuideWeights[CurveIndex];
+			const FVector3f StrandGuideWeights  = RenData.StrandsCurves.CurvesClosestGuideWeights[CurveIndex];
 			for (uint16 VertexIndex = 0; VertexIndex < CurveNumVertices; ++VertexIndex)
 			{
 				const uint32 PointGlobalIndex = VertexIndex + CurveOffset;
-				const FVector& RenPointPosition = (FVector)RenData.StrandsPoints.PointsPosition[PointGlobalIndex];
+				const FVector3f& RenPointPosition = RenData.StrandsPoints.PointsPosition[PointGlobalIndex];
 				const float RenPointDistance = RenData.StrandsPoints.PointsCoordU[PointGlobalIndex] * RenData.StrandsCurves.CurvesLength[CurveIndex];
 
 				bool bHasValidGuide = false;
@@ -1654,7 +1675,7 @@ namespace HairInterpolationBuilder
 	}
 }
 
-FORCEINLINE FIntVector VectorToIntVector(const FVector& Index)
+FORCEINLINE FIntVector VectorToIntVector(const FVector3f& Index)
 {
 	return FIntVector(Index.X, Index.Y, Index.Z);
 }
@@ -1904,13 +1925,13 @@ bool FGroomBuilder::BuildHairDescriptionGroups(const FHairDescription& HairDescr
 			{
 				if (bPrecomputedWeight3)
 				{
-					CurrentHairStrandsDatas->StrandsCurves.CurvesClosestGuideIDs.Add(VectorToIntVector((FVector)ClosestGuides[StrandID]));
-					CurrentHairStrandsDatas->StrandsCurves.CurvesClosestGuideWeights.Add((FVector)GuideWeights[StrandID]);
+					CurrentHairStrandsDatas->StrandsCurves.CurvesClosestGuideIDs.Add(VectorToIntVector(ClosestGuides[StrandID]));
+					CurrentHairStrandsDatas->StrandsCurves.CurvesClosestGuideWeights.Add(GuideWeights[StrandID]);
 				}
 				else
 				{
 					CurrentHairStrandsDatas->StrandsCurves.CurvesClosestGuideIDs.Add(FIntVector(ClosestGuide[StrandID]));
-					CurrentHairStrandsDatas->StrandsCurves.CurvesClosestGuideWeights.Add(FVector(GuideWeight[StrandID]));
+					CurrentHairStrandsDatas->StrandsCurves.CurvesClosestGuideWeights.Add(FVector3f(GuideWeight[StrandID]));
 				}
 			}
 		}
@@ -2247,11 +2268,11 @@ static void DecimateCurve(
 		const uint32 Count = OutIndices.Num();
 		for (uint32 IndexIt = 1; IndexIt < Count - 1; ++IndexIt)
 		{
-			const FVector& P0 = (FVector)InPoints[InOffset + OutIndices[IndexIt - 1]];
-			const FVector& P1 = (FVector)InPoints[InOffset + OutIndices[IndexIt]];
-			const FVector& P2 = (FVector)InPoints[InOffset + OutIndices[IndexIt + 1]];
+			const FVector3f& P0 = InPoints[InOffset + OutIndices[IndexIt - 1]];
+			const FVector3f& P1 = InPoints[InOffset + OutIndices[IndexIt]];
+			const FVector3f& P2 = InPoints[InOffset + OutIndices[IndexIt + 1]];
 
-			const float Area = FVector::CrossProduct(P0 - P1, P2 - P1).Size() * 0.5f;
+			const float Area = FVector3f::CrossProduct(P0 - P1, P2 - P1).Size() * 0.5f;
 
 			if (Area < MinError)
 			{
@@ -2462,9 +2483,9 @@ FORCEINLINE uint32 ToLinearCoord(const FIntVector& T, const FIntVector& Resoluti
 	return T.X + T.Y * Resolution.X + T.Z * Resolution.X * Resolution.Y;
 }
 
-FORCEINLINE FIntVector ToCoord(const FVector& T, const FIntVector& Resolution, const FVector& MinBound, const float VoxelSize)
+FORCEINLINE FIntVector ToCoord(const FVector3f& T, const FIntVector& Resolution, const FVector3f& MinBound, const float VoxelSize)
 {
-	const FVector C = (T - MinBound) / VoxelSize;
+	const FVector3f C = (T - MinBound) / VoxelSize;
 	return FIntVector(
 		FMath::Clamp(FMath::FloorToInt(C.X), 0, Resolution.X - 1),
 		FMath::Clamp(FMath::FloorToInt(C.Y), 0, Resolution.Y - 1),
@@ -2490,7 +2511,7 @@ FORCEINLINE FIntVector ClampToVolume(const FIntVector& CellCoord, const FIntVect
 FORCEINLINE FIntVector ToCellCoord(const FVector3f& P, const FVector3f& MinBound, const FVector4f& MaxBound, const FIntVector& Resolution)
 {
 	bool bIsValid = false;
-	const FVector F = FVector((P - MinBound) / (MaxBound - MinBound));
+	const FVector3f F = FVector3f((P - MinBound) / (MaxBound - MinBound));
 	const FIntVector CellCoord = FIntVector(FMath::FloorToInt(F.X * Resolution.X), FMath::FloorToInt(F.Y * Resolution.Y), FMath::FloorToInt(F.Z * Resolution.Z));
 	return ClampToVolume(CellCoord, Resolution, bIsValid);
 }
@@ -2589,7 +2610,7 @@ static void Voxelize(const FHairDescriptionGroups& InGroups, FHairGrid& Out)
 				{
 					const float T = float(StepIt) / float(StepCount);
 					const FVector3f P = P0 + Segment * T;
-					const FIntVector Coord = ToCoord((FVector)P, Out.Resolution, (FVector)Out.MinBound, Out.VoxelSize);
+					const FIntVector Coord = ToCoord(P, Out.Resolution, Out.MinBound, Out.VoxelSize);
 					const uint32 LinearCoord = ToLinearCoord(Coord, Out.Resolution);
 					if (LinearCoord != PrevLinearCoord)
 					{
@@ -2893,11 +2914,11 @@ static void DecimateCurveLOD(
 			const uint32 Count = OutIndices.Num();
 			for (uint32 IndexIt = 1; IndexIt < Count - 1; ++IndexIt)
 			{
-				const FVector& P0 = (FVector)InPoints[InOffset + OutIndices[IndexIt - 1]];
-				const FVector& P1 = (FVector)InPoints[InOffset + OutIndices[IndexIt]];
-				const FVector& P2 = (FVector)InPoints[InOffset + OutIndices[IndexIt + 1]];
+				const FVector3f& P0 = InPoints[InOffset + OutIndices[IndexIt - 1]];
+				const FVector3f& P1 = InPoints[InOffset + OutIndices[IndexIt]];
+				const FVector3f& P2 = InPoints[InOffset + OutIndices[IndexIt + 1]];
 
-				const float Area = FVector::CrossProduct(P0 - P1, P2 - P1).Size() * 0.5f;
+				const float Area = FVector3f::CrossProduct(P0 - P1, P2 - P1).Size() * 0.5f;
 
 				//     P0 .       . P2
 				//         \Inner/
@@ -2905,9 +2926,9 @@ static void DecimateCurveLOD(
 				// Thres(` . \^/ ) Angle
 				//    --------.---------
 				//            P1
-				const FVector V0 = (P0 - P1).GetSafeNormal();
-				const FVector V1 = (P2 - P1).GetSafeNormal();
-				const float InnerAngle = FMath::Abs(FMath::Acos(FVector::DotProduct(V0, V1)));
+				const FVector3f V0 = (P0 - P1).GetSafeNormal();
+				const FVector3f V1 = (P2 - P1).GetSafeNormal();
+				const float InnerAngle = FMath::Abs(FMath::Acos(FVector3f::DotProduct(V0, V1)));
 				const float Angle = (PI - InnerAngle) * 0.5f;
 
 				if (Area < MinError && Angle < LODAngularThreshold)
@@ -2976,7 +2997,21 @@ static void BuildClusterData(
 	}
 
 	// 2. Compute control point LOD
-	for (uint32 CurveIt = 0; CurveIt < Out.CurveCount; ++CurveIt)
+	TAtomic<uint32> ConcurrentPointCountPerLOD[FHairClusterInfo::MaxLOD];
+	for (uint8 LODIt = 0; LODIt < FHairClusterInfo::MaxLOD; ++LODIt)
+	{
+		ConcurrentPointCountPerLOD[LODIt] = 0;
+	}
+	Out.PointLODs.SetNum(InRenStrandsData.GetNumPoints());
+	ParallelFor(Out.CurveCount, 
+	[
+		LODCount,
+		&ConcurrentPointCountPerLOD,
+		&InSettings,
+		&InRenStrandsData,
+		&Out
+	] (uint32 CurveIt) 
+	//for (uint32 CurveIt = 0; CurveIt < Out.CurveCount; ++CurveIt)
 	{
 		const uint32 PointOffset = InRenStrandsData.StrandsCurves.CurvesOffset[CurveIt];
 		const uint32 PointCount  = InRenStrandsData.StrandsCurves.CurvesCount[CurveIt];
@@ -3000,12 +3035,12 @@ static void BuildClusterData(
 			// For stats
 			for (uint32 LODIt = 0; LODIt < LODCount; ++LODIt)
 			{
-				Out.LODInfos[LODIt].PointCount += CountPerLOD[LODIt];
+				ConcurrentPointCountPerLOD[LODIt] += CountPerLOD[LODIt];
+				//Out.LODInfos[LODIt].PointCount += CountPerLOD[LODIt];
 			}
 		}
 
 		// 2.2 Write LOD mask per control point
-		Out.PointLODs.SetNum(InRenStrandsData.GetNumPoints());
 		for (uint8 LODIt = 0; LODIt < LODCount; ++LODIt)
 		{
 			for (uint32 PointIt = 0; PointIt < PointCount; ++PointIt)
@@ -3019,6 +3054,11 @@ static void BuildClusterData(
 				}
 			}
 		}
+	});
+	// Copy back the 'concurrent PointCountPerLOD' into the final array
+	for (uint32 LODIt = 0; LODIt < LODCount; ++LODIt)
+	{
+		Out.LODInfos[LODIt].PointCount = ConcurrentPointCountPerLOD[LODIt];
 	}
 
 	// 3. Compute radius scale per LOD
@@ -3044,11 +3084,11 @@ static void BuildClusterData(
 			for (uint32 PointIt = 0; PointIt < PointCount; ++PointIt)
 			{
 				uint32 PointGlobalIndex = PointIt + PointOffset;
-				const FVector& V0 = (FVector)InRenStrandsData.StrandsPoints.PointsPosition[PointGlobalIndex];
+				const FVector3f& V0 = InRenStrandsData.StrandsPoints.PointsPosition[PointGlobalIndex];
 				if (PointIt > 0)
 				{
-					const FVector& V1 = (FVector)InRenStrandsData.StrandsPoints.PointsPosition[PointGlobalIndex - 1];
-					FVector OutDir;
+					const FVector3f& V1 = InRenStrandsData.StrandsPoints.PointsPosition[PointGlobalIndex - 1];
+					FVector3f OutDir;
 					float OutLength;
 					(V1 - V0).ToDirectionAndLength(OutDir, OutLength);
 					CurveInfo.Area += InRenStrandsData.StrandsPoints.PointsRadius[PointGlobalIndex] * OutLength;
diff --git a/Engine/Plugins/Runtime/HairStrands/Source/HairStrandsCore/Private/HairCardsBuilder.cpp b/Engine/Plugins/Runtime/HairStrands/Source/HairStrandsCore/Private/HairCardsBuilder.cpp
--- a/Engine/Plugins/Runtime/HairStrands/Source/HairStrandsCore/Private/HairCardsBuilder.cpp
+++ b/Engine/Plugins/Runtime/HairStrands/Source/HairStrandsCore/Private/HairCardsBuilder.cpp
@@ -329,7 +329,7 @@ bool InternalCreateCardsGuides(
 		{
 			const FVector3f P0 = CenterPoints[PointIt];
 
-			OutGuides.BoundingBox += (FVector)P0;
+			OutGuides.BoundingBox += P0;
 
 			OutGuides.StrandsPoints.PointsPosition.Add(P0);
 			OutGuides.StrandsPoints.PointsCoordU.Add(FMath::Clamp(CurrentLength / TotalLength, 0.f, 1.f));
@@ -697,7 +697,6 @@ static bool InternalImportGeometry_WithGeneratedGuides(
 		struct FStrandsRootData
 		{
 			FVector3f Position;
-			uint32    CurveIndex;
 			FVector2f RootUV;
 		};
 		TArray<FStrandsRootData> StrandsRoots;
@@ -711,7 +710,6 @@ static bool InternalImportGeometry_WithGeneratedGuides(
 				FStrandsRootData& RootData = StrandsRoots.AddDefaulted_GetRef();
 				RootData.Position = InStrandsData.StrandsPoints.PointsPosition[Offset];
 				RootData.RootUV = HasHairAttribute(InAttribute, EHairAttribute::RootUV) ? InStrandsData.StrandsCurves.CurvesRootUV[CurveIt] : FVector2f(0,0);
-				RootData.CurveIndex = CurveIt;
 			}
 		}
 
@@ -774,31 +772,48 @@ static bool InternalImportGeometry_WithGeneratedGuides(
 		}
 		
 		// 3. Find cards root / curve root
-		for (const FCardsRootData& CardsRoot : CardsRoots)
+		ParallelFor(CardsRoots.Num(), 
+		[
+			&CardsRoots,
+			&StrandsRoots,
+			&Out,
+			&OutBulk
+		] (uint32 CardRootIt) 
+		//for (const FCardsRootData& CardsRoot : CardsRoots)
 		{
+			const FCardsRootData& CardsRoot = CardsRoots[CardRootIt];
 			for (uint32 CardsRootPositionIndex=0; CardsRootPositionIndex <2; CardsRootPositionIndex++)
 			{
 				// 3.1 Find closet root UV
 				// /!\ N^2 loop: the number of cards should be relatively small
-				auto FindRootUV = [&](const FVector3f CardsRootPosition)
+				auto FindRootUV = [&](const FVector3f CardsRootPosition0, const FVector3f CardsRootPosition1, FVector2f& OutRootUV0, FVector2f& OutRootUV1)
 				{
-					uint32 CurveIndex = ~0;
-					float ClosestDistance = FLT_MAX;
-					FVector2f RootUV = FVector2f::ZeroVector;
+					float ClosestDistance1 = FLT_MAX;
+					float ClosestDistance0 = FLT_MAX;
+					OutRootUV0 = FVector2f::ZeroVector;
+					OutRootUV1 = FVector2f::ZeroVector;
 					for (const FStrandsRootData& StrandsRoot : StrandsRoots)
 					{
-						const float Distance = FVector3f::Distance(StrandsRoot.Position, CardsRootPosition);
-						if (Distance < ClosestDistance)
+						const float Distance0 = FVector3f::Distance(StrandsRoot.Position, CardsRootPosition0);
+						const float Distance1 = FVector3f::Distance(StrandsRoot.Position, CardsRootPosition1);
+
+						if (Distance0 < ClosestDistance0)
+						{
+							ClosestDistance0 = Distance0;
+							OutRootUV0 = StrandsRoot.RootUV;
+						}
+
+						if (Distance1 < ClosestDistance1)
 						{
-							ClosestDistance = Distance;
-							CurveIndex = StrandsRoot.CurveIndex;
-							RootUV = StrandsRoot.RootUV;
+							ClosestDistance1 = Distance1;
+							OutRootUV1 = StrandsRoot.RootUV;
 						}
 					}
-					return RootUV;
+
 				};
-				const FVector2f RootUV0 = FindRootUV(CardsRoot.Root0.Position);
-				const FVector2f RootUV1 = FindRootUV(CardsRoot.Root1.Position);
+
+				FVector2f RootUV0, RootUV1;
+				FindRootUV(CardsRoot.Root0.Position, CardsRoot.Root1.Position, RootUV0, RootUV1);
 
 				// 3.2 Apply root UV to all cards vertices
 				{
@@ -829,7 +844,7 @@ static bool InternalImportGeometry_WithGeneratedGuides(
 					}
 				}
 			}
-		}
+		});
 	}
 
 	return bSuccess;
diff --git a/Engine/Plugins/Runtime/HairStrands/Source/HairStrandsCore/Private/HairStrandsDatas.cpp b/Engine/Plugins/Runtime/HairStrands/Source/HairStrandsCore/Private/HairStrandsDatas.cpp
--- a/Engine/Plugins/Runtime/HairStrands/Source/HairStrandsCore/Private/HairStrandsDatas.cpp
+++ b/Engine/Plugins/Runtime/HairStrands/Source/HairStrandsCore/Private/HairStrandsDatas.cpp
@@ -94,7 +94,7 @@ void FHairStrandsDatas::Reset()
 	StrandsCurves.Reset();
 	StrandsPoints.Reset();
 	HairDensity = 1;
-	BoundingBox = FBox(EForceInit::ForceInit);
+	BoundingBox = FBox3f(EForceInit::ForceInit);
 }
 
 void FHairStrandsInterpolationDatas::Reset()
diff --git a/Engine/Plugins/Runtime/HairStrands/Source/HairStrandsCore/Public/HairStrandsDatas.h b/Engine/Plugins/Runtime/HairStrands/Source/HairStrandsCore/Public/HairStrandsDatas.h
--- a/Engine/Plugins/Runtime/HairStrands/Source/HairStrandsCore/Public/HairStrandsDatas.h
+++ b/Engine/Plugins/Runtime/HairStrands/Source/HairStrandsCore/Public/HairStrandsDatas.h
@@ -510,7 +510,7 @@ struct HAIRSTRANDSCORE_API FHairStrandsCurves
 	TArray<FIntVector> CurvesClosestGuideIDs;
 
 	/** Custom guid weights (indexed with StrandID) (optional) */
-	TArray<FVector> CurvesClosestGuideWeights;
+	TArray<FVector3f> CurvesClosestGuideWeights;
 
 	/** Flags for attributes */
 	uint32 AttributeFlags = 0;
@@ -547,7 +547,7 @@ struct HAIRSTRANDSCORE_API FHairStrandsDatas
 	float HairDensity = 1;
 
 	/* Strands bounding box */
-	FBox BoundingBox = FBox(EForceInit::ForceInit);
+	FBox3f BoundingBox = FBox3f(EForceInit::ForceInit);
 };
 
 struct HAIRSTRANDSCORE_API FHairStrandsBulkData : FHairStrandsBulkCommon
