Fix CIS issue and fix uninitialized/invalid groom bounding box.

diff --git a/Engine/Plugins/Runtime/HairStrands/Source/HairStrandsCore/Private/GroomBuilder.cpp b/Engine/Plugins/Runtime/HairStrands/Source/HairStrandsCore/Private/GroomBuilder.cpp
--- a/Engine/Plugins/Runtime/HairStrands/Source/HairStrandsCore/Private/GroomBuilder.cpp
+++ b/Engine/Plugins/Runtime/HairStrands/Source/HairStrandsCore/Private/GroomBuilder.cpp
@@ -40,7 +40,7 @@ static FAutoConsoleVariableRef CVarHairGroupIndexBuilder_MaxVoxelResolution(TEXT
 
 FString FGroomBuilder::GetVersion()
 {
-	return TEXT("v8r54");
+	return TEXT("v8r55");
 }
 
 namespace FHairStrandsDecimation
@@ -378,6 +378,7 @@ namespace HairStrandsBuilder
 
 		OutBulkData.Header.BoundingBox.Min = (FVector)HairStrands.BoundingBox.Min;
 		OutBulkData.Header.BoundingBox.Max = (FVector)HairStrands.BoundingBox.Max;
+		OutBulkData.Header.BoundingBox.IsValid = 1;
 		OutBulkData.Header.CurveCount = HairStrands.GetNumCurves();
 		OutBulkData.Header.PointCount = HairStrands.GetNumPoints();
 		OutBulkData.Header.MinPointPerCurve = MinPointPerCurve;
@@ -2997,11 +2998,8 @@ static void BuildClusterData(
 	}
 
 	// 2. Compute control point LOD
-	TAtomic<uint32> ConcurrentPointCountPerLOD[FHairClusterInfo::MaxLOD];
-	for (uint8 LODIt = 0; LODIt < FHairClusterInfo::MaxLOD; ++LODIt)
-	{
-		ConcurrentPointCountPerLOD[LODIt] = 0;
-	}
+	static_assert(FHairClusterInfo::MaxLOD == 8);
+	TAtomic<uint32> ConcurrentPointCountPerLOD[FHairClusterInfo::MaxLOD] = { 0, 0, 0, 0, 0, 0, 0, 0 };
 	Out.PointLODs.SetNum(InRenStrandsData.GetNumPoints());
 	ParallelFor(Out.CurveCount, 
 	[
