Fix out of bound MH export due to reaching control point per curve limits.

diff --git a/Engine/Plugins/Runtime/HairStrands/Source/HairStrandsCore/Private/GroomRBFDeformer.cpp b/Engine/Plugins/Runtime/HairStrands/Source/HairStrandsCore/Private/GroomRBFDeformer.cpp
--- a/Engine/Plugins/Runtime/HairStrands/Source/HairStrandsCore/Private/GroomRBFDeformer.cpp
+++ b/Engine/Plugins/Runtime/HairStrands/Source/HairStrandsCore/Private/GroomRBFDeformer.cpp
@@ -289,10 +289,10 @@ struct FRBFDeformedPositions
 	TArray<FVector3f> RenderStrands;
 	TArray<FVector3f> GuideStrands;
 
-	// Trimmed data
+	// Masked data
 	TArray<bool> bIsRenderCurveValid;
 	TArray<bool> bIsRenderVertexValid;
-	bool bHasTrimmedData = false;
+	bool bHasMaskedData = false;
 };
 
 #if WITH_EDITORONLY_DATA
@@ -300,24 +300,44 @@ static void ApplyDeformationToGroom(const TArray<FRBFDeformedPositions>& Deforme
 {
 	// Notes:
 	// * This function applies the deformed position onto the hair description.
-	// * In addition, the groom can contain vertex/curves which needs to be trimmed (e.g., which are masked or cut.
+	// * In addition, the groom can contain vertex/curves which needs to be masked (e.g., which are masked or cut.
 	//   For handling this, the following function will first applied the deformed vertices and then remap all
 	//   curves/vertices so that hair description contains only valid curves/vertices
 	//
-	// /!\ Only rendering strands can be trimmed. Guide always remains unchanged.
+	// /!\ Only rendering strands can be masked. Guide always remains unchanged.
+
+	// The deformation must be stored in the HairDescription to rebuild the hair data when the groom is loaded
+	FHairDescription HairDescription = GroomAsset->GetHairDescription();
 
-	// Check if any group has trimmed data
+	// Check if any group has masked or trimmed data
+	bool bHasMaskedData = false;
 	bool bHasTrimmedData = false;
-	for (const FRBFDeformedPositions& Group : DeformedPositions)
 	{
-		if (Group.bHasTrimmedData)
+		FHairDescriptionGroups HairDescriptionGroups;
+		FGroomBuilder::BuildHairDescriptionGroups(HairDescription, HairDescriptionGroups);
+
+		const uint32 GroupCount = DeformedPositions.Num();
+		check(GroupCount == HairDescriptionGroups.HairGroups.Num());
+		for (uint32 GroupIt = 0; GroupIt < GroupCount; ++GroupIt)
 		{
-			bHasTrimmedData = true;
-			break;
+			const FHairDescriptionGroup& DescGroup = HairDescriptionGroups.HairGroups[GroupIt];
+			const FRBFDeformedPositions& RBFGroup  = DeformedPositions[GroupIt];
+			if (RBFGroup.bHasMaskedData)
+			{
+				bHasMaskedData = true;
+				break;
+			}
+
+			const bool bTrimPoint = (DescGroup.Info.Flags & uint32(EHairGroupInfoFlags::HasTrimmedPoint)) != 0;
+			const bool bTrimCurve = (DescGroup.Info.Flags & uint32(EHairGroupInfoFlags::HasTrimmedCurve)) != 0;
+			if (bTrimPoint || bTrimCurve)
+			{
+				bHasTrimmedData = true;
+			}
 		}
 	}
 
-	// If there are some trimmed data, compute the curves & vertices remapping
+	// If there are some masked/trimmed data, compute the curves & vertices remapping
 	int32 CurrentStrandID = 0;
 	int32 CurrentVertexID = 0;
 
@@ -327,8 +347,7 @@ static void ApplyDeformationToGroom(const TArray<FRBFDeformedPositions>& Deforme
 	uint32 TotalValidNumStrands = 0;
 	uint32 TotalValidNumVertices = 0;
 
-	// The deformation must be stored in the HairDescription to rebuild the hair data when the groom is loaded
-	FHairDescription HairDescription = GroomAsset->GetHairDescription();
+
 
 	TotalNumVertices = HairDescription.GetNumVertices();
 	TotalNumStrands = HairDescription.GetNumStrands();
@@ -371,6 +390,7 @@ static void ApplyDeformationToGroom(const TArray<FRBFDeformedPositions>& Deforme
 	TMap<int32, int32> GroupIDToGroupIndex;
 
 	const int32 GroomNumStrands = HairDescription.GetNumStrands();
+	check(GroomNumStrands < HAIR_MAX_NUM_CURVE_PER_GROUP);
 	for (int32 StrandIndex = 0; StrandIndex < GroomNumStrands; ++StrandIndex)
 	{
 		FStrandID StrandID(StrandIndex);
@@ -399,11 +419,21 @@ static void ApplyDeformationToGroom(const TArray<FRBFDeformedPositions>& Deforme
 			GroupInfo.NumGuideVertices += NumVertices;
 			for (int32 Index = 0; Index < NumVertices; ++Index)
 			{
-				if (GroupInfo.CurrentGuideVertexIndex < DeformedPositions[GroupIndex].GuideStrands.Num())
+				if (Index < HAIR_MAX_NUM_POINT_PER_CURVE)
 				{
-					// Guide vertex are never trimmed
-					FlattenedIsValid.Add(true);
-					FlattenedDeformedPositions.Add(DeformedPositions[GroupIndex].GuideStrands[GroupInfo.CurrentGuideVertexIndex++]);
+					if (GroupInfo.CurrentGuideVertexIndex < DeformedPositions[GroupIndex].GuideStrands.Num())
+					{
+						// Guide vertex are never masked
+						FlattenedIsValid.Add(true);
+						FlattenedDeformedPositions.Add(DeformedPositions[GroupIndex].GuideStrands[GroupInfo.CurrentGuideVertexIndex]);
+						++GroupInfo.CurrentGuideVertexIndex;
+					}
+				}
+				else
+				{
+					// Mark all invalid/trimmed control point as invalid and set their position to last control point to avoid bound computation issue
+					FlattenedIsValid.Add(false);
+					FlattenedDeformedPositions.Add(DeformedPositions[GroupIndex].GuideStrands[GroupInfo.CurrentRenderVertexIndex-1]);
 				}
 			}
 		}
@@ -412,18 +442,31 @@ static void ApplyDeformationToGroom(const TArray<FRBFDeformedPositions>& Deforme
 			GroupInfo.NumRenderVertices += NumVertices;
 			for (int32 Index = 0; Index < NumVertices; ++Index)
 			{
-				if (GroupInfo.CurrentRenderVertexIndex < DeformedPositions[GroupIndex].RenderStrands.Num())
+				if (Index < HAIR_MAX_NUM_POINT_PER_CURVE)
 				{
-					FVertexID VertexID(FlattenedDeformedPositions.Num());
-
-					const bool bIsValid = bHasTrimmedData ? DeformedPositions[GroupIndex].bIsRenderVertexValid[GroupInfo.CurrentRenderVertexIndex] : true;
-					FlattenedIsValid.Add(bIsValid);
-					FlattenedDeformedPositions.Add(DeformedPositions[GroupIndex].RenderStrands[GroupInfo.CurrentRenderVertexIndex++]);
+					check(GroupInfo.CurrentRenderVertexIndex < DeformedPositions[GroupIndex].RenderStrands.Num());
+					if (GroupInfo.CurrentRenderVertexIndex < DeformedPositions[GroupIndex].RenderStrands.Num())
+					{
+						FVertexID VertexID(FlattenedDeformedPositions.Num());
+	
+						const bool bIsValid = bHasMaskedData ? DeformedPositions[GroupIndex].bIsRenderVertexValid[GroupInfo.CurrentRenderVertexIndex] : true;
+						FlattenedIsValid.Add(bIsValid);
+						FlattenedDeformedPositions.Add(DeformedPositions[GroupIndex].RenderStrands[GroupInfo.CurrentRenderVertexIndex]);
+						++GroupInfo.CurrentRenderVertexIndex;
+					}
+				}
+				else
+				{
+					// Mark all invalid/trimmed control point as invalid and set their position to last control point to avoid bound computation issue
+					FlattenedIsValid.Add(false);
+					FlattenedDeformedPositions.Add(DeformedPositions[GroupIndex].RenderStrands[GroupInfo.CurrentRenderVertexIndex-1]);
 				}
 			}
 		}
 	}
 
+	check(FlattenedDeformedPositions.Num() == GroomNumVertices);
+
 	// Output the flattened deformed positions into the HairDescription
 	TVertexAttributesRef<FVector3f> VertexPositions = HairDescription.VertexAttributes().GetAttributesRef<FVector3f>(HairAttribute::Vertex::Position);
 	for (int32 VertexIndex = 0; VertexIndex < GroomNumVertices; ++VertexIndex)
@@ -432,8 +475,8 @@ static void ApplyDeformationToGroom(const TArray<FRBFDeformedPositions>& Deforme
 		VertexPositions[VertexID] = FlattenedDeformedPositions[VertexIndex];
 	}
 
-	// Apply hair description trimming if needed
-	if (bHasTrimmedData)
+	// Apply hair description masking/trimming if needed
+	if (bHasMaskedData || bHasTrimmedData)
 	{
 		TVertexAttributesRef<float> RWVertexWidth = HairDescription.VertexAttributes().GetAttributesRef<float>(HairAttribute::Vertex::Width);
 		for (int32 VertexIndex = 0; VertexIndex < GroomNumVertices; ++VertexIndex)
@@ -727,7 +770,7 @@ void FGroomRBFDeformer::GetRBFDeformedGroomAsset(const UGroomAsset* InGroomAsset
 			// Trim strands based on mask texture
 			DeformedPositions[GroupIndex].bIsRenderVertexValid.Init(true, StrandsData.GetNumPoints());
 			DeformedPositions[GroupIndex].bIsRenderCurveValid.Init(true, StrandsData.GetNumCurves());
-			DeformedPositions[GroupIndex].bHasTrimmedData = false;
+			DeformedPositions[GroupIndex].bHasMaskedData = false;
 			if (MaskTextureSource)
 			{
 				check(MaskTextureSource->GetNumBlocks() == 1);
@@ -766,13 +809,13 @@ void FGroomRBFDeformer::GetRBFDeformedGroomAsset(const UGroomAsset* InGroomAsset
 							if (CoordUThreshold <= 0.f)
 							{
 								DeformedPositions[GroupIndex].bIsRenderVertexValid[PointGlobalIndex0] = false;
-								DeformedPositions[GroupIndex].bHasTrimmedData = true;
+								DeformedPositions[GroupIndex].bHasMaskedData = true;
 							}
 							else if (VertexIndex + 1 < CurveNumVertices)
 							{
 								DeformedPositions[GroupIndex].bIsRenderVertexValid[PointGlobalIndex1] = true;
 
-								// Interpolate position or trim vertex
+								// Interpolate position or masked vertex
 								const float CoordU1 = StrandsData.StrandsPoints.PointsCoordU[PointGlobalIndex1];
 								if (CoordU0 <= CoordUThreshold && CoordU1 > CoordUThreshold)
 								{
@@ -783,7 +826,7 @@ void FGroomRBFDeformer::GetRBFDeformedGroomAsset(const UGroomAsset* InGroomAsset
 										StrandsData.StrandsPoints.PointsPosition[PointGlobalIndex1],
 										S);
 
-									DeformedPositions[GroupIndex].bHasTrimmedData = true;
+									DeformedPositions[GroupIndex].bHasMaskedData = true;
 								}
 								else if (CoordU0 > CoordUThreshold)
 								{
@@ -794,10 +837,10 @@ void FGroomRBFDeformer::GetRBFDeformedGroomAsset(const UGroomAsset* InGroomAsset
 							else if (CoordU0 > CoordUThreshold)
 							{
 								DeformedPositions[GroupIndex].bIsRenderVertexValid[PointGlobalIndex0] = false;
-								DeformedPositions[GroupIndex].bHasTrimmedData = true;
+								DeformedPositions[GroupIndex].bHasMaskedData = true;
 							}
 
-							// Mark the entire curve as trimmed if the first or second vertex are trimmed, because a curve needs to have at least two valid points
+							// Mark the entire curve as masked if the first or second vertex are masked, because a curve needs to have at least two valid points
 							if (!DeformedPositions[GroupIndex].bIsRenderVertexValid[PointGlobalIndex0] && (VertexIndex == 0 || VertexIndex == 1))
 							{
 								DeformedPositions[GroupIndex].bIsRenderCurveValid[CurveIndex] = false;
