Fix incorrect hair mesh RT acceleration structure building while RT is disabled.

diff --git a/Engine/Plugins/Runtime/HairStrands/Source/HairStrandsCore/Private/HairStrandsInterpolation.cpp b/Engine/Plugins/Runtime/HairStrands/Source/HairStrandsCore/Private/HairStrandsInterpolation.cpp
--- a/Engine/Plugins/Runtime/HairStrands/Source/HairStrandsCore/Private/HairStrandsInterpolation.cpp
+++ b/Engine/Plugins/Runtime/HairStrands/Source/HairStrandsCore/Private/HairStrandsInterpolation.cpp
@@ -2352,7 +2352,7 @@ void ComputeHairStrandsInterpolation(
 
 			#if RHI_RAYTRACING
 			FHairGroupInstance::FMeshes::FLOD& LOD = Instance->Meshes.LODs[HairLODIndex];
-			if (LOD.RaytracingResource)
+			if (LOD.RaytracingResource && IsRayTracingEnabled())
 			{
 				FCardsOrMeshesResourceBLASParameters* Parameters = GraphBuilder.AllocParameters<FCardsOrMeshesResourceBLASParameters>();
 				Parameters->PositionBuffer = LOD.DeformedResource ? Register(GraphBuilder, LOD.DeformedResource->GetBuffer(FHairMeshesDeformedResource::Current), ERDGImportedBufferFlags::None).Buffer : nullptr;
@@ -2374,13 +2374,13 @@ void ComputeHairStrandsInterpolation(
 						if (bLocalNeedBuild)
 						{
 							BuildHairAccelerationStructure_Meshes(RHICmdList, LocalLOD.RestResource, LocalLOD.DeformedResource,  &LocalLOD.RaytracingResource->RayTracingGeometry, Instance->Debug.GroomAssetName, HairLODIndex);
+							LocalLOD.RaytracingResource->bIsRTGeometryInitialized = true;
 						}
 						else if (bNeedUpdate)
 						{
 							// TODO: evaluate perf tradeoff of rebuild vs refit
 							UpdateHairAccelerationStructure(RHICmdList, &LocalLOD.RaytracingResource->RayTracingGeometry, EAccelerationStructureBuildMode::Update);
 						}
-						LocalLOD.RaytracingResource->bIsRTGeometryInitialized = true;
 					});
 				}
 			}
