Fix groom slate style registered multiple times.

diff --git a/Engine/Plugins/Runtime/HairStrands/Source/HairStrandsEditor/Public/GroomEditorStyle.h b/Engine/Plugins/Runtime/HairStrands/Source/HairStrandsEditor/Public/GroomEditorStyle.h
--- a/Engine/Plugins/Runtime/HairStrands/Source/HairStrandsEditor/Public/GroomEditorStyle.h
+++ b/Engine/Plugins/Runtime/HairStrands/Source/HairStrandsEditor/Public/GroomEditorStyle.h
@@ -34,7 +34,10 @@ public:
 		Set("GroomEditor.SimulationOptions", new IMAGE_BRUSH("S_SimulationOptions_40x", Icon40x40));
 		Set("GroomEditor.SimulationOptions.Small", new IMAGE_BRUSH("S_SimulationOptions_40x", Icon20x20));
 		Set("GroomEditor.GroomCardGenerator", new IMAGE_BRUSH("S_GroomCardGenerator_128", Icon16x16));
-		FSlateStyleRegistry::RegisterSlateStyle(*this);
+		if (FSlateStyleRegistry::FindSlateStyle(GetStyleName()) == nullptr)
+		{
+			FSlateStyleRegistry::RegisterSlateStyle(*this);
+		}
 	}
 
 	/** Return the name of this style */
