Add option to force velocity output for grooms.

diff --git a/Engine/Plugins/Runtime/HairStrands/Source/HairStrandsCore/Private/GroomComponent.cpp b/Engine/Plugins/Runtime/HairStrands/Source/HairStrandsCore/Private/GroomComponent.cpp
--- a/Engine/Plugins/Runtime/HairStrands/Source/HairStrandsCore/Private/GroomComponent.cpp
+++ b/Engine/Plugins/Runtime/HairStrands/Source/HairStrandsCore/Private/GroomComponent.cpp
@@ -75,6 +75,9 @@ static FAutoConsoleVariableRef CVarHairStrands_UseAttachedSimulationComponents(T
 static int32 GHairStrands_ViewModeClumpIndex = 0;
 static FAutoConsoleVariableRef CVarHairStrands_ViewModeClumpIndex(TEXT("r.HairStrands.ViewMode.ClumpIndex"), GHairStrands_ViewModeClumpIndex, TEXT("Define the ClumpID index (0, 1, or 2) which should be visualized"));
 
+static int32 GHairStrands_ForceVelocityOutput = 0;
+static FAutoConsoleVariableRef CVarHairStrands_ForceVelocityOutput(TEXT("r.HairStrands.ForceVelocityOutput"), GHairStrands_ForceVelocityOutput, TEXT("When enabled, force the cards/meshes to write velocity vectors."));
+
 #define LOCTEXT_NAMESPACE "GroomComponent"
 
 ////////////////////////////////////////////////////////////////////////////////////////////////////////////////////////
@@ -430,7 +433,7 @@ public:
 		ComponentId = Component->GetPrimitiveSceneId().PrimIDValue;
 		Strands_DebugMaterial = Component->Strands_DebugMaterial;
 		bAlwaysHasVelocity = false;
-		if (IsHairStrandsBindingEnable() && Component->RegisteredMeshComponent)
+		if ((IsHairStrandsBindingEnable() && Component->RegisteredMeshComponent) || GHairStrands_ForceVelocityOutput > 0)
 		{
 			bAlwaysHasVelocity = true;
 		}
@@ -966,6 +969,11 @@ public:
 		FPrimitiveSceneInfo* PrimSceneInfo = GetPrimitiveSceneInfo();
 		GetScene().GetPrimitiveUniformShaderParameters_RenderThread(PrimSceneInfo, bHasPrecomputedVolumetricLightmap, PreviousLocalToWorld, SingleCaptureIndex, bOutputVelocity);
 
+		if (GHairStrands_ForceVelocityOutput > 0)
+		{
+			bOutputVelocity = true;
+		}
+
 		const bool bUseProxy = UseProxyLocalToWorld(Instance);
 
 		FMatrix CurrentLocalToWorld = bUseProxy ? GetLocalToWorld() : Instance->GetCurrentLocalToWorld().ToMatrixWithScale();
