Re-enabled parallel GDME for hair after flickering fix.

diff --git a/Engine/Plugins/Runtime/HairStrands/Source/HairStrandsCore/Private/GroomComponent.cpp b/Engine/Plugins/Runtime/HairStrands/Source/HairStrandsCore/Private/GroomComponent.cpp
--- a/Engine/Plugins/Runtime/HairStrands/Source/HairStrandsCore/Private/GroomComponent.cpp
+++ b/Engine/Plugins/Runtime/HairStrands/Source/HairStrandsCore/Private/GroomComponent.cpp
@@ -553,9 +553,6 @@ public:
 				}
 			}
 		}
-
-		// TODO: Rendering hair strands currently requires a deterministic ordering of calls to GetDynamicMeshElements.
-		bSupportsParallelGDME = false;
 	}
 
 	virtual ~FHairStrandsSceneProxy()
