Fix unstable ordering of the visible instances.

diff --git a/Engine/Plugins/Runtime/HairStrands/Source/HairStrandsCore/Private/GroomManager.cpp b/Engine/Plugins/Runtime/HairStrands/Source/HairStrandsCore/Private/GroomManager.cpp
--- a/Engine/Plugins/Runtime/HairStrands/Source/HairStrandsCore/Private/GroomManager.cpp
+++ b/Engine/Plugins/Runtime/HairStrands/Source/HairStrandsCore/Private/GroomManager.cpp
@@ -1402,13 +1402,17 @@ void ProcessHairStrandsBookmark(
 	else if (Bookmark == EHairStrandsBookmark::ProcessDebug)
 	{
 		// Merge all visible instances
-		FHairStrandsInstances DebugInstance;
+		FHairStrandsInstances DebugInstances;
 		if (IsInstanceFrustumCullingEnable())
 		{
-			DebugInstance.Append(Parameters.VisibleStrands);
-			DebugInstance.Append(Parameters.VisibleCardsOrMeshes_Primary);
-			DebugInstance.Append(Parameters.VisibleCardsOrMeshes_Shadow);
-			Instances = &DebugInstance;
+			DebugInstances.Append(Parameters.VisibleStrands);
+			DebugInstances.Append(Parameters.VisibleCardsOrMeshes_Primary);
+			DebugInstances.Append(Parameters.VisibleCardsOrMeshes_Shadow);
+			DebugInstances.Sort([](const FHairStrandsInstance& A, const FHairStrandsInstance& B)
+			{
+				return A.RegisteredIndex < B.RegisteredIndex;
+			});
+			Instances = &DebugInstances;
 		}
 
 		check(GraphBuilder);
