Remove groom mesh update type and deformation type.

diff --git a/Engine/Plugins/Runtime/HairStrands/Source/HairStrandsCore/Private/GroomDebug.cpp b/Engine/Plugins/Runtime/HairStrands/Source/HairStrandsCore/Private/GroomDebug.cpp
--- a/Engine/Plugins/Runtime/HairStrands/Source/HairStrandsCore/Private/GroomDebug.cpp
+++ b/Engine/Plugins/Runtime/HairStrands/Source/HairStrandsCore/Private/GroomDebug.cpp
@@ -64,6 +64,12 @@ static FAutoConsoleVariableRef CVarHairCardsAtlasDebug(TEXT("r.HairStrands.Cards
 
 ///////////////////////////////////////////////////////////////////////////////////////////////////
 
+enum class HairStrandsTriangleType
+{
+	RestPose,
+	DeformedPose,
+};
+
 FCachedGeometry GetCacheGeometryForHair(
 	FRDGBuilder& GraphBuilder,
 	FSceneInterface* Scene,
diff --git a/Engine/Plugins/Runtime/HairStrands/Source/HairStrandsCore/Private/GroomManager.cpp b/Engine/Plugins/Runtime/HairStrands/Source/HairStrandsCore/Private/GroomManager.cpp
--- a/Engine/Plugins/Runtime/HairStrands/Source/HairStrandsCore/Private/GroomManager.cpp
+++ b/Engine/Plugins/Runtime/HairStrands/Source/HairStrandsCore/Private/GroomManager.cpp
@@ -292,7 +292,6 @@ static void RunInternalHairInterpolation(
 							GraphBuilder,
 							ShaderMap,
 							MeshLODIndex,
-							HairStrandsTriangleType::DeformedPose,
 							MeshDataLOD,
 							Instance->Strands.RestRootResource,
 							Instance->Strands.DeformedRootResource);
@@ -334,7 +333,6 @@ static void RunInternalHairInterpolation(
 								GraphBuilder,
 								ShaderMap,
 								MeshLODIndex,
-								HairStrandsTriangleType::DeformedPose,
 								MeshDataLOD,
 								CardsInstance.Guides.RestRootResource,
 								CardsInstance.Guides.DeformedRootResource);
@@ -380,7 +378,6 @@ static void RunInternalHairInterpolation(
 							GraphBuilder,
 							ShaderMap,
 							Instance->Debug.MeshLODIndex,
-							HairStrandsTriangleType::DeformedPose,
 							MeshDataLOD,
 							Instance->Guides.RestRootResource,
 							Instance->Guides.DeformedRootResource);
@@ -391,7 +388,6 @@ static void RunInternalHairInterpolation(
 								GraphBuilder,
 								ShaderMap,
 								Instance->Debug.MeshLODIndex,
-								HairStrandsTriangleType::DeformedPose,
 								MeshDataLOD,
 								Instance->Guides.RestRootResource,
 								Instance->Guides.DeformedRootResource);
diff --git a/Engine/Plugins/Runtime/HairStrands/Source/HairStrandsCore/Private/HairStrandsInterpolation.cpp b/Engine/Plugins/Runtime/HairStrands/Source/HairStrandsCore/Private/HairStrandsInterpolation.cpp
--- a/Engine/Plugins/Runtime/HairStrands/Source/HairStrandsCore/Private/HairStrandsInterpolation.cpp
+++ b/Engine/Plugins/Runtime/HairStrands/Source/HairStrandsCore/Private/HairStrandsInterpolation.cpp
@@ -84,13 +84,6 @@ inline uint32 ComputeGroupSize()
 	return GroupSize;
 }
 
-enum class EDeformationType : uint8
-{
-	Simulation,		// Use the output of the hair simulation
-	RestGuide,		// Use the rest guide as input of the interpolation (no deformation), only weighted interpolation
-	OffsetGuide		// Offset the guides
-};
-
 void GetHairStrandsAttributeParameter(const FHairStrandsBulkData& In, FHairStrandsInstanceAttributeParameters& Out)
 {
 	check(FMath::IsPowerOfTwo(In.Header.Strides.CurveAttributeChunkElementCount));
@@ -343,7 +336,6 @@ IMPLEMENT_GLOBAL_SHADER(FDeformGuideCS, "/Engine/Private/HairStrands/HairStrands
 static void AddDeformSimHairStrandsPass(
 	FRDGBuilder& GraphBuilder,
 	FGlobalShaderMap* ShaderMap,
-	EDeformationType DeformationType,
 	const uint32 MeshLODIndex,
 	const uint32 VertexCount,
 	FHairStrandsRestRootResource* SimRestRootResources,
@@ -366,14 +358,7 @@ static void AddDeformSimHairStrandsPass(
 		InternalDeformationTypeCount
 	};
 
-	EInternalDeformationType InternalDeformationType = InternalDeformationTypeCount;
-	switch (DeformationType)
-	{
-	case EDeformationType::RestGuide   : InternalDeformationType = InternalDeformationType_ByPass; break;
-	case EDeformationType::OffsetGuide : InternalDeformationType = InternalDeformationType_Offset; break;
-	}
-
-	if (InternalDeformationType == InternalDeformationTypeCount) return;
+	EInternalDeformationType InternalDeformationType = InternalDeformationType_Offset;
 
 	FDeformGuideCS::FParameters* Parameters = GraphBuilder.AllocParameters<FDeformGuideCS::FParameters>();
 	Parameters->SimRestPosePositionBuffer = SimRestPosePositionBuffer;
@@ -381,8 +366,6 @@ static void AddDeformSimHairStrandsPass(
 	Parameters->VertexCount = VertexCount;
 	Parameters->SimDeformedOffsetBuffer = SimDeformedOffsetBuffer;
 	Parameters->SimRestOffset = (FVector3f)SimRestOffset;
-
-	if (DeformationType == EDeformationType::OffsetGuide)
 	{
 		const bool bIsPointToCurveBuffersValid = SimPointToCurveBuffer != nullptr;
 		if (bIsPointToCurveBuffersValid)
@@ -1858,7 +1841,6 @@ void ComputeHairStrandsInterpolation(
 								AddDeformSimHairStrandsPass(
 									GraphBuilder,
 									ShaderMap,
-									EDeformationType::OffsetGuide,
 									MeshLODIndex,
 									Instance->Guides.RestResource->GetPointCount(),
 									Instance->Guides.RestRootResource,
@@ -2416,7 +2398,6 @@ void ResetHairStrandsInterpolation(
 	AddDeformSimHairStrandsPass(
 		GraphBuilder,
 		ShaderMap,
-		EDeformationType::OffsetGuide,
 		MeshLODIndex,
 		Instance->Guides.RestResource->GetPointCount(),
 		Instance->Guides.RestRootResource,
diff --git a/Engine/Plugins/Runtime/HairStrands/Source/HairStrandsCore/Private/HairStrandsMeshProjection.cpp b/Engine/Plugins/Runtime/HairStrands/Source/HairStrandsCore/Private/HairStrandsMeshProjection.cpp
--- a/Engine/Plugins/Runtime/HairStrands/Source/HairStrandsCore/Private/HairStrandsMeshProjection.cpp
+++ b/Engine/Plugins/Runtime/HairStrands/Source/HairStrandsCore/Private/HairStrandsMeshProjection.cpp
@@ -275,7 +275,6 @@ void AddHairStrandUpdateMeshTrianglesPass(
 	FRDGBuilder& GraphBuilder,
 	FGlobalShaderMap* ShaderMap,
 	const int32 LODIndex,
-	const HairStrandsTriangleType Type,
 	const FHairStrandsProjectionMeshData::LOD& MeshData,
 	FHairStrandsRestRootResource* RestResources,
 	FHairStrandsDeformedRootResource* DeformedResources)
@@ -286,12 +285,7 @@ void AddHairStrandUpdateMeshTrianglesPass(
 		return;
 	}
 
-	if (Type == HairStrandsTriangleType::RestPose && LODIndex >= RestResources->LODs.Num())
-	{
-		return;
-	}
-
-	if (Type == HairStrandsTriangleType::DeformedPose && (LODIndex >= RestResources->LODs.Num() || LODIndex >= DeformedResources->LODs.Num()))
+	if ((LODIndex >= RestResources->LODs.Num() || LODIndex >= DeformedResources->LODs.Num()))
 	{
 		return;
 	}
@@ -314,7 +308,7 @@ void AddHairStrandUpdateMeshTrianglesPass(
 	const TArray<uint32>& ValidSectionIndices = RestResources->BulkData.GetValidSectionIndices(LODIndex);
 	const uint32 ValidSectionCount = ValidSectionIndices.Num();
 	const uint32 PassCount = FMath::DivideAndRoundUp(ValidSectionCount, FHairUpdateMeshTriangleCS::SectionArrayCount);
-	const bool bComputePreviousDeformedPosition = IsHairStrandContinuousDecimationReorderingEnabled() && Type == HairStrandsTriangleType::DeformedPose;
+	const bool bComputePreviousDeformedPosition = IsHairStrandContinuousDecimationReorderingEnabled();
 
 	const uint32 MaxUniqueTriangleCount = RestLODData.RestUniqueTrianglePositionBuffer.Buffer->Desc.NumElements * 3;
 
@@ -325,11 +319,6 @@ void AddHairStrandUpdateMeshTrianglesPass(
 	const bool bEnableUAVOverlap = true;
 	const ERDGUnorderedAccessViewFlags UAVFlags = bEnableUAVOverlap ? ERDGUnorderedAccessViewFlags::SkipBarrier : ERDGUnorderedAccessViewFlags::None;
 	CommonParameters.UniqueTriangleIndices = RegisterAsSRV(GraphBuilder, RestLODData.UniqueTriangleIndexBuffer);
-	if (Type == HairStrandsTriangleType::RestPose)
-	{
-		OutputCurrBuffer = Register(GraphBuilder, RestLODData.RestUniqueTrianglePositionBuffer, ERDGImportedBufferFlags::CreateUAV, UAVFlags);
-	}
-	else if (Type == HairStrandsTriangleType::DeformedPose)
 	{
 		FHairStrandsDeformedRootResource::FLOD& DeformedLODData = DeformedResources->LODs[LODIndex];
 		OutputCurrBuffer = Register(GraphBuilder, DeformedLODData.GetDeformedUniqueTrianglePositionBuffer(FHairStrandsDeformedRootResource::FLOD::Current), ERDGImportedBufferFlags::CreateUAV, UAVFlags);
@@ -341,11 +330,6 @@ void AddHairStrandUpdateMeshTrianglesPass(
 
 		DeformedLODData.Status = FHairStrandsDeformedRootResource::FLOD::EStatus::Completed;
 	}
-	else
-	{
-		// error
-		return;
-	}
 
 	CommonParameters.OutUniqueTriangleCurrPosition = OutputCurrBuffer.UAV;
 
@@ -724,7 +708,6 @@ void AddHairStrandInitMeshSamplesPass(
 	FRDGBuilder& GraphBuilder,
 	FGlobalShaderMap* ShaderMap,
 	const int32 LODIndex,
-	const HairStrandsTriangleType Type,
 	const FHairStrandsProjectionMeshData::LOD& MeshData,
 	FHairStrandsRestRootResource* RestResources,
 	FHairStrandsDeformedRootResource* DeformedResources)
@@ -734,12 +717,7 @@ void AddHairStrandInitMeshSamplesPass(
 		return;
 	}
 
-	if (Type == HairStrandsTriangleType::RestPose && LODIndex >= RestResources->LODs.Num())
-	{
-		return;
-	}
-
-	if (Type == HairStrandsTriangleType::DeformedPose && (LODIndex >= RestResources->LODs.Num() || LODIndex >= DeformedResources->LODs.Num()))
+	if ((LODIndex >= RestResources->LODs.Num() || LODIndex >= DeformedResources->LODs.Num()))
 	{
 		return;
 	}
@@ -763,21 +741,12 @@ void AddHairStrandInitMeshSamplesPass(
 	if (SectionCount > 0 && RestLODData.SampleCount > 0)
 	{
 		FRDGImportedBuffer OutBuffer;
-		if (Type == HairStrandsTriangleType::RestPose)
-		{
-			OutBuffer = Register(GraphBuilder, RestLODData.RestSamplePositionsBuffer, ERDGImportedBufferFlags::CreateUAV);
-		}
-		else if (Type == HairStrandsTriangleType::DeformedPose)
 		{
 			FHairStrandsDeformedRootResource::FLOD& DeformedLODData = DeformedResources->LODs[LODIndex];
 			check(DeformedLODData.LODIndex == LODIndex);
 
 			OutBuffer = Register(GraphBuilder, DeformedLODData.GetDeformedSamplePositionsBuffer(FHairStrandsDeformedRootResource::FLOD::Current), ERDGImportedBufferFlags::CreateUAV);
 		}
-		else
-		{
-			return;
-		}
 		for (uint32 PassIt = 0; PassIt < PassCount; ++PassIt)
 		{
 			FHairInitMeshSamplesCS::FParameters Parameters;
diff --git a/Engine/Plugins/Runtime/HairStrands/Source/HairStrandsCore/Public/HairStrandsMeshProjection.h b/Engine/Plugins/Runtime/HairStrands/Source/HairStrandsCore/Public/HairStrandsMeshProjection.h
--- a/Engine/Plugins/Runtime/HairStrands/Source/HairStrandsCore/Public/HairStrandsMeshProjection.h
+++ b/Engine/Plugins/Runtime/HairStrands/Source/HairStrandsCore/Public/HairStrandsMeshProjection.h
@@ -52,18 +52,11 @@ struct FHairStrandsProjectionMeshData
 uint32 GetHairStrandsMaxSectionCount();
 uint32 GetHairStrandsMaxTriangleCount();
 
-enum class HairStrandsTriangleType
-{
-	RestPose,
-	DeformedPose,
-};
-
 /* Update the triangles information on which hair stands have been projected */
 void AddHairStrandUpdateMeshTrianglesPass(
 	FRDGBuilder& GraphBuilder,
 	FGlobalShaderMap* ShaderMap,
 	const int32 LODIndex,
-	const HairStrandsTriangleType Type,
 	const FHairStrandsProjectionMeshData::LOD& MeshData,
 	FHairStrandsRestRootResource* RestResources,
 	FHairStrandsDeformedRootResource* DeformedResources);
@@ -73,7 +66,6 @@ void AddHairStrandInitMeshSamplesPass(
 	FRDGBuilder& GraphBuilder,
 	FGlobalShaderMap* ShaderMap,
 	const int32 LODIndex,
-	const HairStrandsTriangleType Type,
 	const FHairStrandsProjectionMeshData::LOD& MeshData,
 	FHairStrandsRestRootResource* RestResources,
 	FHairStrandsDeformedRootResource* DeformedResources);
