Rearrange groom simulation pass to loop over instances for furture batching.

diff --git a/Engine/Plugins/Runtime/HairStrands/Source/HairStrandsCore/Private/GroomManager.cpp b/Engine/Plugins/Runtime/HairStrands/Source/HairStrandsCore/Private/GroomManager.cpp
--- a/Engine/Plugins/Runtime/HairStrands/Source/HairStrandsCore/Private/GroomManager.cpp
+++ b/Engine/Plugins/Runtime/HairStrands/Source/HairStrandsCore/Private/GroomManager.cpp
@@ -361,96 +361,6 @@ static void RunInternalHairInterpolation(
 					// Nothing to do
 				}
 			}
-			else if (EHairInterpolationPassType::Guides == PassType)
-			{
-				// Guide update need to run only if simulation is enabled, or if RBF is enabled (since RFB are transfer through guides)
-				if (bGlobalDeformationEnable || bSimulationEnable || bDeformationEnable)
-				{
-					check(Instance->Guides.IsValid());
-
-					if (BindingType == EHairBindingType::Skinning)
-					{
-						check(Instance->Guides.IsValid());
-						check(Instance->Guides.HasValidRootData());
-						check(Instance->Guides.DeformedRootResource->IsValid(MeshLODIndex));
-
-						AddHairStrandUpdateMeshTrianglesPass(
-							GraphBuilder,
-							ShaderMap,
-							Instance->Debug.MeshLODIndex,
-							MeshDataLOD,
-							Instance->Guides.RestRootResource,
-							Instance->Guides.DeformedRootResource);
-
-						if (bGlobalDeformationEnable)
-						{
-							AddHairStrandInitMeshSamplesPass(
-								GraphBuilder,
-								ShaderMap,
-								Instance->Debug.MeshLODIndex,
-								MeshDataLOD,
-								Instance->Guides.RestRootResource,
-								Instance->Guides.DeformedRootResource);
-
-							AddHairStrandUpdateMeshSamplesPass(
-								GraphBuilder,
-								ShaderMap,
-								Instance->Debug.MeshLODIndex,
-								MeshDataLOD,
-								Instance->Guides.RestRootResource,
-								Instance->Guides.DeformedRootResource);
-						}
-
-						AddHairStrandUpdatePositionOffsetPass(
-							GraphBuilder,
-							ShaderMap,
-							Instance->Debug.MeshLODIndex,
-							Instance->Guides.DeformedRootResource,
-							Instance->Guides.DeformedResource);
-
-						// Add manual transition for the GPU solver as Niagara does not track properly the RDG buffer, and so doesn't issue the correct transitions
-						if (MeshLODIndex >= 0)
-						{
-							GraphBuilder.UseExternalAccessMode(Register(GraphBuilder, Instance->Guides.DeformedRootResource->LODs[MeshLODIndex].GetDeformedUniqueTrianglePositionBuffer(FHairStrandsDeformedRootResource::FLOD::Current), ERDGImportedBufferFlags::CreateSRV).Buffer, ERHIAccess::SRVMask);
-
-							if (bGlobalDeformationEnable)
-							{
-								GraphBuilder.UseExternalAccessMode(Register(GraphBuilder, Instance->Guides.DeformedRootResource->LODs[MeshLODIndex].GetDeformedSamplePositionsBuffer(FHairStrandsDeformedRootResource::FLOD::Current), ERDGImportedBufferFlags::CreateSRV).Buffer, ERHIAccess::SRVMask);
-								GraphBuilder.UseExternalAccessMode(Register(GraphBuilder, Instance->Guides.DeformedRootResource->LODs[MeshLODIndex].GetMeshSampleWeightsBuffer(FHairStrandsDeformedRootResource::FLOD::Current), ERDGImportedBufferFlags::CreateSRV).Buffer, ERHIAccess::SRVMask);
-							}
-						}
-						GraphBuilder.UseExternalAccessMode(Register(GraphBuilder, Instance->Guides.DeformedResource->GetPositionOffsetBuffer(FHairStrandsDeformedResource::EFrameType::Current), ERDGImportedBufferFlags::CreateSRV).Buffer, ERHIAccess::SRVMask);
-					}
-					else if (bSimulationEnable || bDeformationEnable)
-					{
-						check(Instance->Guides.IsValid());
-
-						AddHairStrandUpdatePositionOffsetPass(
-							GraphBuilder,
-							ShaderMap,
-							Instance->Debug.MeshLODIndex,
-							nullptr,
-							Instance->Guides.DeformedResource);
-
-						// Add manual transition for the GPU solver as Niagara does not track properly the RDG buffer, and so doesn't issue the correct transitions
-						GraphBuilder.UseExternalAccessMode(Register(GraphBuilder, Instance->Guides.DeformedResource->GetPositionOffsetBuffer(FHairStrandsDeformedResource::EFrameType::Current), ERDGImportedBufferFlags::CreateSRV).Buffer, ERHIAccess::SRVMask);
-					}
-				}
-			}
-		}
-	}
-
-	// Reset deformation
-	if (EHairInterpolationPassType::Guides == PassType)
-	{
-		for (FHairStrandsInstance* AbstractInstance : Instances)
-		{
-			FHairGroupInstance* Instance = static_cast<FHairGroupInstance*>(AbstractInstance);
-
-			if (Instance->GeometryType == EHairGeometryType::NoneGeometry)
-				continue;
-
-			ResetHairStrandsInterpolation(GraphBuilder, ShaderMap, Instance, Instance->Debug.MeshLODIndex);
 		}
 	}
 
@@ -499,16 +409,177 @@ static void RunHairStrandsInterpolation_Guide(
 	RDG_EVENT_SCOPE(GraphBuilder, "HairGuideInterpolation");
 	RDG_GPU_STAT_SCOPE(GraphBuilder, HairGuideInterpolation);
 
-	RunInternalHairInterpolation(
-		GraphBuilder,
-		Scene,
-		View,
-		ViewUniqueID,
-		Instances,
-		ShaderPrintData,
-		ShaderMap,
-		EHairInterpolationPassType::Guides,
-		nullptr);
+	// Update dynamic mesh triangles
+	struct FInstanceData
+	{
+		int32 HairLODIndex = -1;
+		int32 MeshLODIndex = -1;
+		EHairBindingType BindingType = EHairBindingType::NoneBinding;
+		bool bSimulationEnable = false;
+		bool bDeformationEnable = false;
+		bool bGlobalDeformationEnable = false;
+
+		FHairStrandsProjectionMeshData::LOD MeshDataLOD;
+		FHairGroupInstance* Instance = nullptr;
+
+		bool NeedsMeshUpdate() const
+		{
+			return 
+				MeshLODIndex >= 0 && 
+				(BindingType == EHairBindingType::Skinning) && 
+				(bGlobalDeformationEnable || bSimulationEnable || bDeformationEnable);
+		}
+	};
+
+	// Gather all instances which require guides or RBF deformations
+	TArray<FInstanceData> InstanceDatas;
+	InstanceDatas.Reserve(Instances.Num());
+	for (FHairStrandsInstance* AbstractInstance : Instances)
+	{
+		FHairGroupInstance* Instance = static_cast<FHairGroupInstance*>(AbstractInstance);
+		if (Instance->GeometryType == EHairGeometryType::NoneGeometry)
+		{
+			continue;
+		}	
+		check(Instance->HairGroupPublicData);
+
+		FInstanceData& InstanceData = InstanceDatas.AddDefaulted_GetRef();
+		InstanceData.Instance 					= Instance;
+		InstanceData.BindingType 				= Instance->BindingType;
+		InstanceData.HairLODIndex 				= Instance->HairGroupPublicData->LODIndex;
+		InstanceData.bSimulationEnable 			= Instance->HairGroupPublicData->IsSimulationEnable(InstanceData.HairLODIndex);
+		InstanceData.bDeformationEnable 		= Instance->HairGroupPublicData->bIsDeformationEnable;
+		InstanceData.bGlobalDeformationEnable 	= Instance->HairGroupPublicData->IsGlobalInterpolationEnable(InstanceData.HairLODIndex);
+		InstanceData.MeshLODIndex				= -1;
+
+		// Extract MeshDataLOD and compute MeshLODIndex
+		const FCachedGeometry CachedGeometry = GetCacheGeometryForHair(GraphBuilder, Scene, Instance, ShaderMap, true);
+		for (int32 SectionIndex = 0; SectionIndex < CachedGeometry.Sections.Num(); ++SectionIndex)
+		{
+			// Ensure all mesh's sections have the same LOD index
+			const int32 SectionLodIndex = CachedGeometry.Sections[SectionIndex].LODIndex;
+			if (InstanceData.MeshLODIndex < 0) InstanceData.MeshLODIndex = SectionLodIndex;
+			check(InstanceData.MeshLODIndex == SectionLodIndex);
+
+			InstanceData.MeshDataLOD.Sections.Add(ConvertMeshSection(CachedGeometry, SectionIndex));
+		}
+		Instance->Debug.MeshLODIndex = InstanceData.MeshLODIndex;
+
+		if (InstanceData.NeedsMeshUpdate())
+		{
+			check(InstanceData.Instance->Guides.IsValid());
+		}
+	}	
+
+	// Update dynamic mesh triangles
+	// Guide update need to run only if simulation is enabled, or if RBF is enabled (since RFB are transfer through guides)
+	for (FInstanceData& InstanceData : InstanceDatas)
+	{
+		if (InstanceData.NeedsMeshUpdate())
+		{		
+			check(InstanceData.Instance->Guides.IsValid());
+			check(InstanceData.Instance->Guides.HasValidRootData());
+			check(InstanceData.Instance->Guides.DeformedRootResource->IsValid(InstanceData.MeshLODIndex));
+
+			AddHairStrandUpdateMeshTrianglesPass(
+				GraphBuilder,
+				ShaderMap,
+				InstanceData.Instance->Debug.MeshLODIndex,
+				InstanceData.MeshDataLOD,
+				InstanceData.Instance->Guides.RestRootResource,
+				InstanceData.Instance->Guides.DeformedRootResource);
+		}
+	}
+
+	// Guide update need to run only if simulation is enabled, or if RBF is enabled (since RFB are transfer through guides)
+	for (FInstanceData& InstanceData : InstanceDatas)
+	{
+		if (InstanceData.NeedsMeshUpdate())
+		{		
+			if (InstanceData.bGlobalDeformationEnable)
+			{
+				AddHairStrandInitMeshSamplesPass(
+					GraphBuilder,
+					ShaderMap,
+					InstanceData.Instance->Debug.MeshLODIndex,
+					InstanceData.MeshDataLOD,
+					InstanceData.Instance->Guides.RestRootResource,
+					InstanceData.Instance->Guides.DeformedRootResource);
+			}
+		}
+	}
+
+	// Update position offset for each instance (GPU)
+	for (FInstanceData& InstanceData : InstanceDatas)
+	{
+		if (InstanceData.NeedsMeshUpdate())
+		{
+			AddHairStrandUpdatePositionOffsetPass(
+				GraphBuilder,
+				ShaderMap,
+				InstanceData.Instance->Debug.MeshLODIndex,
+				InstanceData.Instance->Guides.DeformedRootResource,
+				InstanceData.Instance->Guides.DeformedResource);
+
+			// Add manual transition for the GPU solver as Niagara does not track properly the RDG buffer, and so doesn't issue the correct transitions
+			{
+				GraphBuilder.UseExternalAccessMode(Register(GraphBuilder, InstanceData.Instance->Guides.DeformedRootResource->LODs[InstanceData.MeshLODIndex].GetDeformedUniqueTrianglePositionBuffer(FHairStrandsDeformedRootResource::FLOD::Current), ERDGImportedBufferFlags::CreateSRV).Buffer, ERHIAccess::SRVMask);
+
+				if (InstanceData.bGlobalDeformationEnable)
+				{
+					GraphBuilder.UseExternalAccessMode(Register(GraphBuilder, InstanceData.Instance->Guides.DeformedRootResource->LODs[InstanceData.MeshLODIndex].GetDeformedSamplePositionsBuffer(FHairStrandsDeformedRootResource::FLOD::Current), ERDGImportedBufferFlags::CreateSRV).Buffer, ERHIAccess::SRVMask);
+					GraphBuilder.UseExternalAccessMode(Register(GraphBuilder, InstanceData.Instance->Guides.DeformedRootResource->LODs[InstanceData.MeshLODIndex].GetMeshSampleWeightsBuffer(FHairStrandsDeformedRootResource::FLOD::Current), ERDGImportedBufferFlags::CreateSRV).Buffer, ERHIAccess::SRVMask);
+				}
+			}
+			GraphBuilder.UseExternalAccessMode(Register(GraphBuilder, InstanceData.Instance->Guides.DeformedResource->GetPositionOffsetBuffer(FHairStrandsDeformedResource::EFrameType::Current), ERDGImportedBufferFlags::CreateSRV).Buffer, ERHIAccess::SRVMask);
+		}
+	}
+
+	// Update position offset for each instance (CPU)
+	for (FInstanceData& InstanceData : InstanceDatas)
+	{
+		if (InstanceData.MeshLODIndex >= 0 && (InstanceData.BindingType != EHairBindingType::Skinning) && (InstanceData.bSimulationEnable || InstanceData.bDeformationEnable))
+		{		
+			AddHairStrandUpdatePositionOffsetPass(
+				GraphBuilder,
+				ShaderMap,
+				InstanceData.Instance->Debug.MeshLODIndex,
+				nullptr,
+				InstanceData.Instance->Guides.DeformedResource);
+
+			// Add manual transition for the GPU solver as Niagara does not track properly the RDG buffer, and so doesn't issue the correct transitions
+			GraphBuilder.UseExternalAccessMode(Register(GraphBuilder, InstanceData.Instance->Guides.DeformedResource->GetPositionOffsetBuffer(FHairStrandsDeformedResource::EFrameType::Current), ERDGImportedBufferFlags::CreateSRV).Buffer, ERHIAccess::SRVMask);
+		}
+	}
+
+	// Apply deformation from
+	// * Bones 
+	// * RBF
+	for (FInstanceData& InstanceData : InstanceDatas)
+	{
+		if (
+			((InstanceData.Instance->Guides.bIsSimulationEnable || InstanceData.Instance->Guides.bIsDeformationEnable || InstanceData.Instance->Guides.bIsSimulationCacheEnable)) ||
+			(!InstanceData.Instance->Guides.bHasGlobalInterpolation && !InstanceData.Instance->Guides.bIsSimulationEnable && !InstanceData.Instance->Guides.bIsDeformationEnable && !InstanceData.Instance->Guides.bIsSimulationCacheEnable) ||
+			!IsHairStrandsBindingEnable()) return;
+
+		FRDGExternalBuffer RawDeformedPositionBuffer = InstanceData.Instance->Guides.DeformedResource->GetBuffer(FHairStrandsDeformedResource::Current);
+		FRDGImportedBuffer DeformedPositionBuffer = Register(GraphBuilder, RawDeformedPositionBuffer, ERDGImportedBufferFlags::CreateUAV);
+
+		AddDeformSimHairStrandsPass(
+			GraphBuilder,
+			ShaderMap,
+			InstanceData.Instance->Debug.MeshLODIndex,
+			InstanceData.Instance->Guides.RestResource->GetPointCount(),
+			InstanceData.Instance->Guides.RestRootResource,
+			InstanceData.Instance->Guides.DeformedRootResource,
+			RegisterAsSRV(GraphBuilder, InstanceData.Instance->Guides.RestResource->PositionBuffer),
+			RegisterAsSRV(GraphBuilder, InstanceData.Instance->Guides.RestResource->PointToCurveBuffer),
+			DeformedPositionBuffer,
+			InstanceData.Instance->Guides.RestResource->GetPositionOffset(),
+			RegisterAsSRV(GraphBuilder, InstanceData.Instance->Guides.DeformedResource->GetPositionOffsetBuffer(FHairStrandsDeformedResource::Current)),
+			InstanceData.Instance->Guides.bHasGlobalInterpolation,
+			nullptr);
+	}	
 }
 
 void AddDrawDebugClusterPass(
diff --git a/Engine/Plugins/Runtime/HairStrands/Source/HairStrandsCore/Private/HairStrandsInterpolation.cpp b/Engine/Plugins/Runtime/HairStrands/Source/HairStrandsCore/Private/HairStrandsInterpolation.cpp
--- a/Engine/Plugins/Runtime/HairStrands/Source/HairStrandsCore/Private/HairStrandsInterpolation.cpp
+++ b/Engine/Plugins/Runtime/HairStrands/Source/HairStrandsCore/Private/HairStrandsInterpolation.cpp
@@ -333,7 +333,7 @@ public:
 
 IMPLEMENT_GLOBAL_SHADER(FDeformGuideCS, "/Engine/Private/HairStrands/HairStrandsGuideDeform.usf", "MainCS", SF_Compute);
 
-static void AddDeformSimHairStrandsPass(
+void AddDeformSimHairStrandsPass(
 	FRDGBuilder& GraphBuilder,
 	FGlobalShaderMap* ShaderMap,
 	const uint32 MeshLODIndex,
@@ -342,7 +342,7 @@ static void AddDeformSimHairStrandsPass(
 	FHairStrandsDeformedRootResource* SimDeformedRootResources,
 	FRDGBufferSRVRef SimRestPosePositionBuffer,
 	FRDGBufferSRVRef SimPointToCurveBuffer,
-	FRDGImportedBuffer OutSimDeformedPositionBuffer,
+	FRDGImportedBuffer& OutSimDeformedPositionBuffer,
 	const FVector& SimRestOffset,
 	FRDGBufferSRVRef SimDeformedOffsetBuffer,
 	const bool bHasGlobalInterpolation,
@@ -1781,6 +1781,7 @@ void ComputeHairStrandsInterpolation(
 				Strands_PositionOffsetSRV = RegisterAsSRV(GraphBuilder, Instance->Strands.DeformedResource->GetPositionOffsetBuffer(FHairStrandsDeformedResource::EFrameType::Current));
 				Strands_TangentSRV = Strands_DeformedTangent.SRV;
 
+				// Move this in to the interpolation update!!!
 				if (ActiveGroomCacheType == EGroomCacheType::Guides)
 				{
 					FScopeLock Lock(Instance->Debug.GroomCacheBuffers->GetCriticalSection()); // This is not ideally it will block the rendering thread / game thread
@@ -1838,6 +1839,7 @@ void ComputeHairStrandsInterpolation(
 								FRHIShaderResourceView* BoneBufferSRV = FSkeletalMeshDeformerHelpers::GetBoneBufferForReading(SkeletalMeshObject, LodIndex, Instance->DeformedSection, false);
 
 								// Guides deformation based on the skeletal mesh bones
+								FRDGImportedBuffer GuideDeformResource = Register(GraphBuilder, Instance->Guides.DeformedResource->GetBuffer(FHairStrandsDeformedResource::Current), ERDGImportedBufferFlags::CreateUAV);
 								AddDeformSimHairStrandsPass(
 									GraphBuilder,
 									ShaderMap,
@@ -1847,7 +1849,7 @@ void ComputeHairStrandsInterpolation(
 									Instance->Guides.DeformedRootResource,
 									RegisterAsSRV(GraphBuilder, Instance->Guides.RestResource->PositionBuffer),
 									RegisterAsSRV(GraphBuilder, Instance->Guides.RestResource->PointToCurveBuffer),
-									Register(GraphBuilder, Instance->Guides.DeformedResource->GetBuffer(FHairStrandsDeformedResource::Current), ERDGImportedBufferFlags::CreateUAV),
+									GuideDeformResource,
 									Instance->Guides.RestResource->GetPositionOffset(),
 									RegisterAsSRV(GraphBuilder, Instance->Guides.DeformedResource->GetPositionOffsetBuffer(FHairStrandsDeformedResource::Current)),
 									Instance->Guides.bHasGlobalInterpolation,
@@ -2376,37 +2378,3 @@ void ComputeHairStrandsInterpolation(
 
 	ExternalAccessQueue.Submit(GraphBuilder);
 }
-
-void ResetHairStrandsInterpolation(
-	FRDGBuilder& GraphBuilder,
-	FGlobalShaderMap* ShaderMap,
-	FHairGroupInstance* Instance,
-	int32 MeshLODIndex)
-{
-	if (!Instance || 
-		(Instance && (Instance->Guides.bIsSimulationEnable || Instance->Guides.bIsDeformationEnable || Instance->Guides.bIsSimulationCacheEnable)) ||
-		(Instance && !Instance->Guides.bHasGlobalInterpolation && !Instance->Guides.bIsSimulationEnable && !Instance->Guides.bIsDeformationEnable && !Instance->Guides.bIsSimulationCacheEnable) ||
-		!IsHairStrandsBindingEnable()) return;
-
-	DECLARE_GPU_STAT(HairStrandsGuideDeform);
-	RDG_EVENT_SCOPE(GraphBuilder, "HairStrandsGuideDeform");
-	RDG_GPU_STAT_SCOPE(GraphBuilder, HairStrandsGuideDeform);
-
-	FRDGExternalBuffer RawDeformedPositionBuffer = Instance->Guides.DeformedResource->GetBuffer(FHairStrandsDeformedResource::Current);
-	FRDGImportedBuffer DeformedPositionBuffer = Register(GraphBuilder, RawDeformedPositionBuffer, ERDGImportedBufferFlags::CreateUAV);
-
-	AddDeformSimHairStrandsPass(
-		GraphBuilder,
-		ShaderMap,
-		MeshLODIndex,
-		Instance->Guides.RestResource->GetPointCount(),
-		Instance->Guides.RestRootResource,
-		Instance->Guides.DeformedRootResource,
-		RegisterAsSRV(GraphBuilder, Instance->Guides.RestResource->PositionBuffer),
-		RegisterAsSRV(GraphBuilder, Instance->Guides.RestResource->PointToCurveBuffer),
-		DeformedPositionBuffer,
-		Instance->Guides.RestResource->GetPositionOffset(),
-		RegisterAsSRV(GraphBuilder, Instance->Guides.DeformedResource->GetPositionOffsetBuffer(FHairStrandsDeformedResource::Current)),
-		Instance->Guides.bHasGlobalInterpolation,
-		nullptr);
-}
diff --git a/Engine/Plugins/Runtime/HairStrands/Source/HairStrandsCore/Public/HairStrandsInterpolation.h b/Engine/Plugins/Runtime/HairStrands/Source/HairStrandsCore/Public/HairStrandsInterpolation.h
--- a/Engine/Plugins/Runtime/HairStrands/Source/HairStrandsCore/Public/HairStrandsInterpolation.h
+++ b/Engine/Plugins/Runtime/HairStrands/Source/HairStrandsCore/Public/HairStrandsInterpolation.h
@@ -5,6 +5,7 @@
 #include "Containers/ContainersFwd.h"
 #include "HAL/Platform.h"
 #include "Math/MathFwd.h"
+#include "RenderGraphFwd.h"
 
 #if UE_ENABLE_INCLUDE_ORDER_DEPRECATED_IN_5_3
 #include "CoreMinimal.h"
@@ -20,11 +21,15 @@ class FGlobalShaderMap;
 class FRDGBuilder;
 class FSkeletalMeshRenderData;
 class UGroomBindingAsset;
+class FRHIShaderResourceView;
 
 struct FHairGroupInstance;
 struct FHairStrandClusterData;
 struct FRWBuffer;
 struct FShaderPrintData;
+struct FHairStrandsRestRootResource;
+struct FHairStrandsDeformedRootResource;
+struct FRDGImportedBuffer;
 
 enum class EGroomViewMode : uint8;
 
@@ -33,11 +38,20 @@ enum class EGroomViewMode : uint8;
 //  1) when hair simulation is enabled, the first frame is correct
 //  2) when hair simulation is enabled/disabled (i.e., toggle/change) 
 //     we reset to deform buffer to rest state)
-void ResetHairStrandsInterpolation(
+void AddDeformSimHairStrandsPass(
 	FRDGBuilder& GraphBuilder,
 	FGlobalShaderMap* ShaderMap,
-	FHairGroupInstance* Instance,
-	int32 LODIndex);
+	const uint32 MeshLODIndex,
+	const uint32 VertexCount,
+	FHairStrandsRestRootResource* SimRestRootResources,
+	FHairStrandsDeformedRootResource* SimDeformedRootResources,
+	FRDGBufferSRVRef SimRestPosePositionBuffer,
+	FRDGBufferSRVRef SimPointToCurveBuffer,
+	FRDGImportedBuffer& OutSimDeformedPositionBuffer,
+	const FVector& SimRestOffset,
+	FRDGBufferSRVRef SimDeformedOffsetBuffer,
+	const bool bHasGlobalInterpolation,
+	FRHIShaderResourceView* BoneBufferSRV);
 
 void ComputeHairStrandsInterpolation(
 	FRDGBuilder& GraphBuilder,
