Move guide cache update into the guide update function.

diff --git a/Engine/Plugins/Runtime/HairStrands/Source/HairStrandsCore/Private/GroomManager.cpp b/Engine/Plugins/Runtime/HairStrands/Source/HairStrandsCore/Private/GroomManager.cpp
--- a/Engine/Plugins/Runtime/HairStrands/Source/HairStrandsCore/Private/GroomManager.cpp
+++ b/Engine/Plugins/Runtime/HairStrands/Source/HairStrandsCore/Private/GroomManager.cpp
@@ -28,6 +28,7 @@
 #include "PrimitiveSceneInfo.h"
 #include "HairStrandsClusterCulling.h"
 #include "GroomVisualizationData.h"
+#include "SkeletalMeshDeformerHelpers.h"
 
 static int32 GHairStrandsMinLOD = 0;
 static FAutoConsoleVariableRef CVarGHairStrandsMinLOD(TEXT("r.HairStrands.MinLOD"), GHairStrandsMinLOD, TEXT("Clamp the min hair LOD to this value, preventing to reach lower/high-quality LOD."), ECVF_Scalability);
@@ -394,6 +395,8 @@ static void RunInternalHairInterpolation(
 	}
 }
 
+EGroomCacheType GetHairInstanceCacheType(const FHairGroupInstance* Instance);
+
 static void RunHairStrandsInterpolation_Guide(
 	FRDGBuilder& GraphBuilder,
 	FSceneInterface* Scene,
@@ -414,6 +417,7 @@ static void RunHairStrandsInterpolation_Guide(
 	{
 		int32 HairLODIndex = -1;
 		int32 MeshLODIndex = -1;
+		EGroomCacheType ActiveGroomCacheType = EGroomCacheType::None;
 		EHairBindingType BindingType = EHairBindingType::NoneBinding;
 		bool bSimulationEnable = false;
 		bool bDeformationEnable = false;
@@ -445,6 +449,7 @@ static void RunHairStrandsInterpolation_Guide(
 
 		FInstanceData& InstanceData = InstanceDatas.AddDefaulted_GetRef();
 		InstanceData.Instance 					= Instance;
+		InstanceData.ActiveGroomCacheType 		= GetHairInstanceCacheType(Instance);
 		InstanceData.BindingType 				= Instance->BindingType;
 		InstanceData.HairLODIndex 				= Instance->HairGroupPublicData->LODIndex;
 		InstanceData.bSimulationEnable 			= Instance->HairGroupPublicData->IsSimulationEnable(InstanceData.HairLODIndex);
@@ -552,15 +557,13 @@ static void RunHairStrandsInterpolation_Guide(
 		}
 	}
 
-	// Apply deformation from
-	// * Bones 
-	// * RBF
+	// Apply deformation from RBF
 	for (FInstanceData& InstanceData : InstanceDatas)
 	{
 		if (
 			((InstanceData.Instance->Guides.bIsSimulationEnable || InstanceData.Instance->Guides.bIsDeformationEnable || InstanceData.Instance->Guides.bIsSimulationCacheEnable)) ||
 			(!InstanceData.Instance->Guides.bHasGlobalInterpolation && !InstanceData.Instance->Guides.bIsSimulationEnable && !InstanceData.Instance->Guides.bIsDeformationEnable && !InstanceData.Instance->Guides.bIsSimulationCacheEnable) ||
-			!IsHairStrandsBindingEnable()) return;
+			!IsHairStrandsBindingEnable()) continue;
 
 		FRDGExternalBuffer RawDeformedPositionBuffer = InstanceData.Instance->Guides.DeformedResource->GetBuffer(FHairStrandsDeformedResource::Current);
 		FRDGImportedBuffer DeformedPositionBuffer = Register(GraphBuilder, RawDeformedPositionBuffer, ERDGImportedBufferFlags::CreateUAV);
@@ -579,7 +582,82 @@ static void RunHairStrandsInterpolation_Guide(
 			RegisterAsSRV(GraphBuilder, InstanceData.Instance->Guides.DeformedResource->GetPositionOffsetBuffer(FHairStrandsDeformedResource::Current)),
 			InstanceData.Instance->Guides.bHasGlobalInterpolation,
 			nullptr);
-	}	
+	}
+
+	// Apply deformation from skeletal mesh bones
+	for (FInstanceData& InstanceData : InstanceDatas)
+	{
+		if (InstanceData.ActiveGroomCacheType != EGroomCacheType::None && InstanceData.Instance->Guides.bIsDeformationEnable && InstanceData.Instance->Guides.DeformedResource && InstanceData.Instance->DeformedComponent && (InstanceData.Instance->DeformedSection != INDEX_NONE))
+		{
+			if (USkeletalMeshComponent* SkeletalMeshComponent = Cast<USkeletalMeshComponent>(InstanceData.Instance->DeformedComponent))
+			{
+				if (FSkeletalMeshObject* SkeletalMeshObject = SkeletalMeshComponent->MeshObject)
+				{
+					const int32 LodIndex = SkeletalMeshObject->GetLOD();
+					FSkeletalMeshLODRenderData const* LodRenderData = &SkeletalMeshObject->GetSkeletalMeshRenderData().LODRenderData[LodIndex];
+				
+					if (LodRenderData->RenderSections.IsValidIndex(InstanceData.Instance->DeformedSection)) 
+					{
+						FRHIShaderResourceView* BoneBufferSRV = FSkeletalMeshDeformerHelpers::GetBoneBufferForReading(SkeletalMeshObject, LodIndex, InstanceData.Instance->DeformedSection, false);
+
+						// Guides deformation based on the skeletal mesh bones
+						FRDGImportedBuffer GuideDeformResource = Register(GraphBuilder, InstanceData.Instance->Guides.DeformedResource->GetBuffer(FHairStrandsDeformedResource::Current), ERDGImportedBufferFlags::CreateUAV);
+						AddDeformSimHairStrandsPass(
+							GraphBuilder,
+							ShaderMap,
+							InstanceData.MeshLODIndex,
+							InstanceData.Instance->Guides.RestResource->GetPointCount(),
+							InstanceData.Instance->Guides.RestRootResource,
+							InstanceData.Instance->Guides.DeformedRootResource,
+							RegisterAsSRV(GraphBuilder, InstanceData.Instance->Guides.RestResource->PositionBuffer),
+							RegisterAsSRV(GraphBuilder, InstanceData.Instance->Guides.RestResource->PointToCurveBuffer),
+							GuideDeformResource,
+							InstanceData.Instance->Guides.RestResource->GetPositionOffset(),
+							RegisterAsSRV(GraphBuilder, InstanceData.Instance->Guides.DeformedResource->GetPositionOffsetBuffer(FHairStrandsDeformedResource::Current)),
+							InstanceData.Instance->Guides.bHasGlobalInterpolation,
+							BoneBufferSRV);
+					}
+				}
+			}
+		}
+	}
+
+	// Apply guide cache update
+	for (FInstanceData& InstanceData : InstanceDatas)
+	{
+		if (InstanceData.ActiveGroomCacheType == EGroomCacheType::Guides)
+		{
+			FScopeLock Lock(InstanceData.Instance->Debug.GroomCacheBuffers->GetCriticalSection());
+			FGroomCacheGroupData* GroomCacheData0 = const_cast<FGroomCacheGroupData*>(&InstanceData.Instance->Debug.GroomCacheBuffers->GetCurrentFrameBuffer().GroupsData[InstanceData.Instance->Debug.GroupIndex]);
+			FGroomCacheGroupData* GroomCacheData1 = const_cast<FGroomCacheGroupData*>(&InstanceData.Instance->Debug.GroomCacheBuffers->GetNextFrameBuffer().GroupsData[InstanceData.Instance->Debug.GroupIndex]);
+
+			const float InterpolationFactor = InstanceData.Instance->Debug.GroomCacheBuffers->GetInterpolationFactor();
+
+			// Update CPU position offset to GPU
+			InstanceData.Instance->Guides.DeformedResource->SetPositionOffset(FHairStrandsDeformedResource::EFrameType::Current, FMath::Lerp(GroomCacheData0->BoundingBox.GetCenter(), GroomCacheData1->BoundingBox.GetCenter(), InterpolationFactor));
+			AddHairStrandUpdatePositionOffsetPass(
+				GraphBuilder,
+				ShaderMap,
+				InstanceData.MeshLODIndex,
+				InstanceData.Instance->Guides.DeformedRootResource,
+				InstanceData.Instance->Guides.DeformedResource);
+
+			// Pass to upload GroomCache guide positions
+			FGroomCacheResources CacheResources0 = CreateGroomCacheBuffer(GraphBuilder, GroomCacheData0->VertexData);
+			FGroomCacheResources CacheResources1 = CreateGroomCacheBuffer(GraphBuilder, GroomCacheData1->VertexData);
+			AddGroomCacheUpdatePass(
+				GraphBuilder,
+				ShaderMap,
+				InstanceData.Instance->Guides.RestResource->GetPointCount(),
+				InterpolationFactor,
+				CacheResources0,
+				CacheResources1,
+				RegisterAsSRV(GraphBuilder, InstanceData.Instance->Guides.RestResource->PositionBuffer),
+				RegisterAsSRV(GraphBuilder, InstanceData.Instance->Guides.DeformedResource->GetPositionOffsetBuffer(FHairStrandsDeformedResource::EFrameType::Current)),
+				RegisterAsUAV(GraphBuilder, InstanceData.Instance->Guides.DeformedResource->GetBuffer(FHairStrandsDeformedResource::Current)));
+		}
+	}
+
 }
 
 void AddDrawDebugClusterPass(
diff --git a/Engine/Plugins/Runtime/HairStrands/Source/HairStrandsCore/Private/HairStrandsInterpolation.cpp b/Engine/Plugins/Runtime/HairStrands/Source/HairStrandsCore/Private/HairStrandsInterpolation.cpp
--- a/Engine/Plugins/Runtime/HairStrands/Source/HairStrandsCore/Private/HairStrandsInterpolation.cpp
+++ b/Engine/Plugins/Runtime/HairStrands/Source/HairStrandsCore/Private/HairStrandsInterpolation.cpp
@@ -100,14 +100,7 @@ void GetHairStrandsAttributeParameter(const FHairStrandsBulkData& In, FHairStran
 	PACK_POINT_HAIR_ATTRIBUTE_OFFSETS(Out.PointAttributeOffsets, In.Header.PointAttributeOffsets);
 }
 
-struct FGroomCacheResources
-{
-	FRDGBufferSRVRef PositionBuffer = nullptr;
-	FRDGBufferSRVRef RadiusBuffer = nullptr;
-	bool bHasRadiusData = false;
-};
-
-static FGroomCacheResources CreateGroomCacheBuffer(FRDGBuilder& GraphBuilder, FGroomCacheVertexData& InVertexData)
+FGroomCacheResources CreateGroomCacheBuffer(FRDGBuilder& GraphBuilder, FGroomCacheVertexData& InVertexData)
 {
 	FGroomCacheResources Out;
 	
@@ -251,7 +244,7 @@ public:
 
 IMPLEMENT_GLOBAL_SHADER(FGroomCacheUpdatePassCS, "/Engine/Private/HairStrands/HairStrandsInterpolation.usf", "MainCS", SF_Compute);
 
-static void AddGroomCacheUpdatePass(
+void AddGroomCacheUpdatePass(
 	FRDGBuilder& GraphBuilder,
 	FGlobalShaderMap* ShaderMap,
 	uint32 PointCount,
@@ -1689,39 +1682,6 @@ void ComputeHairStrandsInterpolation(
 			// Disable culling when drawing only guides, as the culling output has been computed for the strands, not for the guides.
 			Instance->HairGroupPublicData->SetCullingResultAvailable(false);
 
-			// If groom guide cache, is enabled, used the deform position for debug visualization
-			if (ActiveGroomCacheType == EGroomCacheType::Guides)
-			{
-				FScopeLock Lock(Instance->Debug.GroomCacheBuffers->GetCriticalSection());
-				FGroomCacheGroupData* GroomCacheData0 = const_cast<FGroomCacheGroupData*>(&Instance->Debug.GroomCacheBuffers->GetCurrentFrameBuffer().GroupsData[Instance->Debug.GroupIndex]);
-				FGroomCacheGroupData* GroomCacheData1 = const_cast<FGroomCacheGroupData*>(&Instance->Debug.GroomCacheBuffers->GetNextFrameBuffer().GroupsData[Instance->Debug.GroupIndex]);
-
-				const float InterpolationFactor = Instance->Debug.GroomCacheBuffers->GetInterpolationFactor();
-				Instance->Guides.DeformedResource->SetPositionOffset(FHairStrandsDeformedResource::EFrameType::Current, FMath::Lerp(GroomCacheData0->BoundingBox.GetCenter(), GroomCacheData1->BoundingBox.GetCenter(), InterpolationFactor));
-
-				FGroomCacheResources CacheResources0 = CreateGroomCacheBuffer(GraphBuilder, GroomCacheData0->VertexData);
-				FGroomCacheResources CacheResources1 = CreateGroomCacheBuffer(GraphBuilder, GroomCacheData1->VertexData);
-
-				AddHairStrandUpdatePositionOffsetPass(
-					GraphBuilder,
-					ShaderMap,
-					MeshLODIndex,
-					Instance->Guides.DeformedRootResource,
-					Instance->Guides.DeformedResource);
-
-				// Pass to upload GroomCache guide positions
-				AddGroomCacheUpdatePass(
-					GraphBuilder,
-					ShaderMap,
-					Instance->Guides.RestResource->GetPointCount(),
-					InterpolationFactor,
-					CacheResources0,
-					CacheResources1,
-					RegisterAsSRV(GraphBuilder, Instance->Guides.RestResource->PositionBuffer),
-					RegisterAsSRV(GraphBuilder, Instance->Guides.DeformedResource->GetPositionOffsetBuffer(FHairStrandsDeformedResource::EFrameType::Current)),
-					RegisterAsUAV(GraphBuilder, Instance->Guides.DeformedResource->GetBuffer(FHairStrandsDeformedResource::Current)));
-			}
-
 			AddHairTangentPass(
 				GraphBuilder,
 				ShaderMap,
@@ -1781,29 +1741,11 @@ void ComputeHairStrandsInterpolation(
 				Strands_PositionOffsetSRV = RegisterAsSRV(GraphBuilder, Instance->Strands.DeformedResource->GetPositionOffsetBuffer(FHairStrandsDeformedResource::EFrameType::Current));
 				Strands_TangentSRV = Strands_DeformedTangent.SRV;
 
-				// Move this in to the interpolation update!!!
+				// 2.0 If using a guide cache, update the render strands offset point
 				if (ActiveGroomCacheType == EGroomCacheType::Guides)
 				{
-					FScopeLock Lock(Instance->Debug.GroomCacheBuffers->GetCriticalSection()); // This is not ideally it will block the rendering thread / game thread
-
-					FGroomCacheGroupData* GroomCacheData0 = const_cast<FGroomCacheGroupData*>(&Instance->Debug.GroomCacheBuffers->GetCurrentFrameBuffer().GroupsData[Instance->Debug.GroupIndex]);
-					FGroomCacheGroupData* GroomCacheData1 = const_cast<FGroomCacheGroupData*>(&Instance->Debug.GroomCacheBuffers->GetNextFrameBuffer().GroupsData[Instance->Debug.GroupIndex]);
-					const float InterpolationFactor = Instance->Debug.GroomCacheBuffers->GetInterpolationFactor();
-					const FVector OffsetPosition = FMath::Lerp(GroomCacheData0->BoundingBox.GetCenter(), GroomCacheData1->BoundingBox.GetCenter(), InterpolationFactor);
-
-					FGroomCacheResources CacheResources0 = CreateGroomCacheBuffer(GraphBuilder, GroomCacheData0->VertexData);
-					FGroomCacheResources CacheResources1 = CreateGroomCacheBuffer(GraphBuilder, GroomCacheData1->VertexData);
-
-					Instance->Guides.DeformedResource->SetPositionOffset(FHairStrandsDeformedResource::EFrameType::Current, OffsetPosition);
-
-					AddHairStrandUpdatePositionOffsetPass(
-						GraphBuilder,
-						ShaderMap,
-						MeshLODIndex,
-						Instance->Guides.DeformedRootResource,
-						Instance->Guides.DeformedResource);
-
 					// Apply the same offset to the render strands for proper rendering
+					const FVector OffsetPosition = Instance->Guides.DeformedResource->GetPositionOffset(FHairStrandsDeformedResource::EFrameType::Current);
 					Instance->Strands.DeformedResource->SetPositionOffset(FHairStrandsDeformedResource::EFrameType::Current, OffsetPosition);
 
 					AddHairStrandUpdatePositionOffsetPass(
@@ -1812,51 +1754,6 @@ void ComputeHairStrandsInterpolation(
 						MeshLODIndex,
 						Instance->Strands.DeformedRootResource,
 						Instance->Strands.DeformedResource);
-
-					// Pass to upload GroomCache guide positions
-					AddGroomCacheUpdatePass(
-						GraphBuilder,
-						ShaderMap,
-						Instance->Guides.RestResource->GetPointCount(),
-						InterpolationFactor,
-						CacheResources0,
-						CacheResources1,
-						RegisterAsSRV(GraphBuilder, Instance->Guides.RestResource->PositionBuffer),
-						RegisterAsSRV(GraphBuilder, Instance->Guides.DeformedResource->GetPositionOffsetBuffer(FHairStrandsDeformedResource::EFrameType::Current)),
-						RegisterAsUAV(GraphBuilder, Instance->Guides.DeformedResource->GetBuffer(FHairStrandsDeformedResource::Current)));
-				}
-				else if (Instance->Guides.bIsDeformationEnable && Instance->Guides.DeformedResource && Instance->DeformedComponent && (Instance->DeformedSection != INDEX_NONE))
-				{
-					if (USkeletalMeshComponent* SkeletalMeshComponent = Cast<USkeletalMeshComponent>(Instance->DeformedComponent))
-					{
-						if(FSkeletalMeshObject* SkeletalMeshObject = SkeletalMeshComponent->MeshObject)
-						{
-							const int32 LodIndex = SkeletalMeshObject->GetLOD();
-							FSkeletalMeshLODRenderData const* LodRenderData = &SkeletalMeshObject->GetSkeletalMeshRenderData().LODRenderData[LodIndex];
-				
-							if(LodRenderData->RenderSections.Num() > Instance->DeformedSection) 
-							{
-								FRHIShaderResourceView* BoneBufferSRV = FSkeletalMeshDeformerHelpers::GetBoneBufferForReading(SkeletalMeshObject, LodIndex, Instance->DeformedSection, false);
-
-								// Guides deformation based on the skeletal mesh bones
-								FRDGImportedBuffer GuideDeformResource = Register(GraphBuilder, Instance->Guides.DeformedResource->GetBuffer(FHairStrandsDeformedResource::Current), ERDGImportedBufferFlags::CreateUAV);
-								AddDeformSimHairStrandsPass(
-									GraphBuilder,
-									ShaderMap,
-									MeshLODIndex,
-									Instance->Guides.RestResource->GetPointCount(),
-									Instance->Guides.RestRootResource,
-									Instance->Guides.DeformedRootResource,
-									RegisterAsSRV(GraphBuilder, Instance->Guides.RestResource->PositionBuffer),
-									RegisterAsSRV(GraphBuilder, Instance->Guides.RestResource->PointToCurveBuffer),
-									GuideDeformResource,
-									Instance->Guides.RestResource->GetPositionOffset(),
-									RegisterAsSRV(GraphBuilder, Instance->Guides.DeformedResource->GetPositionOffsetBuffer(FHairStrandsDeformedResource::Current)),
-									Instance->Guides.bHasGlobalInterpolation,
-									BoneBufferSRV);
-							}
-						}
-					}
 				}
 
 				// 2.1 Compute deformation position based on simulation/skinning/RBF
diff --git a/Engine/Plugins/Runtime/HairStrands/Source/HairStrandsCore/Public/HairStrandsInterpolation.h b/Engine/Plugins/Runtime/HairStrands/Source/HairStrandsCore/Public/HairStrandsInterpolation.h
--- a/Engine/Plugins/Runtime/HairStrands/Source/HairStrandsCore/Public/HairStrandsInterpolation.h
+++ b/Engine/Plugins/Runtime/HairStrands/Source/HairStrandsCore/Public/HairStrandsInterpolation.h
@@ -30,6 +30,7 @@ struct FShaderPrintData;
 struct FHairStrandsRestRootResource;
 struct FHairStrandsDeformedRootResource;
 struct FRDGImportedBuffer;
+struct FGroomCacheVertexData;
 
 enum class EGroomViewMode : uint8;
 
@@ -65,4 +66,23 @@ void ComputeHairStrandsInterpolation(
 	int32 LODIndex,
 	FHairStrandClusterData* ClusterData);
 
+struct FGroomCacheResources
+{
+	FRDGBufferSRVRef PositionBuffer = nullptr;
+	FRDGBufferSRVRef RadiusBuffer = nullptr;
+	bool bHasRadiusData = false;
+};
+FGroomCacheResources CreateGroomCacheBuffer(FRDGBuilder& GraphBuilder, FGroomCacheVertexData& InVertexData);
+
+void AddGroomCacheUpdatePass(
+	FRDGBuilder& GraphBuilder,
+	FGlobalShaderMap* ShaderMap,
+	uint32 PointCount,
+	float InterpolationFactor,
+	FGroomCacheResources CacheResources0,
+	FGroomCacheResources CacheResources1,
+	FRDGBufferSRVRef InBuffer,
+	FRDGBufferSRVRef InDeformedOffsetBuffer,
+	FRDGBufferUAVRef OutBuffer);
+
 HAIRSTRANDSCORE_API void ComputeInterpolationWeights(UGroomBindingAsset* BindingAsset, FSkeletalMeshRenderData* TargetRenderData, TArray<FRWBuffer>& TransferedPositions);
