Change RBF buffer to structured buffer to scalarize loads.

Reduce RBF deformation cost by 50%. 0.09ms -> 0.04ms (per groom instance)

diff --git a/Engine/Plugins/Runtime/HairStrands/Shaders/Private/NiagaraDataInterfaceHairStrandsTemplate.ush b/Engine/Plugins/Runtime/HairStrands/Shaders/Private/NiagaraDataInterfaceHairStrandsTemplate.ush
--- a/Engine/Plugins/Runtime/HairStrands/Shaders/Private/NiagaraDataInterfaceHairStrandsTemplate.ush
+++ b/Engine/Plugins/Runtime/HairStrands/Shaders/Private/NiagaraDataInterfaceHairStrandsTemplate.ush
@@ -31,8 +31,8 @@ uint				{ParameterName}_RestUpdate;
 uint				{ParameterName}_LocalSimulation;
 int					{ParameterName}_SampleCount;
 int4				{ParameterName}_BoundingBoxOffsets;
-Buffer<float4>		{ParameterName}_RestSamplePositionsBuffer;
-Buffer<float4>		{ParameterName}_MeshSampleWeightsBuffer;
+StructuredBuffer<float4> {ParameterName}_RestSamplePositionsBuffer;
+StructuredBuffer<float4> {ParameterName}_MeshSampleWeightsBuffer;
 Buffer<float>		{ParameterName}_ParamsScaleBuffer;
 float3				{ParameterName}_BoneLinearVelocity;
 float3				{ParameterName}_BoneAngularVelocity;
diff --git a/Engine/Plugins/Runtime/HairStrands/Source/HairStrandsCore/Private/GroomDebug.cpp b/Engine/Plugins/Runtime/HairStrands/Source/HairStrandsCore/Private/GroomDebug.cpp
--- a/Engine/Plugins/Runtime/HairStrands/Source/HairStrandsCore/Private/GroomDebug.cpp
+++ b/Engine/Plugins/Runtime/HairStrands/Source/HairStrandsCore/Private/GroomDebug.cpp
@@ -258,8 +258,8 @@ class FHairProjectionHairDebugCS : public FGlobalShader
 
 		SHADER_PARAMETER_RDG_BUFFER_SRV(Buffer, RestPositionBuffer)
 		SHADER_PARAMETER_RDG_BUFFER_SRV(Buffer, DeformedPositionBuffer)
-		SHADER_PARAMETER_RDG_BUFFER_SRV(Buffer, RestSamplePositionsBuffer)
-		SHADER_PARAMETER_RDG_BUFFER_SRV(Buffer, DeformedSamplePositionsBuffer)
+		SHADER_PARAMETER_RDG_BUFFER_SRV(StructuredBuffer, RestSamplePositionsBuffer)
+		SHADER_PARAMETER_RDG_BUFFER_SRV(StructuredBuffer, DeformedSamplePositionsBuffer)
 		SHADER_PARAMETER_RDG_BUFFER_SRV(Buffer, RootBarycentricBuffer)
 		SHADER_PARAMETER_RDG_BUFFER_SRV(Buffer, RootToUniqueTriangleIndexBuffer)
 
diff --git a/Engine/Plugins/Runtime/HairStrands/Source/HairStrandsCore/Private/GroomResources.cpp b/Engine/Plugins/Runtime/HairStrands/Source/HairStrandsCore/Private/GroomResources.cpp
--- a/Engine/Plugins/Runtime/HairStrands/Source/HairStrandsCore/Private/GroomResources.cpp
+++ b/Engine/Plugins/Runtime/HairStrands/Source/HairStrandsCore/Private/GroomResources.cpp
@@ -1419,7 +1419,7 @@ void FHairStrandsRestRootResource::InternalAllocate(FRDGBuilder& GraphBuilder, u
 
 			InternalCreateVertexBufferRDG_FromHairBulkData<FHairStrandsWeightFormat>(GraphBuilder, CPUData.MeshInterpolationWeightsBuffer, InteroplationWeightCount, GPUData.MeshInterpolationWeightsBuffer, ToHairResourceDebugName(HAIRSTRANDS_RESOUCE_NAME(CurveType, Hair.StrandsRestRoot_MeshInterpolationWeightsBuffer), ResourceName), OwnerName, EHairResourceUsageType::Static);
 			InternalCreateVertexBufferRDG_FromHairBulkData<FHairStrandsIndexFormat>(GraphBuilder, CPUData.MeshSampleIndicesBuffer, LODHeader.SampleCount, GPUData.MeshSampleIndicesBuffer, ToHairResourceDebugName(HAIRSTRANDS_RESOUCE_NAME(CurveType, Hair.StrandsRestRoot_MeshSampleIndicesBuffer), ResourceName), OwnerName, EHairResourceUsageType::Static);
-			InternalCreateVertexBufferRDG_FromHairBulkData<FHairStrandsMeshTrianglePositionFormat>(GraphBuilder, CPUData.RestSamplePositionsBuffer, LODHeader.SampleCount, GPUData.RestSamplePositionsBuffer, ToHairResourceDebugName(HAIRSTRANDS_RESOUCE_NAME(CurveType, Hair.StrandsRestRoot_RestSamplePositionsBuffer), ResourceName), OwnerName, EHairResourceUsageType::Static);
+			InternalCreateStructuredBufferRDG_FromHairBulkData<FHairStrandsMeshTrianglePositionFormat>(GraphBuilder, CPUData.RestSamplePositionsBuffer, LODHeader.SampleCount, GPUData.RestSamplePositionsBuffer, ToHairResourceDebugName(HAIRSTRANDS_RESOUCE_NAME(CurveType, Hair.StrandsRestRoot_RestSamplePositionsBuffer), ResourceName), OwnerName, EHairResourceUsageType::Static);
 		}
 
 		GPUData.Status = FLOD::EStatus::Completed;
diff --git a/Engine/Plugins/Runtime/HairStrands/Source/HairStrandsCore/Private/HairStrandsInterpolation.cpp b/Engine/Plugins/Runtime/HairStrands/Source/HairStrandsCore/Private/HairStrandsInterpolation.cpp
--- a/Engine/Plugins/Runtime/HairStrands/Source/HairStrandsCore/Private/HairStrandsInterpolation.cpp
+++ b/Engine/Plugins/Runtime/HairStrands/Source/HairStrandsCore/Private/HairStrandsInterpolation.cpp
@@ -307,8 +307,8 @@ class FDeformGuideCS : public FGlobalShader
 		SHADER_PARAMETER_RDG_BUFFER_UAV(RWBuffer, OutSimDeformedPositionBuffer)
 
 		SHADER_PARAMETER(uint32, SampleCount)
-		SHADER_PARAMETER_RDG_BUFFER_SRV(Buffer, RestSamplePositionsBuffer)
-		SHADER_PARAMETER_RDG_BUFFER_SRV(Buffer, MeshSampleWeightsBuffer)
+		SHADER_PARAMETER_RDG_BUFFER_SRV(StructuredBuffer, RestSamplePositionsBuffer)
+		SHADER_PARAMETER_RDG_BUFFER_SRV(StructuredBuffer, MeshSampleWeightsBuffer)
 
 		SHADER_PARAMETER_SRV(Buffer<float4>, BoneDeformedPositionBuffer)
 	END_SHADER_PARAMETER_STRUCT()
diff --git a/Engine/Plugins/Runtime/HairStrands/Source/HairStrandsCore/Private/HairStrandsMeshProjection.cpp b/Engine/Plugins/Runtime/HairStrands/Source/HairStrandsCore/Private/HairStrandsMeshProjection.cpp
--- a/Engine/Plugins/Runtime/HairStrands/Source/HairStrandsCore/Private/HairStrandsMeshProjection.cpp
+++ b/Engine/Plugins/Runtime/HairStrands/Source/HairStrandsCore/Private/HairStrandsMeshProjection.cpp
@@ -535,8 +535,8 @@ private:
 		SHADER_PARAMETER(uint32, VertexCount)
 		SHADER_PARAMETER(uint32, MaxSampleCount)
 
-		SHADER_PARAMETER_RDG_BUFFER_SRV(Buffer, RestSamplePositionsBuffer)
-		SHADER_PARAMETER_RDG_BUFFER_SRV(Buffer, MeshSampleWeightsBuffer)
+		SHADER_PARAMETER_RDG_BUFFER_SRV(StructuredBuffer, RestSamplePositionsBuffer)
+		SHADER_PARAMETER_RDG_BUFFER_SRV(StructuredBuffer, MeshSampleWeightsBuffer)
 
 		SHADER_PARAMETER_SRV(Buffer, RestPositionBuffer)
 		SHADER_PARAMETER_RDG_BUFFER_UAV(RWBuffer, OutDeformedPositionBuffer)
diff --git a/Engine/Plugins/Runtime/HairStrands/Source/HairStrandsCore/Private/Niagara/NiagaraDataInterfaceHairStrands.cpp b/Engine/Plugins/Runtime/HairStrands/Source/HairStrandsCore/Private/Niagara/NiagaraDataInterfaceHairStrands.cpp
--- a/Engine/Plugins/Runtime/HairStrands/Source/HairStrandsCore/Private/Niagara/NiagaraDataInterfaceHairStrands.cpp
+++ b/Engine/Plugins/Runtime/HairStrands/Source/HairStrandsCore/Private/Niagara/NiagaraDataInterfaceHairStrands.cpp
@@ -57,8 +57,8 @@ BEGIN_SHADER_PARAMETER_STRUCT(FShaderParameters,)
 	SHADER_PARAMETER(uint32,							LocalSimulation)
 	SHADER_PARAMETER(int,								SampleCount)
 	SHADER_PARAMETER(FIntVector4,						BoundingBoxOffsets)
-	SHADER_PARAMETER_RDG_BUFFER_SRV(Buffer<float4>,		RestSamplePositionsBuffer)
-	SHADER_PARAMETER_RDG_BUFFER_SRV(Buffer<float4>,		MeshSampleWeightsBuffer)
+	SHADER_PARAMETER_RDG_BUFFER_SRV(StructuredBuffer<float4>, RestSamplePositionsBuffer)
+	SHADER_PARAMETER_RDG_BUFFER_SRV(StructuredBuffer<float4>, MeshSampleWeightsBuffer)
 	SHADER_PARAMETER_RDG_BUFFER_SRV(Buffer<float>,		ParamsScaleBuffer)
 	SHADER_PARAMETER(FVector3f,							BoneLinearVelocity)
 	SHADER_PARAMETER(FVector3f,							BoneAngularVelocity)
diff --git a/Engine/Shaders/Private/HairStrands/HairStrandsBindingCommon.ush b/Engine/Shaders/Private/HairStrands/HairStrandsBindingCommon.ush
--- a/Engine/Shaders/Private/HairStrands/HairStrandsBindingCommon.ush
+++ b/Engine/Shaders/Private/HairStrands/HairStrandsBindingCommon.ush
@@ -176,8 +176,8 @@ float3 TransformPoint(float3 P, float3 B, const FHairMeshTriangle RestTriangle,
 float3 ApplyRBF(
 	float3 RestControlPoint,
 	int InMaxSampleCount,
-	Buffer<float4> InRestSamplePositionsBuffer,
-	Buffer<float4> InMeshSampleWeightsBuffer)
+	StructuredBuffer<float4> InRestSamplePositionsBuffer,
+	StructuredBuffer<float4> InMeshSampleWeightsBuffer)
 {
 	float3 ControlPoint = RestControlPoint;
 	// Apply rbf interpolation from the samples set
@@ -192,4 +192,4 @@ float3 ApplyRBF(
 	ControlPoint += InMeshSampleWeightsBuffer[InMaxSampleCount + 2].xyz * RestControlPoint.y;
 	ControlPoint += InMeshSampleWeightsBuffer[InMaxSampleCount + 3].xyz * RestControlPoint.z;
 	return ControlPoint;
-}
\ No newline at end of file
+}
diff --git a/Engine/Shaders/Private/HairStrands/HairStrandsDebug.usf b/Engine/Shaders/Private/HairStrands/HairStrandsDebug.usf
--- a/Engine/Shaders/Private/HairStrands/HairStrandsDebug.usf
+++ b/Engine/Shaders/Private/HairStrands/HairStrandsDebug.usf
@@ -325,8 +325,8 @@ Buffer<uint>  RootToUniqueTriangleIndexBuffer;
 Buffer<float4>  RestPositionBuffer;
 Buffer<float4>  DeformedPositionBuffer;
 
-Buffer<float4> RestSamplePositionsBuffer;
-Buffer<float4> DeformedSamplePositionsBuffer;
+StructuredBuffer<float4> RestSamplePositionsBuffer;
+StructuredBuffer<float4> DeformedSamplePositionsBuffer;
 
 [numthreads(256, 1, 1)]
 void MainCS(uint3 DispatchThreadId : SV_DispatchThreadID)
diff --git a/Engine/Shaders/Private/HairStrands/HairStrandsGuideDeform.usf b/Engine/Shaders/Private/HairStrands/HairStrandsGuideDeform.usf
--- a/Engine/Shaders/Private/HairStrands/HairStrandsGuideDeform.usf
+++ b/Engine/Shaders/Private/HairStrands/HairStrandsGuideDeform.usf
@@ -54,8 +54,8 @@ Buffer<float4>  BoneDeformedPositionBuffer;
 
 #if PERMUTATION_DEFORMATION == DEFORMATION_OFFSET_GUIDE_GLOBAL
 uint SampleCount;
-Buffer<float4>  RestSamplePositionsBuffer;
-Buffer<float4>  MeshSampleWeightsBuffer;
+StructuredBuffer<float4>  RestSamplePositionsBuffer;
+StructuredBuffer<float4>  MeshSampleWeightsBuffer;
 #endif
 
 float3 DisplacePosition(uint VertexIndex, float3 Pos)
diff --git a/Engine/Shaders/Private/HairStrands/HairStrandsMesh.usf b/Engine/Shaders/Private/HairStrands/HairStrandsMesh.usf
--- a/Engine/Shaders/Private/HairStrands/HairStrandsMesh.usf
+++ b/Engine/Shaders/Private/HairStrands/HairStrandsMesh.usf
@@ -11,8 +11,8 @@
 uint VertexCount;
 uint MaxSampleCount;
 
-Buffer<float4> RestSamplePositionsBuffer;
-Buffer<float4> MeshSampleWeightsBuffer;
+StructuredBuffer<float4> RestSamplePositionsBuffer;
+StructuredBuffer<float4> MeshSampleWeightsBuffer;
 
 Buffer<float4> RestPositionBuffer;
 RWBuffer<float4> OutDeformedPositionBuffer;
