ShaderPrint platform changes
- When compiling for platforms that use HLSLcc, force them to use DXC instead.
- If we can't use DXC and are forced to use FXC, just say that ShaderPrint can't be supported.

diff --git a/Engine/Plugins/Runtime/HairStrands/Source/HairStrandsCore/Private/GroomDebug.cpp b/Engine/Plugins/Runtime/HairStrands/Source/HairStrandsCore/Private/GroomDebug.cpp
--- a/Engine/Plugins/Runtime/HairStrands/Source/HairStrandsCore/Private/GroomDebug.cpp
+++ b/Engine/Plugins/Runtime/HairStrands/Source/HairStrandsCore/Private/GroomDebug.cpp
@@ -168,8 +168,9 @@ class FHairProjectionMeshDebugCS : public FGlobalShader
 	static void ModifyCompilationEnvironment(const FGlobalShaderPermutationParameters& Parameters, FShaderCompilerEnvironment& OutEnvironment)
 	{
 		FGlobalShader::ModifyCompilationEnvironment(Parameters, OutEnvironment);
+		ShaderPrint::ModifyCompilationEnvironment(Parameters, OutEnvironment);
+
 		OutEnvironment.SetDefine(TEXT("SHADER_MESH_PROJECTION_SKIN_CACHE"), 1);
-		OutEnvironment.CompilerFlags.Add(CFLAG_ForceDXC);
 	}
 };
 
@@ -275,8 +276,9 @@ class FHairProjectionHairDebugCS : public FGlobalShader
 	static void ModifyCompilationEnvironment(const FGlobalShaderPermutationParameters& Parameters, FShaderCompilerEnvironment& OutEnvironment)
 	{
 		FGlobalShader::ModifyCompilationEnvironment(Parameters, OutEnvironment);
+		ShaderPrint::ModifyCompilationEnvironment(Parameters, OutEnvironment);
+
 		OutEnvironment.SetDefine(TEXT("SHADER_MESH_PROJECTION_HAIR"), 1);
-		OutEnvironment.CompilerFlags.Add(CFLAG_ForceDXC);
 	}
 };
 
@@ -497,8 +499,9 @@ public:
 	static void ModifyCompilationEnvironment(const FGlobalShaderPermutationParameters& Parameters, FShaderCompilerEnvironment& OutEnvironment)
 	{
 		FGlobalShader::ModifyCompilationEnvironment(Parameters, OutEnvironment);
+		ShaderPrint::ModifyCompilationEnvironment(Parameters, OutEnvironment);
+
 		OutEnvironment.SetDefine(TEXT("SHADER_CARDS_ATLAS"), 1);
-		OutEnvironment.CompilerFlags.Add(CFLAG_ForceDXC);
 	}
 };
 
@@ -662,8 +665,9 @@ public:
 	static void ModifyCompilationEnvironment(const FGlobalShaderPermutationParameters& Parameters, FShaderCompilerEnvironment& OutEnvironment)
 	{
 		FGlobalShader::ModifyCompilationEnvironment(Parameters, OutEnvironment);
+		ShaderPrint::ModifyCompilationEnvironment(Parameters, OutEnvironment);
+
 		OutEnvironment.SetDefine(TEXT("SHADER_CARDS_GUIDE"), 1);
-		OutEnvironment.CompilerFlags.Add(CFLAG_ForceDXC);
 	}
 };
 
@@ -851,10 +855,10 @@ public:
 	{
 		// Skip optimization for avoiding long compilation time due to large UAV writes
 		FGlobalShader::ModifyCompilationEnvironment(Parameters, OutEnvironment);
+		ShaderPrint::ModifyCompilationEnvironment(Parameters, OutEnvironment);
+
 		OutEnvironment.CompilerFlags.Add(CFLAG_Debug);
-		OutEnvironment.CompilerFlags.Add(CFLAG_ForceDXC);
 		OutEnvironment.SetDefine(TEXT("SHADER_PRINT_INSTANCE"), 1);
-		OutEnvironment.CompilerFlags.Add(CFLAG_ForceDXC);
 	}
 };
 
@@ -1126,10 +1130,10 @@ public:
 	{
 		// Skip optimization for avoiding long compilation time due to large UAV writes
 		FGlobalShader::ModifyCompilationEnvironment(Parameters, OutEnvironment);
+		ShaderPrint::ModifyCompilationEnvironment(Parameters, OutEnvironment);
+
 		OutEnvironment.CompilerFlags.Add(CFLAG_Debug);
-		OutEnvironment.CompilerFlags.Add(CFLAG_ForceDXC);
 		OutEnvironment.SetDefine(TEXT("SHADER_PRINT_MEMORY"), 1);
-		OutEnvironment.CompilerFlags.Add(CFLAG_ForceDXC);
 	}
 };
 
diff --git a/Engine/Source/Runtime/Renderer/Private/GlobalDistanceField.cpp b/Engine/Source/Runtime/Renderer/Private/GlobalDistanceField.cpp
--- a/Engine/Source/Runtime/Renderer/Private/GlobalDistanceField.cpp
+++ b/Engine/Source/Runtime/Renderer/Private/GlobalDistanceField.cpp
@@ -1961,8 +1961,9 @@ class FGlobalDistanceFieldDebugCS : public FGlobalShader
 	static void ModifyCompilationEnvironment(const FGlobalShaderPermutationParameters& Parameters, FShaderCompilerEnvironment& OutEnvironment)
 	{
 		FGlobalShader::ModifyCompilationEnvironment(Parameters, OutEnvironment);
+		ShaderPrint::ModifyCompilationEnvironment(Parameters, OutEnvironment);
+
 		OutEnvironment.SetDefine(TEXT("THREADGROUP_SIZE"), GetGroupSize());
-		OutEnvironment.CompilerFlags.Add(CFLAG_ForceDXC);
 	}
 };
 
diff --git a/Engine/Source/Runtime/Renderer/Private/HairStrands/HairStrandsDebug.cpp b/Engine/Source/Runtime/Renderer/Private/HairStrands/HairStrandsDebug.cpp
--- a/Engine/Source/Runtime/Renderer/Private/HairStrands/HairStrandsDebug.cpp
+++ b/Engine/Source/Runtime/Renderer/Private/HairStrands/HairStrandsDebug.cpp
@@ -125,9 +125,10 @@ public:
 	{
 		// Skip optimization for avoiding long compilation time due to large UAV writes
 		FGlobalShader::ModifyCompilationEnvironment(Parameters, OutEnvironment);
+		ShaderPrint::ModifyCompilationEnvironment(Parameters, OutEnvironment);
+
 		OutEnvironment.CompilerFlags.Add(CFLAG_Debug);
 		OutEnvironment.SetDefine(TEXT("SHADER_LOD_INFO"), 1);
-		OutEnvironment.CompilerFlags.Add(CFLAG_ForceDXC);
 	}
 };
 
@@ -684,8 +685,9 @@ public:
 	static void ModifyCompilationEnvironment(const FGlobalShaderPermutationParameters& Parameters, FShaderCompilerEnvironment& OutEnvironment)
 	{
 		FGlobalShader::ModifyCompilationEnvironment(Parameters, OutEnvironment);
+		ShaderPrint::ModifyCompilationEnvironment(Parameters, OutEnvironment);
+
 		OutEnvironment.CompilerFlags.Add(CFLAG_Debug);
-		OutEnvironment.CompilerFlags.Add(CFLAG_ForceDXC);
 		OutEnvironment.SetDefine(TEXT("SHADER_DOMINFO"), 1);
 	}
 };
diff --git a/Engine/Source/Runtime/Renderer/Private/IESTextureManager.cpp b/Engine/Source/Runtime/Renderer/Private/IESTextureManager.cpp
--- a/Engine/Source/Runtime/Renderer/Private/IESTextureManager.cpp
+++ b/Engine/Source/Runtime/Renderer/Private/IESTextureManager.cpp
@@ -126,8 +126,9 @@ public:
 	static void ModifyCompilationEnvironment(const FGlobalShaderPermutationParameters& Parameters, FShaderCompilerEnvironment& OutEnvironment)
 	{
 		FGlobalShader::ModifyCompilationEnvironment(Parameters, OutEnvironment);
+		ShaderPrint::ModifyCompilationEnvironment(Parameters, OutEnvironment);
+
 		OutEnvironment.SetDefine(TEXT("SHADER_DEBUG"), 1);
-		OutEnvironment.CompilerFlags.Add(CFLAG_ForceDXC);
 	}
 };
 
diff --git a/Engine/Source/Runtime/Renderer/Private/LightGridInjection.cpp b/Engine/Source/Runtime/Renderer/Private/LightGridInjection.cpp
--- a/Engine/Source/Runtime/Renderer/Private/LightGridInjection.cpp
+++ b/Engine/Source/Runtime/Renderer/Private/LightGridInjection.cpp
@@ -1091,12 +1091,12 @@ class FDebugLightGridPS : public FGlobalShader
 	static void ModifyCompilationEnvironment(const FGlobalShaderPermutationParameters& Parameters, FShaderCompilerEnvironment& OutEnvironment)
 	{
 		FGlobalShader::ModifyCompilationEnvironment(Parameters, OutEnvironment);
+		ShaderPrint::ModifyCompilationEnvironment(Parameters, OutEnvironment);
 
 		// Stay debug and skip optimizations to reduce compilation time on this long shader.
 		OutEnvironment.CompilerFlags.Add(CFLAG_Debug);
 		OutEnvironment.CompilerFlags.Add(CFLAG_AllowTypedUAVLoads);
 		OutEnvironment.SetDefine(TEXT("SHADER_DEBUG_LIGHT_GRID_PS"), 1);
-		OutEnvironment.CompilerFlags.Add(CFLAG_ForceDXC);
 		FForwardLightingParameters::ModifyCompilationEnvironment(Parameters.Platform, OutEnvironment);
 	}
 };
diff --git a/Engine/Source/Runtime/Renderer/Private/PostProcess/PostProcessing.cpp b/Engine/Source/Runtime/Renderer/Private/PostProcess/PostProcessing.cpp
--- a/Engine/Source/Runtime/Renderer/Private/PostProcess/PostProcessing.cpp
+++ b/Engine/Source/Runtime/Renderer/Private/PostProcess/PostProcessing.cpp
@@ -2949,8 +2949,9 @@ public:
 	static void ModifyCompilationEnvironment(const FGlobalShaderPermutationParameters& Parameters, FShaderCompilerEnvironment& OutEnvironment)
 	{
 		FGlobalShader::ModifyCompilationEnvironment(Parameters, OutEnvironment);
+		ShaderPrint::ModifyCompilationEnvironment(Parameters, OutEnvironment);
+
 		OutEnvironment.SetDefine(TEXT("SHADER_GBUFFER_PICKING"), 1);
-		OutEnvironment.CompilerFlags.Add(CFLAG_ForceDXC);
 	}
 };
 
diff --git a/Engine/Source/Runtime/Renderer/Private/RectLightTextureManager.cpp b/Engine/Source/Runtime/Renderer/Private/RectLightTextureManager.cpp
--- a/Engine/Source/Runtime/Renderer/Private/RectLightTextureManager.cpp
+++ b/Engine/Source/Runtime/Renderer/Private/RectLightTextureManager.cpp
@@ -306,8 +306,9 @@ public:
 	static void ModifyCompilationEnvironment(const FGlobalShaderPermutationParameters& Parameters, FShaderCompilerEnvironment& OutEnvironment)
 	{
 		FGlobalShader::ModifyCompilationEnvironment(Parameters, OutEnvironment);
+		ShaderPrint::ModifyCompilationEnvironment(Parameters, OutEnvironment);
+
 		OutEnvironment.SetDefine(TEXT("SHADER_DEBUG"), 1);
-		OutEnvironment.CompilerFlags.Add(CFLAG_ForceDXC);
 	}
 };
 
diff --git a/Engine/Source/Runtime/Renderer/Private/ShaderPrint.cpp b/Engine/Source/Runtime/Renderer/Private/ShaderPrint.cpp
--- a/Engine/Source/Runtime/Renderer/Private/ShaderPrint.cpp
+++ b/Engine/Source/Runtime/Renderer/Private/ShaderPrint.cpp
@@ -289,11 +289,24 @@ namespace ShaderPrint
 
 	bool IsSupported(EShaderPlatform InShaderPlatform)
 	{
-		// Should be supported everywhere?
-		// If we want to restrict platforms then we can use existing logic here or add a DataDrivenPlatformInfo.
+		// Avoid FXC since it struggles pretty hard with the shaders
+		if (IsD3DPlatform(InShaderPlatform) && !FDataDrivenShaderPlatformInfo::GetSupportsDxc(InShaderPlatform))
+		{
+			return false;
+		}
+
 		return true;
 	}
 
+	void ModifyCompilationEnvironment(const FGlobalShaderPermutationParameters& Parameters, FShaderCompilerEnvironment& OutEnvironment)
+	{
+		// Work around issues with HLSLcc by switching to DXC
+		if (IsHlslccShaderPlatform(Parameters.Platform))
+		{
+			OutEnvironment.CompilerFlags.Add(CFLAG_ForceDXC);
+		}
+	}
+
 	bool IsEnabled()
 	{
 		return GEnabled != 0;
@@ -895,18 +908,15 @@ namespace ShaderPrint
 	{
 		TRACE_CPUPROFILER_EVENT_SCOPE(ShaderPrint::BeginView);
 
-		if (!IsSupported(View.GetShaderPlatform()))
-		{
-			View.ShaderPrintData = FShaderPrintData();
-			return;
-		}
-
 		// Create the render data and store on the view.
 		FShaderPrintSetup ShaderPrintSetup(View);
 		View.ShaderPrintData = CreateShaderPrintData(GraphBuilder, ShaderPrintSetup, View.ViewState);
 
-		// Upload/Copy ShaderPrint parameters into UEDiagnostic buffer
-		InternalUploadParameters(GraphBuilder, View.ShaderPrintData);
+		if (IsSupported(View.GetShaderPlatform()))
+		{
+			// Upload/Copy ShaderPrint parameters into UEDiagnostic buffer
+			InternalUploadParameters(GraphBuilder, View.ShaderPrintData);
+		}
 	}
 
 	void BeginViews(FRDGBuilder& GraphBuilder, TArrayView<FViewInfo> Views)
diff --git a/Engine/Source/Runtime/Renderer/Private/VolumetricCloudRendering.cpp b/Engine/Source/Runtime/Renderer/Private/VolumetricCloudRendering.cpp
--- a/Engine/Source/Runtime/Renderer/Private/VolumetricCloudRendering.cpp
+++ b/Engine/Source/Runtime/Renderer/Private/VolumetricCloudRendering.cpp
@@ -1422,8 +1422,9 @@ public:
 	static void ModifyCompilationEnvironment(const FGlobalShaderPermutationParameters& Parameters, FShaderCompilerEnvironment& OutEnvironment)
 	{
 		FGlobalShader::ModifyCompilationEnvironment(Parameters, OutEnvironment);
+		ShaderPrint::ModifyCompilationEnvironment(Parameters, OutEnvironment);
+
 		OutEnvironment.SetDefine(TEXT("SHADER_DEBUG_SHADOW_CS"), TEXT("1"));
-		OutEnvironment.CompilerFlags.Add(CFLAG_ForceDXC);
 	}
 };
 
diff --git a/Engine/Source/Runtime/Renderer/Public/ShaderPrintParameters.h b/Engine/Source/Runtime/Renderer/Public/ShaderPrintParameters.h
--- a/Engine/Source/Runtime/Renderer/Public/ShaderPrintParameters.h
+++ b/Engine/Source/Runtime/Renderer/Public/ShaderPrintParameters.h
@@ -15,11 +15,13 @@
 #include "ShaderParameterMacros.h"
 #include "ShaderParameters.h"
 
-struct FFrozenShaderPrintData;
-class FSceneView;
-struct FShaderPrintData;
 class FRDGBuilder;
+class FSceneView;
 class FViewInfo;
+struct FFrozenShaderPrintData;
+struct FGlobalShaderPermutationParameters;
+struct FShaderCompilerEnvironment;
+struct FShaderPrintData;
 
 namespace ShaderPrint
 {
@@ -49,6 +51,9 @@ namespace ShaderPrint
 	// Use this to create debug shader permutations only for supported platforms.
 	RENDERER_API bool IsSupported(EShaderPlatform Platform);
 
+	// Set any flags or defines needed when using ShaderPrint
+	RENDERER_API void ModifyCompilationEnvironment(const FGlobalShaderPermutationParameters& Parameters, FShaderCompilerEnvironment& OutEnvironment);
+
 	// Have we enabled the ShaderPrint system?
 	// Note that even when the ShaderPrint system is enabled, it may be disabled on any view due to platform support or view flags.
 	RENDERER_API bool IsEnabled();
