Reorganise hair strands interpolation for batching work.

diff --git a/Engine/Plugins/Runtime/HairStrands/Source/HairStrandsCore/Private/GroomManager.cpp b/Engine/Plugins/Runtime/HairStrands/Source/HairStrandsCore/Private/GroomManager.cpp
--- a/Engine/Plugins/Runtime/HairStrands/Source/HairStrandsCore/Private/GroomManager.cpp
+++ b/Engine/Plugins/Runtime/HairStrands/Source/HairStrandsCore/Private/GroomManager.cpp
@@ -29,6 +29,7 @@
 #include "HairStrandsClusterCulling.h"
 #include "GroomVisualizationData.h"
 #include "SkeletalMeshDeformerHelpers.h"
+#include "HairStrandsInterface.h"
 
 static int32 GHairStrandsMinLOD = 0;
 static FAutoConsoleVariableRef CVarGHairStrandsMinLOD(TEXT("r.HairStrands.MinLOD"), GHairStrandsMinLOD, TEXT("Clamp the min hair LOD to this value, preventing to reach lower/high-quality LOD."), ECVF_Scalability);
@@ -59,6 +60,33 @@ static FAutoConsoleVariableRef CVarHairStrands_AutoLOD_Force(TEXT("r.HairStrands
 static FAutoConsoleVariableRef CVarHairStrands_AutoLOD_Scale(TEXT("r.HairStrands.AutoLOD.Scale"), GHairStrands_AutoLOD_Scale, TEXT("Hair strands Auto LOD rate at which curves get decimated based on screen coverage."), ECVF_RenderThreadSafe);
 static FAutoConsoleVariableRef CVarHairStrands_AutoLOD_Bias(TEXT("r.HairStrands.AutoLOD.Bias"), GHairStrands_AutoLOD_Bias, TEXT("Hair strands Auto LOD screen size bias at which curves get decimated."), ECVF_RenderThreadSafe);
 
+///////////////////////////////////////////////////////////////////////////////////////////////////
+// Foward declaration
+
+bool IsHairStrandsForceRebuildBVH();
+bool IsHairStrandsTransferPositionOnLODChange();
+uint32 GetHairRaytracingProceduralSplits();
+void CreateHairStrandsDebugAttributeBuffer(FRDGBuilder& GraphBuilder, FRDGExternalBuffer* DebugAttributeBuffer, uint32 VertexCount, const FName& OwnerName);
+FHairGroupPublicData::FVertexFactoryInput InternalComputeHairStrandsVertexInputData(FRDGBuilder* GraphBuilder, const FHairGroupInstance* Instance, EGroomViewMode ViewMode);
+EGroomCacheType GetHairInstanceCacheType(const FHairGroupInstance* Instance);
+
+void AddHairClusterAABBPass(
+	FRDGBuilder& GraphBuilder,
+	FGlobalShaderMap* ShaderMap,
+	const uint32 ActiveCurveCount,
+	const FVector& TranslatedWorldOffset,
+	const FShaderPrintData* ShaderPrintData,
+	const EHairAABBUpdateType UpdateType,
+	FHairGroupInstance* Instance,
+	FRDGBufferSRVRef RenderDeformedOffsetBuffer,
+	FHairStrandClusterData::FHairGroup* ClusterData,
+	FRDGHairStrandsCullingData& ClusterAABBData,
+	FRDGBufferSRVRef RenderPositionBufferSRV,
+	FRDGBufferSRVRef& DrawIndirectRasterComputeBuffer);
+
+///////////////////////////////////////////////////////////////////////////////////////////////////
+// Utils
+
 EHairBufferSwapType GetHairSwapBufferType()
 {
 	switch (GHairStrands_SwapBufferType)
@@ -280,50 +308,7 @@ static void RunInternalHairInterpolation(
 			const bool bGlobalDeformationEnable = Instance->HairGroupPublicData->IsGlobalInterpolationEnable(HairLODIndex);
 			check(InstanceBindingType == BindingType);
 
-			if (EHairInterpolationPassType::Strands == PassType)
-			{
-				if (InstanceGeometryType == EHairGeometryType::Strands)
-				{
-					if (BindingType == EHairBindingType::Skinning || bGlobalDeformationEnable)
-					{
-						check(Instance->Strands.HasValidRootData());
-						check(Instance->Strands.DeformedRootResource->IsValid(MeshLODIndex));
-
-						AddHairStrandUpdateMeshTrianglesPass(
-							GraphBuilder,
-							ShaderMap,
-							MeshLODIndex,
-							MeshDataLOD,
-							Instance->Strands.RestRootResource,
-							Instance->Strands.DeformedRootResource);
-
-						AddHairStrandUpdatePositionOffsetPass(
-							GraphBuilder,
-							ShaderMap,
-							EHairPositionUpdateType::Strands,
-							Instance->RegisteredIndex,
-							HairLODIndex,
-							MeshLODIndex,
-							Instance->Strands.DeformedRootResource,
-							Instance->Strands.DeformedResource);
-					}
-					else if (bSimulationEnable || bDeformationEnable)
-					{
-						check(Instance->Strands.HasValidData());
-
-						AddHairStrandUpdatePositionOffsetPass(
-							GraphBuilder,
-							ShaderMap,
-							EHairPositionUpdateType::Strands,
-							Instance->RegisteredIndex,
-							HairLODIndex,
-							MeshLODIndex,
-							nullptr,
-							Instance->Strands.DeformedResource);
-					}
-				}
-			}
-			else if (EHairInterpolationPassType::CardsAndMeshes == PassType)
+			if (EHairInterpolationPassType::CardsAndMeshes == PassType)
 			{
 				if (InstanceGeometryType == EHairGeometryType::Cards)
 				{
@@ -378,7 +363,7 @@ static void RunInternalHairInterpolation(
 	}
 
 	// Hair interpolation
-	if (EHairInterpolationPassType::Strands == PassType || EHairInterpolationPassType::CardsAndMeshes == PassType)
+	if (EHairInterpolationPassType::CardsAndMeshes == PassType)
 	{
 		check(View);
 		const FVector& TranslatedWorldOffset = View->ViewMatrices.GetPreViewTranslation();
@@ -386,8 +371,7 @@ static void RunInternalHairInterpolation(
 		{
 			FHairGroupInstance* Instance = static_cast<FHairGroupInstance*>(AbstractInstance);
 
-			const bool bGeometryCompatibleWithPassType = 
-				(EHairInterpolationPassType::Strands == PassType && Instance->GeometryType == EHairGeometryType::Strands) ||
+			const bool bGeometryCompatibleWithPassType =
 				(EHairInterpolationPassType::CardsAndMeshes == PassType && (Instance->GeometryType == EHairGeometryType::Cards || Instance->GeometryType == EHairGeometryType::Meshes));
 			if (bGeometryCompatibleWithPassType)
 			{
@@ -407,8 +391,6 @@ static void RunInternalHairInterpolation(
 	}
 }
 
-EGroomCacheType GetHairInstanceCacheType(const FHairGroupInstance* Instance);
-
 static void RunHairStrandsInterpolation_Guide(
 	FRDGBuilder& GraphBuilder,
 	FSceneInterface* Scene,
@@ -700,51 +682,602 @@ static void RunHairStrandsInterpolation_Strands(
 	FGlobalShaderMap* ShaderMap)
 {
 	if (Instances.IsEmpty()) { return; }
-
+	if (!IsHairStrandsEnabled(EHairStrandsShaderType::Strands, Scene->GetShaderPlatform())) { return;  }
 	check(IsInRenderingThread());
+	check(View);
 
 	DECLARE_GPU_STAT(HairStrandsInterpolation);
-	RDG_EVENT_SCOPE(GraphBuilder, "HairStrandsInterpolation");
+	RDG_EVENT_SCOPE(GraphBuilder, "HairInterpolation(Strands)");
 	RDG_GPU_STAT_SCOPE(GraphBuilder, HairStrandsInterpolation);
 
+	struct FInstanceRDGResources
+	{
+		FRDGBufferSRVRef PositionSRV = nullptr;
+		FRDGBufferSRVRef PositionOffsetSRV = nullptr;
+		FRDGBufferSRVRef TangentSRV = nullptr;
+
+		FRDGBufferSRVRef DeformerPositionSRV = nullptr;
+		FRDGBufferUAVRef PositionUAV = nullptr;
+		FRDGBufferUAVRef TangentUAV = nullptr;
+
+		FRDGImportedBuffer DeformedPosition;
+		FRDGImportedBuffer DeformedPrevPosition;
+
+		FRDGImportedBuffer Raytracing_PositionBuffer;
+		FRDGImportedBuffer Raytracing_IndexBuffer;
+	};
+
+	struct FInstanceData
+	{
+		int32 HairLODIndex = -1;
+		int32 MeshLODIndex = -1;
+		uint32 ActivePointCount = 0;
+		uint32 ActiveCurveCount = 0;
+		EGroomCacheType ActiveGroomCacheType = EGroomCacheType::None;
+		EHairBindingType BindingType = EHairBindingType::NoneBinding;
+		bool bSimulationEnable = false;
+		bool bDeformationEnable = false;
+		bool bGlobalDeformationEnable = false;
+		bool bNeedDeformation = false;
+		bool bNeedRaytracing = false;
+		bool bNeedRaytracingUpdate = false;
+		bool bNeedRaytracingBuild = false;
+
+
+		FInstanceRDGResources RDGResources;
+
+		FHairStrandsProjectionMeshData::LOD MeshDataLOD;
+		FHairGroupInstance* Instance = nullptr;
+		FRDGHairStrandsCullingData CullingData;
+
+		bool NeedsMeshUpdate() const
+		{
+			return 
+				MeshLODIndex >= 0 && 
+				(BindingType == EHairBindingType::Skinning) && 
+				(bGlobalDeformationEnable || bSimulationEnable || bDeformationEnable);
+		}
+	};
+
+
+	// Gather all strands instances
+	const uint32 ViewRayTracingMask = RHI_RAYTRACING ? (View->Family->EngineShowFlags.PathTracing ? EHairViewRayTracingMask::PathTracing : EHairViewRayTracingMask::RayTracing) : 0u;
+	bool bHasAnySimCacheInstances = false;
+	bool bHasAnyRenCacheInstances = false;
+	bool bHasAnyRaytracingInstances = false;
+	TArray<FInstanceData> InstanceDatas;
+	InstanceDatas.Reserve(Instances.Num());
+	for (FHairStrandsInstance* AbstractInstance : Instances)
+	{
+		FHairGroupInstance* Instance = static_cast<FHairGroupInstance*>(AbstractInstance);
+		if (Instance->GeometryType != EHairGeometryType::Strands)
+		{
+			continue;
+		}	
+		check(Instance->HairGroupPublicData);
+
+		FInstanceData& InstanceData = InstanceDatas.AddDefaulted_GetRef();
+		InstanceData.Instance 					= Instance;
+		InstanceData.ActivePointCount 			= Instance->HairGroupPublicData->GetActiveStrandsPointCount();
+		InstanceData.ActiveCurveCount 			= Instance->HairGroupPublicData->GetActiveStrandsCurveCount();
+		InstanceData.ActiveGroomCacheType 		= GetHairInstanceCacheType(Instance);
+		InstanceData.BindingType 				= Instance->BindingType;
+		InstanceData.HairLODIndex 				= Instance->HairGroupPublicData->LODIndex;
+		InstanceData.bSimulationEnable 			= Instance->HairGroupPublicData->IsSimulationEnable(InstanceData.HairLODIndex);
+		InstanceData.bDeformationEnable 		= Instance->HairGroupPublicData->bIsDeformationEnable;
+		InstanceData.bGlobalDeformationEnable 	= Instance->HairGroupPublicData->IsGlobalInterpolationEnable(InstanceData.HairLODIndex);
+		InstanceData.MeshLODIndex				= -1;
+		InstanceData.bNeedDeformation 			= Instance->Strands.DeformedResource != nullptr;
+		InstanceData.bNeedRaytracing			= false;
+		InstanceData.bNeedRaytracingUpdate		= false;
+		InstanceData.bNeedRaytracingBuild		= false;
+
+		// Register resources
+		FInstanceRDGResources RDGResources;
+		if (InstanceData.bNeedDeformation)
+		{		
+			// Trach on which view the position has been update, so that we can ensure motion vector are coherent
+			Instance->Strands.DeformedResource->GetUniqueViewID(FHairStrandsDeformedResource::Current) = ViewUniqueID;
+			
+			InstanceData.RDGResources.DeformedPosition		= Register(GraphBuilder, Instance->Strands.DeformedResource->GetBuffer(FHairStrandsDeformedResource::Current), ERDGImportedBufferFlags::CreateViews);
+			InstanceData.RDGResources.DeformedPrevPosition	= Register(GraphBuilder, Instance->Strands.DeformedResource->GetBuffer(FHairStrandsDeformedResource::Previous), ERDGImportedBufferFlags::CreateViews);
+
+			InstanceData.RDGResources.PositionSRV 			= InstanceData.RDGResources.DeformedPosition.SRV;
+			InstanceData.RDGResources.PositionUAV			= InstanceData.RDGResources.DeformedPosition.UAV;
+			InstanceData.RDGResources.PositionOffsetSRV 	= RegisterAsSRV(GraphBuilder, Instance->Strands.DeformedResource->GetPositionOffsetBuffer(FHairStrandsDeformedResource::EFrameType::Current));		
+			InstanceData.RDGResources.TangentSRV 			= RegisterAsSRV(GraphBuilder, Instance->Strands.DeformedResource->TangentBuffer);
+			InstanceData.RDGResources.TangentUAV 			= RegisterAsUAV(GraphBuilder, Instance->Strands.DeformedResource->TangentBuffer);
+			InstanceData.RDGResources.DeformerPositionSRV 	= RegisterAsSRV(GraphBuilder, Instance->Strands.DeformedResource->DeformerBuffer);
+		}
+		else
+		{
+			InstanceData.RDGResources.PositionSRV 			= RegisterAsSRV(GraphBuilder, Instance->Strands.RestResource->PositionBuffer);
+			InstanceData.RDGResources.PositionUAV			= nullptr;
+			InstanceData.RDGResources.PositionOffsetSRV 	= RegisterAsSRV(GraphBuilder, Instance->Strands.RestResource->PositionOffsetBuffer);
+			InstanceData.RDGResources.TangentSRV 			= nullptr;  
+			InstanceData.RDGResources.DeformerPositionSRV 	= nullptr;
+		}
+
+		#if RHI_RAYTRACING
+		InstanceData.bNeedRaytracing = InstanceData.Instance->Strands.RenRaytracingResource && (InstanceData.Instance->Strands.ViewRayTracingMask & ViewRayTracingMask) != 0;
+		if (InstanceData.bNeedRaytracing)
+		{
+			InstanceData.RDGResources.Raytracing_PositionBuffer = Register(GraphBuilder, InstanceData.Instance->Strands.RenRaytracingResource->PositionBuffer, ERDGImportedBufferFlags::CreateViews);
+			InstanceData.RDGResources.Raytracing_IndexBuffer 	= Register(GraphBuilder, InstanceData.Instance->Strands.RenRaytracingResource->IndexBuffer, ERDGImportedBufferFlags::CreateViews);
+			bHasAnyRaytracingInstances = true;
+		}
+		#endif
+
+		if (InstanceData.ActiveGroomCacheType == EGroomCacheType::Guides)  { bHasAnySimCacheInstances = true; }
+		if (InstanceData.ActiveGroomCacheType == EGroomCacheType::Strands) { bHasAnyRenCacheInstances = true; }
+
+		// Extract MeshDataLOD and compute MeshLODIndex
+		const FCachedGeometry CachedGeometry = GetCacheGeometryForHair(GraphBuilder, Scene, Instance, ShaderMap, true);
+		for (int32 SectionIndex = 0; SectionIndex < CachedGeometry.Sections.Num(); ++SectionIndex)
+		{
+			// Ensure all mesh's sections have the same LOD index
+			const int32 SectionLodIndex = CachedGeometry.Sections[SectionIndex].LODIndex;
+			if (InstanceData.MeshLODIndex < 0) InstanceData.MeshLODIndex = SectionLodIndex;
+			check(InstanceData.MeshLODIndex == SectionLodIndex);
+
+			InstanceData.MeshDataLOD.Sections.Add(ConvertMeshSection(CachedGeometry, SectionIndex));
+		}
+		Instance->Debug.MeshLODIndex = InstanceData.MeshLODIndex;
+		Instance->HairGroupPublicData->VFInput.Strands	= FHairGroupPublicData::FVertexFactoryInput::FStrands();
+
+		if (InstanceData.MeshLODIndex >= 0 && (InstanceData.BindingType == EHairBindingType::Skinning || InstanceData.bGlobalDeformationEnable))
+		{
+			check(Instance->Strands.HasValidRootData());
+			check(Instance->Strands.DeformedRootResource->IsValid(InstanceData.MeshLODIndex));
+		}
+
+		if (InstanceData.MeshLODIndex >= 0 && (InstanceData.bSimulationEnable || InstanceData.bDeformationEnable))
+		{
+			check(Instance->Strands.HasValidData());
+		}
+	}	
+
+	FRDGExternalAccessQueue ExternalAccessQueue;
+	const EGroomViewMode ViewMode = GetGroomViewMode(*View);
+
+	// Early culling
 	FHairStrandClusterData ClusterData;
-	if (IsHairStrandsEnabled(EHairStrandsShaderType::Strands, Scene->GetShaderPlatform()))
 	{
-		for (FHairStrandsInstance* AbstractInstance : Instances)
+		// Gather culling jobs
+		for (FInstanceData& InstanceData : InstanceDatas)
 		{
-			FHairGroupInstance* Instance = static_cast<FHairGroupInstance*>(AbstractInstance);
-			if (Instance->GeometryType == EHairGeometryType::Strands && Instance->Strands.IsValid())
+			if (InstanceData.Instance->Strands.bCullingEnable)
 			{
-				if (Instance->Strands.bCullingEnable)
-				{
-					AddInstanceToClusterData(Instance, ClusterData);
-				}
-				else
-				{
-					Instance->HairGroupPublicData->SetCullingResultAvailable(false);
-				}
+				AddInstanceToClusterData(InstanceData.Instance, ClusterData);
+			}
+			else
+			{
+				InstanceData.Instance->HairGroupPublicData->SetCullingResultAvailable(false);
 			}
 		}
+
+		// Culling pass
 		ComputeHairStrandsClustersCulling(GraphBuilder, *ShaderMap, Views, ShaderPrintData, ClusterData);
 
 		// Run cluster debug view here (instead of GroomDebug.h/.cpp, as we need to have the (transient) cluster data 
-		const EGroomViewMode ViewMode = GetGroomViewMode(*View);
 		if (ViewMode == EGroomViewMode::Cluster || ViewMode == EGroomViewMode::ClusterAABB)
 		{
 			AddDrawDebugClusterPass(GraphBuilder, *View, ShaderMap, ShaderPrintData, ViewMode, ClusterData);
 		}
 	}
 
-	RunInternalHairInterpolation(
-		GraphBuilder,
-		Scene,
-		View,
-		ViewUniqueID,
-		Instances,
-		ShaderPrintData,
-		ShaderMap,
-		EHairInterpolationPassType::Strands,
-		&ClusterData);
+	// Update dynamic mesh triangles
+	for (FInstanceData& InstanceData : InstanceDatas)
+	{
+		if (InstanceData.MeshLODIndex >= 0 && (InstanceData.BindingType == EHairBindingType::Skinning || InstanceData.bGlobalDeformationEnable))
+		{
+			AddHairStrandUpdateMeshTrianglesPass(
+				GraphBuilder,
+				ShaderMap,
+				InstanceData.MeshLODIndex,
+				InstanceData.MeshDataLOD,
+				InstanceData.Instance->Strands.RestRootResource,
+				InstanceData.Instance->Strands.DeformedRootResource);
+		}
+	}
+
+	// Update position offset (GPU)
+	for (FInstanceData& InstanceData : InstanceDatas)
+	{
+		if (InstanceData.MeshLODIndex >= 0 && (InstanceData.BindingType == EHairBindingType::Skinning || InstanceData.bGlobalDeformationEnable))
+		{
+			AddHairStrandUpdatePositionOffsetPass(
+				GraphBuilder,
+				ShaderMap,
+				EHairPositionUpdateType::Strands,
+				InstanceData.Instance->RegisteredIndex,
+				InstanceData.HairLODIndex,
+				InstanceData.MeshLODIndex,
+				InstanceData.Instance->Strands.DeformedRootResource,
+				InstanceData.Instance->Strands.DeformedResource);
+		}
+	}
+
+	// Update position offset (CPU)
+	for (FInstanceData& InstanceData : InstanceDatas)
+	{
+		if (InstanceData.MeshLODIndex >= 0 && (InstanceData.bSimulationEnable || InstanceData.bDeformationEnable))
+		{
+			AddHairStrandUpdatePositionOffsetPass(
+				GraphBuilder,
+				ShaderMap,
+				EHairPositionUpdateType::Strands,
+				InstanceData.Instance->RegisteredIndex,
+				InstanceData.HairLODIndex,
+				InstanceData.MeshLODIndex,
+				nullptr,
+				InstanceData.Instance->Strands.DeformedResource);
+		}
+	}
+
+	// Update position offset (CPU) - With Sim Cache
+	if (bHasAnySimCacheInstances)
+	{
+		for (FInstanceData& InstanceData : InstanceDatas)
+		{
+			if (InstanceData.ActiveGroomCacheType == EGroomCacheType::Guides && InstanceData.bNeedDeformation)
+			{
+				// Apply the same offset to the render strands for proper rendering
+				const FVector OffsetPosition = InstanceData.Instance->Guides.DeformedResource->GetPositionOffset(FHairStrandsDeformedResource::EFrameType::Current);
+				InstanceData.Instance->Strands.DeformedResource->SetPositionOffset(FHairStrandsDeformedResource::EFrameType::Current, OffsetPosition);
+		
+				AddHairStrandUpdatePositionOffsetPass(
+					GraphBuilder,
+					ShaderMap,
+					EHairPositionUpdateType::Strands,
+					InstanceData.Instance->RegisteredIndex,
+					0,
+					InstanceData.MeshLODIndex,
+					InstanceData.Instance->Strands.DeformedRootResource,
+					InstanceData.Instance->Strands.DeformedResource);
+			}
+		}
+	}
+
+	// Clear cluster AABBs (used optionally  for voxel allocation & for culling)
+	for (FInstanceData& InstanceData : InstanceDatas)
+	{
+		//if (InstanceData.bNeedDeformation)
+		{
+			InstanceData.CullingData = ImportCullingData(GraphBuilder, InstanceData.Instance->HairGroupPublicData);
+			AddClearClusterAABBPass(
+				GraphBuilder,
+				ShaderMap,
+				InstanceData.bNeedDeformation ? EHairAABBUpdateType::UpdateClusterAABB : EHairAABBUpdateType::UpdateGroupAABB,
+				InstanceData.CullingData.ClusterCount,
+				InstanceData.CullingData.ClusterAABBBuffer,
+				InstanceData.CullingData.GroupAABBBuffer);
+		}
+	}
+
+	for (FInstanceData& InstanceData : InstanceDatas)
+	{
+		if (InstanceData.ActiveGroomCacheType != EGroomCacheType::Strands && InstanceData.bNeedDeformation)
+		{
+			const bool bValidGuide		= InstanceData.bNeedDeformation && (InstanceData.bSimulationEnable || InstanceData.bGlobalDeformationEnable || InstanceData.bDeformationEnable || InstanceData.Instance->Guides.bIsSimulationCacheEnable);// || (WITH_EDITOR && PatchMode == EHairPatchAttribute::GuideInflucence);
+			const bool bUseSingleGuide	= InstanceData.bNeedDeformation && bValidGuide && InstanceData.Instance->Strands.InterpolationResource->UseSingleGuide();
+			const bool bHasSkinning		= InstanceData.bNeedDeformation && InstanceData.BindingType == EHairBindingType::Skinning;
+
+			AddHairStrandsInterpolationPass(
+				GraphBuilder,
+				ShaderMap,
+				ShaderPrintData,
+				InstanceData.Instance,
+				InstanceData.ActivePointCount,
+				InstanceData.MeshLODIndex,
+				InstanceData.Instance->Strands.Modifier.HairLengthScale,
+				InstanceData.Instance->Strands.HairInterpolationType,
+				EHairGeometryType::Strands,
+				InstanceData.CullingData,
+				InstanceData.Instance->Strands.RestResource->GetPositionOffset(),
+				bValidGuide ? InstanceData.Instance->Guides.RestResource->GetPositionOffset() : FVector::ZeroVector,
+				InstanceData.RDGResources.PositionOffsetSRV,
+				bValidGuide ? RegisterAsSRV(GraphBuilder, InstanceData.Instance->Guides.DeformedResource->GetPositionOffsetBuffer(FHairStrandsDeformedResource::Current)) : nullptr,
+				bHasSkinning ? InstanceData.Instance->Strands.RestRootResource : nullptr,
+				bHasSkinning && bValidGuide ? InstanceData.Instance->Guides.RestRootResource : nullptr,
+				bHasSkinning ? InstanceData.Instance->Strands.DeformedRootResource : nullptr,
+				bHasSkinning && bValidGuide ? InstanceData.Instance->Guides.DeformedRootResource : nullptr,
+				RegisterAsSRV(GraphBuilder, InstanceData.Instance->Strands.RestResource->PositionBuffer),
+				RegisterAsSRV(GraphBuilder, InstanceData.Instance->Strands.RestResource->PointToCurveBuffer),
+				bUseSingleGuide,
+				bValidGuide ? RegisterAsSRV(GraphBuilder, InstanceData.Instance->Strands.InterpolationResource->InterpolationBuffer) : nullptr,
+				bValidGuide ? RegisterAsSRV(GraphBuilder, InstanceData.Instance->Guides.RestResource->PositionBuffer) : nullptr,
+				bValidGuide ? RegisterAsSRV(GraphBuilder, InstanceData.Instance->Guides.DeformedResource->GetBuffer(FHairStrandsDeformedResource::Current)) : nullptr,
+				bValidGuide ? RegisterAsSRV(GraphBuilder, InstanceData.Instance->Strands.InterpolationResource->SimRootPointIndexBuffer) : nullptr,
+				bValidGuide ? RegisterAsSRV(GraphBuilder, InstanceData.Instance->Guides.RestResource->PointToCurveBuffer) : nullptr,
+				InstanceData.RDGResources.DeformerPositionSRV,
+				InstanceData.RDGResources.PositionUAV,
+				FHairStrandsDeformedRootResource::FLOD::Current);
+
+			GraphBuilder.SetBufferAccessFinal(InstanceData.RDGResources.DeformedPosition.Buffer, ERHIAccess::SRVMask);
+		}
+	}
+
+	// Previous Position
+	// * Transfer prev. position
+	// * Or recompute interpolation pass for previous positions if needed
+	for (FInstanceData& InstanceData : InstanceDatas)
+	{
+		if (InstanceData.ActiveGroomCacheType != EGroomCacheType::Strands && InstanceData.bNeedDeformation)
+		{
+			const bool bTransferPrevPosition =
+				InstanceData.Instance->HairGroupPublicData->VFInput.bHasLODSwitch &&
+				IsHairStrandsTransferPositionOnLODChange();
+			if (bTransferPrevPosition)
+			{
+				AddTransferPositionPass(GraphBuilder, ShaderMap, InstanceData.ActivePointCount, InstanceData.RDGResources.DeformedPosition.SRV, InstanceData.RDGResources.DeformedPrevPosition.UAV);
+				GraphBuilder.SetBufferAccessFinal(InstanceData.RDGResources.DeformedPrevPosition.Buffer, ERHIAccess::SRVMask);
+			}
+		}
+	}
+
+	// Render Strands cache update
+	if (bHasAnyRenCacheInstances)
+	{
+		for (FInstanceData& InstanceData : InstanceDatas)
+		{
+			if (InstanceData.ActiveGroomCacheType == EGroomCacheType::Strands)
+			{
+				FScopeLock Lock(InstanceData.Instance->Debug.GroomCacheBuffers->GetCriticalSection());
+				
+				FGroomCacheGroupData* GroomCacheData0 = const_cast<FGroomCacheGroupData*>(&InstanceData.Instance->Debug.GroomCacheBuffers->GetCurrentFrameBuffer().GroupsData[InstanceData.Instance->Debug.GroupIndex]);
+				FGroomCacheGroupData* GroomCacheData1 = const_cast<FGroomCacheGroupData*>(&InstanceData.Instance->Debug.GroomCacheBuffers->GetNextFrameBuffer().GroupsData[InstanceData.Instance->Debug.GroupIndex]);
+				const float InterpolationFactor = InstanceData.Instance->Debug.GroomCacheBuffers->GetInterpolationFactor();
+				
+				FGroomCacheResources CacheResources0 = CreateGroomCacheBuffer(GraphBuilder, GroomCacheData0->VertexData);
+				FGroomCacheResources CacheResources1 = CreateGroomCacheBuffer(GraphBuilder, GroomCacheData1->VertexData);
+				
+				// Update position offset
+				{
+					const FVector OffsetPosition = FMath::Lerp(GroomCacheData0->BoundingBox.GetCenter(), GroomCacheData1->BoundingBox.GetCenter(), InterpolationFactor);
+					InstanceData.Instance->Strands.DeformedResource->SetPositionOffset(FHairStrandsDeformedResource::EFrameType::Current, OffsetPosition);
+				}
+				
+				// Update max radius
+				if (GroomCacheData0->VertexData.PointsRadius.Num() > 0)
+				{
+					InstanceData.Instance->Strands.Modifier.HairWidth = FMath::Lerp(GroomCacheData0->StrandData.MaxRadius, GroomCacheData1->StrandData.MaxRadius, InterpolationFactor) * 2.0f;
+				}
+				
+				AddHairStrandUpdatePositionOffsetPass(
+					GraphBuilder,
+					ShaderMap,
+					EHairPositionUpdateType::Strands,
+					InstanceData.Instance->RegisteredIndex,
+					0,
+					InstanceData.MeshLODIndex,
+					InstanceData.Instance->Strands.DeformedRootResource,
+					InstanceData.Instance->Strands.DeformedResource);
+				
+				// Pass to upload GroomCache strands positions
+				AddGroomCacheUpdatePass(
+					GraphBuilder,
+					ShaderMap,
+					InstanceData.ActivePointCount,
+					InterpolationFactor,
+					CacheResources0,
+					CacheResources1,
+					RegisterAsSRV(GraphBuilder, InstanceData.Instance->Strands.RestResource->PositionBuffer),
+					RegisterAsSRV(GraphBuilder, InstanceData.Instance->Strands.DeformedResource->GetPositionOffsetBuffer(FHairStrandsDeformedResource::EFrameType::Current)),
+					InstanceData.RDGResources.PositionUAV);
+			}
+		}
+	}
+
+	// Update tangent data based on the deformed positions
+	for (FInstanceData& InstanceData : InstanceDatas)
+	{
+		if (InstanceData.bNeedDeformation)
+		{
+			AddHairTangentPass(
+				GraphBuilder,
+				ShaderMap,
+				InstanceData.ActivePointCount,
+				InstanceData.Instance->HairGroupPublicData,
+				InstanceData.RDGResources.PositionSRV,
+				InstanceData.RDGResources.TangentUAV);
+		}
+		else
+		{
+			InstanceData.Instance->Strands.RestResource->GetTangentBuffer(GraphBuilder, ShaderMap, InstanceData.ActivePointCount, InstanceData.ActiveCurveCount);
+		}
+	}
+
+	// Patch attribute for debug visualization (guide influence or clusters visualization)
+	const bool bNeedPatchPass = ViewMode == EGroomViewMode::RenderHairStrands || ViewMode == EGroomViewMode::Cluster || ViewMode == EGroomViewMode::ClusterAABB;
+	if (bNeedPatchPass)
+	{
+		for (FInstanceData& InstanceData : InstanceDatas)
+		{
+			EHairPatchAttribute PatchMode = EHairPatchAttribute::None;
+			if (ViewMode == EGroomViewMode::RenderHairStrands && InstanceData.bNeedDeformation)	{ PatchMode = EHairPatchAttribute::GuideInflucence; }
+			if (ViewMode == EGroomViewMode::Cluster || ViewMode == EGroomViewMode::ClusterAABB)	{ PatchMode = EHairPatchAttribute::ClusterInfluence; }
+
+			if (PatchMode != EHairPatchAttribute::None)
+			{
+				// Create an debug buffer for storing cluster visualization data. This is only used for debug purpose, hence only enable in editor build.
+				// Special case for debug mode were the attribute buffer is patch with some custom data to show hair properties (strands belonging to the same cluster, ...)
+				if (InstanceData.Instance->Strands.DebugCurveAttributeBuffer.Buffer == nullptr)
+				{
+					CreateHairStrandsDebugAttributeBuffer(GraphBuilder, &InstanceData.Instance->Strands.DebugCurveAttributeBuffer, InstanceData.Instance->Strands.Data->GetCurveAttributeSizeInBytes(), FName(InstanceData.Instance->Debug.MeshComponentName));
+				}
+				FRDGImportedBuffer OutRenCurveAttributeBuffer = Register(GraphBuilder, InstanceData.Instance->Strands.DebugCurveAttributeBuffer, ERDGImportedBufferFlags::CreateUAV);
+			
+				const bool bValidGuide = InstanceData.bNeedDeformation && (InstanceData.bSimulationEnable || InstanceData.bGlobalDeformationEnable || InstanceData.bDeformationEnable || InstanceData.Instance->Guides.bIsSimulationCacheEnable);// || (WITH_EDITOR && PatchMode == EHairPatchAttribute::GuideInflucence);
+				const bool bUseSingleGuide	= InstanceData.bNeedDeformation && bValidGuide && InstanceData.Instance->Strands.InterpolationResource->UseSingleGuide();
+
+				check(InstanceData.Instance->Strands.Data);
+				AddPatchAttributePass(
+					GraphBuilder,
+					ShaderMap,
+					InstanceData.ActiveCurveCount,
+					PatchMode,
+					bValidGuide,
+					bUseSingleGuide,
+					*InstanceData.Instance->Strands.Data,
+					Register(GraphBuilder, InstanceData.Instance->Strands.RestResource->CurveAttributeBuffer, ERDGImportedBufferFlags::CreateSRV).Buffer,
+					RegisterAsSRV(GraphBuilder, InstanceData.Instance->Strands.RestResource->CurveBuffer),
+					RegisterAsSRV(GraphBuilder, InstanceData.Instance->Strands.ClusterResource->CurveToClusterIdBuffer),
+					bValidGuide ? RegisterAsSRV(GraphBuilder, InstanceData.Instance->Strands.InterpolationResource->InterpolationBuffer) : nullptr,
+					OutRenCurveAttributeBuffer);
+			}
+		}
+	}
+
+	// Update the VF input with the update resources
+	for (FInstanceData& InstanceData : InstanceDatas)
+	{
+		InstanceData.Instance->HairGroupPublicData->VFInput = InternalComputeHairStrandsVertexInputData(&GraphBuilder, InstanceData.Instance, ViewMode);
+	}
+
+	// 3. Compute cluster AABBs (used for LODing and voxelization)
+	const FVector& TranslatedWorldOffset = View->ViewMatrices.GetPreViewTranslation();
+	for (FInstanceData& InstanceData : InstanceDatas)
+	{
+		// Optim: If an instance does not voxelize it's data, then there is no need for having valid AABB
+		bool bNeedGPUAABB = InstanceData.Instance->Strands.Modifier.bSupportVoxelization && InstanceData.Instance->bCastShadow;
+		const EHairAABBUpdateType UpdateType = InstanceData.bNeedDeformation ? EHairAABBUpdateType::UpdateClusterAABB : EHairAABBUpdateType::UpdateGroupAABB;
+			
+		FHairStrandClusterData::FHairGroup* HairGroupCluster = nullptr;
+		if (InstanceData.CullingData.bCullingResultAvailable)
+		{
+			// Sanity check
+			check(InstanceData.Instance->Strands.bCullingEnable);
+			HairGroupCluster = &ClusterData.HairGroups[InstanceData.Instance->HairGroupPublicData->ClusterDataIndex];
+			if (!HairGroupCluster->bVisible)
+			{
+				bNeedGPUAABB = false;
+			}
+		}
+			
+		InstanceData.Instance->HairGroupPublicData->SetClusterAABBValid(false);
+		InstanceData.Instance->HairGroupPublicData->SetGroupAABBValid(false);
+			
+		if (bNeedGPUAABB)
+		{
+			FRDGImportedBuffer Strands_CulledVertexCount = Register(GraphBuilder, InstanceData.Instance->HairGroupPublicData->GetDrawIndirectRasterComputeBuffer(), ERDGImportedBufferFlags::CreateSRV);
+			AddHairClusterAABBPass(
+				GraphBuilder,
+				ShaderMap,
+				InstanceData.ActiveCurveCount,
+				TranslatedWorldOffset,
+				ShaderPrintData,
+				UpdateType,
+				InstanceData.Instance,
+				InstanceData.RDGResources.PositionOffsetSRV,
+				HairGroupCluster,
+				InstanceData.CullingData,
+				InstanceData.RDGResources.PositionSRV,
+				Strands_CulledVertexCount.SRV);
+			
+			InstanceData.Instance->HairGroupPublicData->SetClusterAABBValid(UpdateType == EHairAABBUpdateType::UpdateClusterAABB);
+			InstanceData.Instance->HairGroupPublicData->SetGroupAABBValid(true);
+		}
+	}
+
+	#if RHI_RAYTRACING
+	if (bHasAnyRaytracingInstances)
+	{
+		const uint32 ProceduralSplits = GetHairRaytracingProceduralSplits();
+
+		// 1. Update raytracing geometry (update only if the view mask and the RT geometry mask match)
+		for (FInstanceData& InstanceData : InstanceDatas)
+		{
+			if (InstanceData.bNeedRaytracing)
+			{
+				const float CLODScale = InstanceData.Instance->HairGroupPublicData->ContinuousLODCoverageScale;
+				const float HairRadiusRT = InstanceData.Instance->HairGroupPublicData->VFInput.Strands.Common.RaytracingRadiusScale * InstanceData.Instance->HairGroupPublicData->VFInput.Strands.Common.Radius * CLODScale;
+				const float HairRootScaleRT = InstanceData.Instance->HairGroupPublicData->VFInput.Strands.Common.RootScale;
+				const float HairTipScaleRT = InstanceData.Instance->HairGroupPublicData->VFInput.Strands.Common.TipScale;
+
+				InstanceData.bNeedRaytracingUpdate = InstanceData.Instance->Strands.DeformedResource != nullptr ||
+					InstanceData.Instance->Strands.CachedHairScaledRadius != HairRadiusRT ||
+					InstanceData.Instance->Strands.CachedHairRootScale != HairRootScaleRT ||
+					InstanceData.Instance->Strands.CachedHairTipScale != HairTipScaleRT ||
+					InstanceData.Instance->Strands.CachedProceduralSplits != ProceduralSplits;
+				InstanceData.bNeedRaytracingBuild = !InstanceData.Instance->Strands.RenRaytracingResource->bIsRTGeometryInitialized;
+				if (InstanceData.bNeedRaytracingBuild || InstanceData.bNeedRaytracingUpdate)
+				{
+					const bool bProceduralPrimitive = InstanceData.Instance->Strands.RenRaytracingResource->bProceduralPrimitive;
+					AddGenerateRaytracingGeometryPass(
+						GraphBuilder,
+						ShaderMap,
+						ShaderPrintData,
+						InstanceData.Instance->RegisteredIndex,
+						InstanceData.ActivePointCount,
+						bProceduralPrimitive,
+						ProceduralSplits,
+						HairRadiusRT,
+						HairRootScaleRT,
+						HairTipScaleRT,
+						InstanceData.RDGResources.PositionOffsetSRV,
+						InstanceData.CullingData,
+						InstanceData.RDGResources.PositionSRV,
+						InstanceData.RDGResources.TangentSRV,
+						InstanceData.RDGResources.Raytracing_PositionBuffer.UAV,
+						InstanceData.RDGResources.Raytracing_IndexBuffer.UAV);
+			
+					InstanceData.Instance->Strands.CachedHairScaledRadius = HairRadiusRT;
+					InstanceData.Instance->Strands.CachedHairRootScale = HairRootScaleRT;
+					InstanceData.Instance->Strands.CachedHairTipScale = HairTipScaleRT;
+				}
+			}
+		}
+
+		// 2. Build/update acceleration structure
+		for (FInstanceData& InstanceData : InstanceDatas)
+		{
+			if (InstanceData.bNeedRaytracing && (InstanceData.bNeedRaytracingBuild || InstanceData.bNeedRaytracingUpdate))
+			{
+				AddBuildStrandsAccelerationStructurePass(
+					GraphBuilder,
+					InstanceData.Instance,
+					ProceduralSplits,
+					InstanceData.bNeedRaytracingUpdate,
+					InstanceData.RDGResources.Raytracing_PositionBuffer.Buffer,
+					InstanceData.RDGResources.Raytracing_IndexBuffer.Buffer);
+			}
+		}
+	}
+	#endif
+
+	if (ViewMode == EGroomViewMode::SimHairStrands)
+	{
+		for (FInstanceData& InstanceData : InstanceDatas)
+		{
+			if (InstanceData.Instance->Guides.DeformedResource)
+			{
+				// Disable culling when drawing only guides, as the culling output has been computed for the strands, not for the guides.
+				InstanceData.Instance->HairGroupPublicData->SetCullingResultAvailable(false);
+					
+				AddHairTangentPass(
+					GraphBuilder,
+					ShaderMap,
+					InstanceData.Instance->Guides.RestResource->GetPointCount(),
+					InstanceData.Instance->HairGroupPublicData,
+					RegisterAsSRV(GraphBuilder, InstanceData.Instance->Guides.DeformedResource->GetBuffer(FHairStrandsDeformedResource::Current)),
+					RegisterAsUAV(GraphBuilder, InstanceData.Instance->Guides.DeformedResource->TangentBuffer));					
+			}
+			InstanceData.Instance->HairGroupPublicData->VFInput = InternalComputeHairStrandsVertexInputData(&GraphBuilder, InstanceData.Instance, ViewMode);
+		}
+	}
+
+	// Sanity check
+	for (FInstanceData& InstanceData : InstanceDatas)
+	{
+		check(InstanceData.Instance->HairGroupPublicData->VFInput.Strands.PositionBuffer.Buffer);
+		InstanceData.Instance->Strands.UniformBuffer.UpdateUniformBufferImmediate(GraphBuilder.RHICmdList, InstanceData.Instance->GetHairStandsUniformShaderParameters(ViewMode));
+		InstanceData.Instance->HairGroupPublicData->VFInput.GeometryType = EHairGeometryType::Strands;
+		InstanceData.Instance->HairGroupPublicData->VFInput.LocalToWorldTransform = InstanceData.Instance->GetCurrentLocalToWorld();
+		InstanceData.Instance->HairGroupPublicData->bSupportVoxelization = InstanceData.Instance->Strands.Modifier.bSupportVoxelization && InstanceData.Instance->bCastShadow;
+	}
+
+	ExternalAccessQueue.Submit(GraphBuilder);
 }
 
 static void RunHairStrandsInterpolation_Cards(
diff --git a/Engine/Plugins/Runtime/HairStrands/Source/HairStrandsCore/Private/GroomResources.cpp b/Engine/Plugins/Runtime/HairStrands/Source/HairStrandsCore/Private/GroomResources.cpp
--- a/Engine/Plugins/Runtime/HairStrands/Source/HairStrandsCore/Private/GroomResources.cpp
+++ b/Engine/Plugins/Runtime/HairStrands/Source/HairStrandsCore/Private/GroomResources.cpp
@@ -1069,7 +1069,7 @@ void AddHairTangentPass(
 	uint32 PointCount,
 	FHairGroupPublicData* HairGroupPublicData,
 	FRDGBufferSRVRef PositionBuffer,
-	FRDGImportedBuffer OutTangentBuffer);
+	FRDGBufferUAVRef OutTangentBuffer);
 
 FRDGExternalBuffer FHairStrandsRestResource::GetTangentBuffer(FRDGBuilder& GraphBuilder, FGlobalShaderMap* ShaderMap, uint32 ActivePointCount, uint32 ActiveCurveCount)
 {
@@ -1093,7 +1093,7 @@ FRDGExternalBuffer FHairStrandsRestResource::GetTangentBuffer(FRDGBuilder& Graph
 			ActivePointCount,
 			nullptr,
 			RegisterAsSRV(GraphBuilder, PositionBuffer),
-			Register(GraphBuilder, TangentBuffer, ERDGImportedBufferFlags::CreateUAV));
+			RegisterAsUAV(GraphBuilder, TangentBuffer));
 
 		CachedTangentPointCount = ActivePointCount;
 	}
diff --git a/Engine/Plugins/Runtime/HairStrands/Source/HairStrandsCore/Private/HairStrandsInterpolation.cpp b/Engine/Plugins/Runtime/HairStrands/Source/HairStrandsCore/Private/HairStrandsInterpolation.cpp
--- a/Engine/Plugins/Runtime/HairStrands/Source/HairStrandsCore/Private/HairStrandsInterpolation.cpp
+++ b/Engine/Plugins/Runtime/HairStrands/Source/HairStrandsCore/Private/HairStrandsInterpolation.cpp
@@ -60,6 +60,16 @@ enum class EHairCardsSimulationType
 	RBF
 };
 
+bool IsHairStrandsTransferPositionOnLODChange()
+{
+	return GHairStrandsTransferPositionOnLODChange > 0;
+}
+
+bool IsHairStrandsForceRebuildBVH()
+{
+	return GHairStrands_Raytracing_ForceRebuildBVH > 0;
+}
+
 EHairCardsSimulationType GetHairCardsSimulationType()
 {
 	return GHairCardsInterpolationType >= 2 ? 
@@ -72,7 +82,7 @@ bool NeedsUpdateCardsMeshTriangles()
 	return GetHairCardsSimulationType() == EHairCardsSimulationType::Guide;
 }
 
-int GetHairRaytracingProceduralSplits()
+uint32 GetHairRaytracingProceduralSplits()
 {
 	return FMath::Clamp(GHairRaytracingProceduralSplits, 1, STRANDS_PROCEDURAL_INTERSECTOR_MAX_SPLITS);
 }
@@ -187,7 +197,7 @@ public:
 
 IMPLEMENT_GLOBAL_SHADER(FTransferVelocityPassCS, "/Engine/Private/HairStrands/HairStrandsInterpolation.usf", "MainCS", SF_Compute);
 
-static void AddTransferPositionPass(
+void AddTransferPositionPass(
 	FRDGBuilder& GraphBuilder,
 	FGlobalShaderMap* ShaderMap,
 	const uint32 ElementCount,
@@ -427,18 +437,6 @@ void AddDeformSimHairStrandsPass(
 
 ///////////////////////////////////////////////////////////////////////////////////////////////////
 
-struct FRDGHairStrandsCullingData
-{
-	bool bCullingResultAvailable = false;
-	FRDGImportedBuffer HairStrandsVF_CullingIndirectBuffer;
-	FRDGImportedBuffer HairStrandsVF_CullingIndexBuffer;
-	FRDGImportedBuffer HairStrandsVF_CullingRadiusScaleBuffer;
-
-	uint32 ClusterCount = 0;
-	FRDGImportedBuffer ClusterAABBBuffer;
-	FRDGImportedBuffer GroupAABBBuffer;
-};
-
 FRDGHairStrandsCullingData ImportCullingData(FRDGBuilder& GraphBuilder, FHairGroupPublicData* In)
 {
 	FRDGHairStrandsCullingData Out;
@@ -535,7 +533,7 @@ IMPLEMENT_GLOBAL_SHADER(FHairInterpolationCS, "/Engine/Private/HairStrands/HairS
 // 
 // This functions is in charge of deforming the rendering strands (OutRenderPositionBuffer) based on the guides
 // The guides are either the result of a simulation deformation, or a global interpolation deformation (aka RBF deformation), or both
-static void AddHairStrandsInterpolationPass(
+void AddHairStrandsInterpolationPass(
 	FRDGBuilder& GraphBuilder,
 	FGlobalShaderMap* ShaderMap,
 	const FShaderPrintData* ShaderPrintData,
@@ -563,7 +561,7 @@ static void AddHairStrandsInterpolationPass(
 	const FRDGBufferSRVRef& SimRootPointIndexBuffer,
 	const FRDGBufferSRVRef& SimPointToCurveBuffer,
 	const FRDGBufferSRVRef& RenDeformerPositionBuffer,
-	FRDGImportedBuffer& OutRenPositionBuffer,
+	FRDGBufferUAVRef& OutRenPositionBuffer,
 	const FHairStrandsDeformedRootResource::FLOD::EFrameType DeformedFrame)
 {
 	FHairInterpolationCS::FParameters* Parameters = GraphBuilder.AllocParameters<FHairInterpolationCS::FParameters>();
@@ -571,7 +569,7 @@ static void AddHairStrandsInterpolationPass(
 	Parameters->SimRestPosePositionBuffer = SimRestPosePositionBuffer;
 	Parameters->SimDeformedPositionBuffer = SimDeformedPositionBuffer;
 	Parameters->InterpolationBuffer = InterpolationBuffer;
-	Parameters->OutRenDeformedPositionBuffer = OutRenPositionBuffer.UAV;
+	Parameters->OutRenDeformedPositionBuffer = OutRenPositionBuffer;
 
 	Parameters->VertexCount = VertexCount;
 	Parameters->InRenHairPositionOffset = (FVector3f)InRenHairWorldOffset;
@@ -697,8 +695,6 @@ static void AddHairStrandsInterpolationPass(
 			Parameters,
 			DispatchCount);
 	}
-
-	GraphBuilder.SetBufferAccessFinal(OutRenPositionBuffer.Buffer, ERHIAccess::SRVMask);
 }
 
 ///////////////////////////////////////////////////////////////////////////////////////////////////
@@ -732,14 +728,7 @@ public:
 
 IMPLEMENT_GLOBAL_SHADER(FHairPatchAttributeCS, "/Engine/Private/HairStrands/HairStrandsInterpolation.usf", "MainCS", SF_Compute);
 
-enum class EHairPatchAttribute : uint8
-{
-	None,
-	GuideInflucence,
-	ClusterInfluence
-};
-
-static void AddPatchAttributePass(
+void AddPatchAttributePass(
 	FRDGBuilder& GraphBuilder,
 	FGlobalShaderMap* ShaderMap,
 	const uint32 CurveCount,
@@ -829,13 +818,7 @@ public:
 
 IMPLEMENT_GLOBAL_SHADER(FHairClusterAABBCS, "/Engine/Private/HairStrands/HairStrandsClusterCulling.usf", "ClusterAABBEvaluationCS", SF_Compute);
 
-enum class EHairAABBUpdateType
-{
-	UpdateClusterAABB,
-	UpdateGroupAABB
-};
-
-static void AddHairClusterAABBPass(
+void AddHairClusterAABBPass(
 	FRDGBuilder& GraphBuilder,
 	FGlobalShaderMap* ShaderMap,
 	const uint32 ActiveCurveCount,
@@ -1094,13 +1077,13 @@ void AddHairTangentPass(
 	uint32 PointCount,
 	FHairGroupPublicData* HairGroupPublicData,
 	FRDGBufferSRVRef PositionBuffer,
-	FRDGImportedBuffer OutTangentBuffer)
+	FRDGBufferUAVRef OutTangentBuffer)
 {
 	const bool bCullingEnable = IsHairStrandContinuousDecimationReorderingEnabled() ? false : (HairGroupPublicData && HairGroupPublicData->GetCullingResultAvailable());
 
 	FHairTangentCS::FParameters* Parameters = GraphBuilder.AllocParameters<FHairTangentCS::FParameters>();
 	Parameters->PositionBuffer = PositionBuffer;
-	Parameters->OutputTangentBuffer = OutTangentBuffer.UAV;
+	Parameters->OutputTangentBuffer = OutTangentBuffer;
 	Parameters->PointCount = PointCount;
 	Parameters->HairStrandsVF_bCullingEnable = bCullingEnable ? 1 : 0;
 
@@ -1185,12 +1168,12 @@ BEGIN_SHADER_PARAMETER_STRUCT(FStrandsResourceBLASParameters, )
 END_SHADER_PARAMETER_STRUCT()
 
 BEGIN_SHADER_PARAMETER_STRUCT(FCardsOrMeshesResourceBLASParameters, )
-RDG_BUFFER_ACCESS(PositionBuffer, ERHIAccess::SRVCompute)
+	RDG_BUFFER_ACCESS(PositionBuffer, ERHIAccess::SRVCompute)
 END_SHADER_PARAMETER_STRUCT()
 
 IMPLEMENT_GLOBAL_SHADER(FHairRaytracingGeometryCS, "/Engine/Private/HairStrands/HairStrandsRaytracingGeometry.usf", "MainCS", SF_Compute);
 
-static void AddGenerateRaytracingGeometryPass(
+void AddGenerateRaytracingGeometryPass(
 	FRDGBuilder& GraphBuilder,
 	FGlobalShaderMap* ShaderMap,
 	const FShaderPrintData* ShaderPrintData,
@@ -1275,7 +1258,7 @@ public:
 
 IMPLEMENT_GLOBAL_SHADER(FClearClusterAABBCS, "/Engine/Private/HairStrands/HairStrandsClusterCulling.usf", "MainClearClusterAABBCS", SF_Compute);
 
-static void AddClearClusterAABBPass(
+void AddClearClusterAABBPass(
 	FRDGBuilder& GraphBuilder,
 	FGlobalShaderMap* ShaderMap,
 	const EHairAABBUpdateType UpdateType,
@@ -1387,7 +1370,7 @@ static void BuildHairAccelerationStructure_Strands(
 	OutRayTracingGeometry->CreateRayTracingGeometry(RHICmdList, ERTAccelerationStructureBuildPriority::Immediate);
 }
 
-static void BuildHairAccelerationStructure_Cards(FRHICommandList& RHICmdList, 	
+static void BuildHairAccelerationStructure_Cards(FRHICommandList& RHICmdList,
 	FHairCardsRestResource* RestResource,
 	FHairCardsDeformedResource* DeformedResource,
 	FRayTracingGeometry* OutRayTracingGeometry,
@@ -1473,6 +1456,62 @@ static void BuildHairAccelerationStructure_Meshes(FRHICommandList& RHICmdList,
 	OutRayTracingGeometry->SetInitializer(Initializer);
 	OutRayTracingGeometry->CreateRayTracingGeometry(RHICmdList, ERTAccelerationStructureBuildPriority::Immediate);
 }
+
+void AddBuildStrandsAccelerationStructurePass(
+	FRDGBuilder& GraphBuilder,
+	FHairGroupInstance* Instance,
+	uint32 ProceduralSplits,
+	bool bNeedUpdate,
+	FRDGBufferRef Raytracing_PositionBuffer,
+	FRDGBufferRef Raytracing_IndexBuffer)
+{
+	const uint32 HairLODIndex = Instance->HairGroupPublicData->GetIntLODIndex();
+	const bool bProceduralPrimitive = Instance->Strands.RenRaytracingResource->bProceduralPrimitive;
+
+	FStrandsResourceBLASParameters* Parameters = GraphBuilder.AllocParameters < FStrandsResourceBLASParameters > ();
+	Parameters->PositionBuffer = Raytracing_PositionBuffer;
+	Parameters->IndexBuffer = Raytracing_IndexBuffer;
+	
+	GraphBuilder.AddPass(
+		RDG_EVENT_NAME("HairStrands::UpdateBLAS(Strands)"),
+		Parameters,
+		ERDGPassFlags::NeverCull | ERDGPassFlags::Compute,
+		[Instance, HairLODIndex, bNeedUpdate, bProceduralPrimitive, ProceduralSplits](FRHICommandListImmediate& RHICmdList)
+		{
+			SCOPED_GPU_MASK(RHICmdList, FRHIGPUMask::All());
+	
+			const bool bLocalNeedBuild = !Instance->Strands.RenRaytracingResource->bIsRTGeometryInitialized || (Instance->Strands.CachedProceduralSplits != ProceduralSplits);
+			if (bLocalNeedBuild)
+			{
+				FBufferRHIRef PositionBuffer(Instance->Strands.RenRaytracingResource->PositionBuffer.Buffer->GetRHI());
+	
+				// no index buffer needed for procedural primitive
+				FBufferRHIRef IndexBuffer(bProceduralPrimitive ? nullptr : Instance->Strands.RenRaytracingResource->IndexBuffer.Buffer->GetRHI());
+				Instance->Strands.CachedProceduralSplits = ProceduralSplits;
+				BuildHairAccelerationStructure_Strands(
+					RHICmdList,
+					Instance->Strands.RenRaytracingResource->MaxVertexCount,
+					Instance->Strands.RenRaytracingResource->MaxIndexCount,
+					PositionBuffer,
+					IndexBuffer,
+					&Instance->Strands.RenRaytracingResource->RayTracingGeometry,
+					Instance->Debug.GroomAssetName,
+					Instance->Debug.MeshComponentName,
+					HairLODIndex,
+					bProceduralPrimitive,
+					ProceduralSplits
+				);
+			}
+			else if (bNeedUpdate)
+			{
+				// hair strands can move chaotically during simulation which can really tank ray tracing performance
+				// even though a rebuild is more expensive, we can still come out ahead overall, even with a single ray cast per pixel
+				UpdateHairAccelerationStructure(RHICmdList, &Instance->Strands.RenRaytracingResource->RayTracingGeometry, IsHairStrandsForceRebuildBVH() ? EAccelerationStructureBuildMode::Build : EAccelerationStructureBuildMode::Update);
+			}
+			Instance->Strands.RenRaytracingResource->bIsRTGeometryInitialized = true;
+		});
+}
+
 #endif // RHI_RAYTRACING
 
 static void ConvertHairStrandsVFParameters(
@@ -1492,7 +1531,7 @@ static void ConvertHairStrandsVFParameters(
 #define CONVERT_HAIRSSTRANDS_VF_PARAMETERS(OutName, ExternalBuffer) ConvertHairStrandsVFParameters(GraphBuilder, OutName, OutName##External, OutName##RHISRV, ExternalBuffer);
 
 // Compute/Update the hair strands description which will be used for rendering (VF) / voxelization & co.
-static FHairGroupPublicData::FVertexFactoryInput InternalComputeHairStrandsVertexInputData(FRDGBuilder* GraphBuilder, const FHairGroupInstance* Instance, EGroomViewMode ViewMode)
+FHairGroupPublicData::FVertexFactoryInput InternalComputeHairStrandsVertexInputData(FRDGBuilder* GraphBuilder, const FHairGroupInstance* Instance, EGroomViewMode ViewMode)
 {
 	FHairGroupPublicData::FVertexFactoryInput OutVFInput;
 	if (!Instance || Instance->GeometryType != EHairGeometryType::Strands)
@@ -1594,6 +1633,7 @@ static FHairGroupPublicData::FVertexFactoryInput InternalComputeHairStrandsVerte
 		OutVFInput.Strands.Common.CurveCount 		= Instance->HairGroupPublicData->GetActiveStrandsCurveCount();
 	}
 
+	OutVFInput.Strands.Common.RegisteredIndex = Instance->RegisteredIndex;
 	OutVFInput.Strands.Common.Radius = (GStrandHairWidth > 0 ? GStrandHairWidth : Instance->Strands.Modifier.HairWidth) * 0.5f;
 	OutVFInput.Strands.Common.RootScale = Instance->Strands.Modifier.HairRootScale;
 	OutVFInput.Strands.Common.TipScale = Instance->Strands.Modifier.HairTipScale;
@@ -1670,417 +1710,7 @@ void ComputeHairStrandsInterpolation(
 
 	FRDGExternalAccessQueue ExternalAccessQueue;
 
-	if (InstanceGeometryType == EHairGeometryType::Strands)
-	{
-		DECLARE_GPU_STAT(HairStrandsInterpolation);
-		RDG_EVENT_SCOPE(GraphBuilder, "HairInterpolation(Strands)");
-		RDG_GPU_STAT_SCOPE(GraphBuilder, HairStrandsInterpolation);
-
-		const EGroomCacheType ActiveGroomCacheType = GetHairInstanceCacheType(Instance);
-
-		// Debug mode:
-		// * None	: Display hair normally
-		// * Sim	: Show sim strands
-		// * Render : Show rendering strands with sim color influence
-		if (ViewMode == EGroomViewMode::SimHairStrands && Instance->Guides.DeformedResource)
-		{
-			// Disable culling when drawing only guides, as the culling output has been computed for the strands, not for the guides.
-			Instance->HairGroupPublicData->SetCullingResultAvailable(false);
-
-			AddHairTangentPass(
-				GraphBuilder,
-				ShaderMap,
-				Instance->Guides.RestResource->GetPointCount(),
-				Instance->HairGroupPublicData,
-				RegisterAsSRV(GraphBuilder, Instance->Guides.DeformedResource->GetBuffer(FHairStrandsDeformedResource::Current)),
-				Register(GraphBuilder, Instance->Guides.DeformedResource->TangentBuffer, ERDGImportedBufferFlags::CreateUAV));
-
-			Instance->HairGroupPublicData->VFInput = InternalComputeHairStrandsVertexInputData(&GraphBuilder, Instance, ViewMode);
-		}
-		else
-		{
-			const bool bNeedDeformation = Instance->Strands.DeformedResource != nullptr;
-			const EHairAABBUpdateType UpdateType = bNeedDeformation ? EHairAABBUpdateType::UpdateClusterAABB : EHairAABBUpdateType::UpdateGroupAABB;
-
-			// 1. Clear cluster AABBs (used optionally  for voxel allocation & for culling)
-			FRDGHairStrandsCullingData CullingData = ImportCullingData(GraphBuilder, Instance->HairGroupPublicData);
-			{
-				AddClearClusterAABBPass(
-					GraphBuilder,
-					ShaderMap,
-					UpdateType,
-					CullingData.ClusterCount,
-					CullingData.ClusterAABBBuffer,
-					CullingData.GroupAABBBuffer);
-			}
-
-			// "WITH_EDITOR && PatchMode == EHairPatchAttribute::GuideInflucence" is a special path when visualizing groom within the groom editor, so that we can diplay guides even when there is no simulation or global interpolation enabled
-			EHairPatchAttribute PatchMode = EHairPatchAttribute::None;
-			if (ViewMode == EGroomViewMode::RenderHairStrands && bNeedDeformation)				{ PatchMode = EHairPatchAttribute::GuideInflucence; }
-			if (ViewMode == EGroomViewMode::Cluster || ViewMode == EGroomViewMode::ClusterAABB)	{ PatchMode = EHairPatchAttribute::ClusterInfluence; }
-
-			// 2. Deform hair if needed (e.g.: skinning, simulation, RBF) and recompute tangent
-			FRDGBufferSRVRef Strands_PositionSRV = nullptr;
-			FRDGBufferSRVRef Strands_PositionOffsetSRV = nullptr;
-			FRDGBufferSRVRef Strands_PrevPositionOffsetSRV = nullptr;
-			FRDGBufferSRVRef Strands_TangentSRV = nullptr;
-
-			const bool bValidGuide		= bNeedDeformation && (Instance->Guides.bIsSimulationEnable || Instance->Guides.bHasGlobalInterpolation || Instance->Guides.bIsDeformationEnable || Instance->Guides.bIsSimulationCacheEnable);// || (WITH_EDITOR && PatchMode == EHairPatchAttribute::GuideInflucence);
-			const bool bUseSingleGuide	= bNeedDeformation && bValidGuide && Instance->Strands.InterpolationResource->UseSingleGuide();
-			const bool bHasSkinning		= bNeedDeformation && Instance->BindingType == EHairBindingType::Skinning;
-
-			const uint32 ActivePointCount = Instance->HairGroupPublicData->GetActiveStrandsPointCount();
-			const uint32 ActiveCurveCount = Instance->HairGroupPublicData->GetActiveStrandsCurveCount();
-
-			if (bNeedDeformation)
-			{
-				FRDGImportedBuffer Strands_DeformedPosition = Register(GraphBuilder, Instance->Strands.DeformedResource->GetBuffer(FHairStrandsDeformedResource::Current), ERDGImportedBufferFlags::CreateViews);
-				FRDGImportedBuffer Strands_DeformedPrevPosition = Register(GraphBuilder, Instance->Strands.DeformedResource->GetBuffer(FHairStrandsDeformedResource::Previous), ERDGImportedBufferFlags::CreateViews);
-				FRDGImportedBuffer Strands_DeformedTangent = Register(GraphBuilder, Instance->Strands.DeformedResource->TangentBuffer, ERDGImportedBufferFlags::CreateViews);
-				FRDGBufferSRVRef Strands_DeformerPositionSRV = Register(GraphBuilder, Instance->Strands.DeformedResource->DeformerBuffer, ERDGImportedBufferFlags::CreateSRV).SRV;
-
-				// Trach on which view the position has been update, so that we can ensure motion vector are coherent
-				Instance->Strands.DeformedResource->GetUniqueViewID(FHairStrandsDeformedResource::Current) = ViewUniqueID;
-
-				Strands_PositionSRV = Strands_DeformedPosition.SRV;
-				Strands_PositionOffsetSRV = RegisterAsSRV(GraphBuilder, Instance->Strands.DeformedResource->GetPositionOffsetBuffer(FHairStrandsDeformedResource::EFrameType::Current));
-				Strands_TangentSRV = Strands_DeformedTangent.SRV;
-
-				// 2.0 If using a guide cache, update the render strands offset point
-				if (ActiveGroomCacheType == EGroomCacheType::Guides)
-				{
-					// Apply the same offset to the render strands for proper rendering
-					const FVector OffsetPosition = Instance->Guides.DeformedResource->GetPositionOffset(FHairStrandsDeformedResource::EFrameType::Current);
-					Instance->Strands.DeformedResource->SetPositionOffset(FHairStrandsDeformedResource::EFrameType::Current, OffsetPosition);
-
-					AddHairStrandUpdatePositionOffsetPass(
-						GraphBuilder,
-						ShaderMap,
-						EHairPositionUpdateType::Strands,
-						Instance->RegisteredIndex,
-						0,
-						MeshLODIndex,
-						Instance->Strands.DeformedRootResource,
-						Instance->Strands.DeformedResource);
-				}
-
-				// 2.1 Compute deformation position based on simulation/skinning/RBF
-				if (ActiveGroomCacheType != EGroomCacheType::Strands)
-				{
-					// 2.1.1 Current Position
-					AddHairStrandsInterpolationPass(
-						GraphBuilder,
-						ShaderMap,
-						ShaderPrintData,
-						Instance,
-						ActivePointCount,
-						MeshLODIndex,
-						Instance->Strands.Modifier.HairLengthScale,
-						Instance->Strands.HairInterpolationType,
-						InstanceGeometryType,
-						CullingData,
-						Instance->Strands.RestResource->GetPositionOffset(),
-						bValidGuide ? Instance->Guides.RestResource->GetPositionOffset() : FVector::ZeroVector,
-						Strands_PositionOffsetSRV,
-						bValidGuide ? RegisterAsSRV(GraphBuilder, Instance->Guides.DeformedResource->GetPositionOffsetBuffer(FHairStrandsDeformedResource::Current)) : nullptr,
-						bHasSkinning ? Instance->Strands.RestRootResource : nullptr,
-						bHasSkinning && bValidGuide ? Instance->Guides.RestRootResource : nullptr,
-						bHasSkinning ? Instance->Strands.DeformedRootResource : nullptr,
-						bHasSkinning && bValidGuide ? Instance->Guides.DeformedRootResource : nullptr,
-						RegisterAsSRV(GraphBuilder, Instance->Strands.RestResource->PositionBuffer),
-						RegisterAsSRV(GraphBuilder, Instance->Strands.RestResource->PointToCurveBuffer),
-						bUseSingleGuide,
-						bValidGuide ? RegisterAsSRV(GraphBuilder, Instance->Strands.InterpolationResource->InterpolationBuffer) : nullptr,
-						bValidGuide ? RegisterAsSRV(GraphBuilder, Instance->Guides.RestResource->PositionBuffer) : nullptr,
-						bValidGuide ? RegisterAsSRV(GraphBuilder, Instance->Guides.DeformedResource->GetBuffer(FHairStrandsDeformedResource::Current)) : nullptr,
-						bValidGuide ? RegisterAsSRV(GraphBuilder, Instance->Strands.InterpolationResource->SimRootPointIndexBuffer) : nullptr,
-						bValidGuide ? RegisterAsSRV(GraphBuilder, Instance->Guides.RestResource->PointToCurveBuffer) : nullptr,
-						Strands_DeformerPositionSRV,
-						Strands_DeformedPosition,
-						FHairStrandsDeformedRootResource::FLOD::Current);
-					
-					// 2.1.2 Previous Position
-					// * Transfer prev. position
-					// * Or recompute interpolation pass for previous positions if needed
-					const bool bRecomputePrevPosition = IsHairVisibilityComputeRasterContinuousLODEnabled();
-					const bool bTransferPrevPosition = 
-						!bRecomputePrevPosition &&
-						Instance->HairGroupPublicData->VFInput.bHasLODSwitch && 
-						GHairStrandsTransferPositionOnLODChange > 0;
-					if (bTransferPrevPosition)
-					{
-						AddTransferPositionPass(GraphBuilder, ShaderMap, ActivePointCount, Strands_DeformedPosition.SRV, Strands_DeformedPrevPosition.UAV);
-						GraphBuilder.SetBufferAccessFinal(Strands_DeformedPrevPosition.Buffer, ERHIAccess::SRVMask);
-					}
-					else if (bRecomputePrevPosition)
-					{
-						Strands_PrevPositionOffsetSRV = RegisterAsSRV(GraphBuilder, Instance->Strands.DeformedResource->GetPositionOffsetBuffer(FHairStrandsDeformedResource::EFrameType::Previous));
-
-						AddHairStrandsInterpolationPass(
-							GraphBuilder,
-							ShaderMap,
-							ShaderPrintData,
-							Instance,
-							ActivePointCount,
-							MeshLODIndex,
-							Instance->Strands.Modifier.HairLengthScale,
-							Instance->Strands.HairInterpolationType,
-							InstanceGeometryType,
-							CullingData,
-							Instance->Strands.RestResource->GetPositionOffset(),
-							bValidGuide ? Instance->Guides.RestResource->GetPositionOffset() : FVector::ZeroVector,
-							Strands_PrevPositionOffsetSRV,
-							bValidGuide ? RegisterAsSRV(GraphBuilder, Instance->Guides.DeformedResource->GetPositionOffsetBuffer(FHairStrandsDeformedResource::Previous)) : nullptr,
-							bHasSkinning ? Instance->Strands.RestRootResource : nullptr,
-							bHasSkinning && bValidGuide ? Instance->Guides.RestRootResource : nullptr,
-							bHasSkinning ? Instance->Strands.DeformedRootResource : nullptr,
-							bHasSkinning && bValidGuide ? Instance->Guides.DeformedRootResource : nullptr,
-							RegisterAsSRV(GraphBuilder, Instance->Strands.RestResource->PositionBuffer),
-							RegisterAsSRV(GraphBuilder, Instance->Strands.RestResource->PointToCurveBuffer),
-							bUseSingleGuide,
-							bValidGuide ? RegisterAsSRV(GraphBuilder, Instance->Strands.InterpolationResource->InterpolationBuffer) : nullptr,
-							bValidGuide ? RegisterAsSRV(GraphBuilder, Instance->Guides.RestResource->PositionBuffer) : nullptr,
-							bValidGuide ? RegisterAsSRV(GraphBuilder, Instance->Guides.DeformedResource->GetBuffer(FHairStrandsDeformedResource::Previous)) : nullptr,
-							bValidGuide ? RegisterAsSRV(GraphBuilder, Instance->Strands.InterpolationResource->SimRootPointIndexBuffer) : nullptr,
-							bValidGuide ? RegisterAsSRV(GraphBuilder, Instance->Guides.RestResource->PointToCurveBuffer) : nullptr,
-							Strands_DeformerPositionSRV,
-							Strands_DeformedPrevPosition,
-							FHairStrandsDeformedRootResource::FLOD::Previous);
-					}
-				}
-				else if (ActiveGroomCacheType == EGroomCacheType::Strands)
-				{
-					FScopeLock Lock(Instance->Debug.GroomCacheBuffers->GetCriticalSection());
-
-					FGroomCacheGroupData* GroomCacheData0 = const_cast<FGroomCacheGroupData*>(&Instance->Debug.GroomCacheBuffers->GetCurrentFrameBuffer().GroupsData[Instance->Debug.GroupIndex]);
-					FGroomCacheGroupData* GroomCacheData1 = const_cast<FGroomCacheGroupData*>(&Instance->Debug.GroomCacheBuffers->GetNextFrameBuffer().GroupsData[Instance->Debug.GroupIndex]);
-					const float InterpolationFactor = Instance->Debug.GroomCacheBuffers->GetInterpolationFactor();
-
-					FGroomCacheResources CacheResources0 = CreateGroomCacheBuffer(GraphBuilder, GroomCacheData0->VertexData);
-					FGroomCacheResources CacheResources1 = CreateGroomCacheBuffer(GraphBuilder, GroomCacheData1->VertexData);
-
-					// Update position offset
-					{
-						const FVector OffsetPosition = FMath::Lerp(GroomCacheData0->BoundingBox.GetCenter(), GroomCacheData1->BoundingBox.GetCenter(), InterpolationFactor);
-						Instance->Strands.DeformedResource->SetPositionOffset(FHairStrandsDeformedResource::EFrameType::Current, OffsetPosition);
-					}
-
-					// Update max radius
-					if (GroomCacheData0->VertexData.PointsRadius.Num() > 0)
-					{
-						Instance->Strands.Modifier.HairWidth = FMath::Lerp(GroomCacheData0->StrandData.MaxRadius, GroomCacheData1->StrandData.MaxRadius, InterpolationFactor) * 2.0f;
-					}
-
-					AddHairStrandUpdatePositionOffsetPass(
-						GraphBuilder,
-						ShaderMap,
-						EHairPositionUpdateType::Strands,
-						Instance->RegisteredIndex,
-						0,
-						MeshLODIndex,
-						Instance->Strands.DeformedRootResource,
-						Instance->Strands.DeformedResource);
-
-					// Pass to upload GroomCache strands positions
-					AddGroomCacheUpdatePass(
-						GraphBuilder,
-						ShaderMap,
-						ActivePointCount,
-						InterpolationFactor,
-						CacheResources0,
-						CacheResources1,
-						RegisterAsSRV(GraphBuilder, Instance->Strands.RestResource->PositionBuffer),
-						RegisterAsSRV(GraphBuilder, Instance->Strands.DeformedResource->GetPositionOffsetBuffer(FHairStrandsDeformedResource::EFrameType::Current)),
-						Strands_DeformedPosition.UAV);
-				}
-
-				// 2.2 Update tangent data based on the deformed positions
-				AddHairTangentPass(
-					GraphBuilder,
-					ShaderMap,
-					ActivePointCount,
-					Instance->HairGroupPublicData,
-					Strands_DeformedPosition.SRV,
-					Strands_DeformedTangent);
-			}
-			else
-			{
-				Strands_PositionSRV = RegisterAsSRV(GraphBuilder, Instance->Strands.RestResource->PositionBuffer);
-				Strands_PositionOffsetSRV = RegisterAsSRV(GraphBuilder, Instance->Strands.RestResource->PositionOffsetBuffer);
-				// Generated the static tangent if they haven't been generated yet
-				Instance->Strands.RestResource->GetTangentBuffer(GraphBuilder, ShaderMap, ActivePointCount, ActiveCurveCount);
-				Strands_TangentSRV = RegisterAsSRV(GraphBuilder, Instance->Strands.RestResource->TangentBuffer);
-			}
-
-			// 2.1 Patch attribute for debug visualization (guide influence or clusters visualization)
-			if (PatchMode != EHairPatchAttribute::None)
-			{
-				// Create an debug buffer for storing cluster visualization data. This is only used for debug purpose, hence only enable in editor build.
-				// Special case for debug mode were the attribute buffer is patch with some custom data to show hair properties (strands belonging to the same cluster, ...)
-				if (Instance->Strands.DebugCurveAttributeBuffer.Buffer == nullptr)
-				{
-					CreateHairStrandsDebugAttributeBuffer(GraphBuilder, &Instance->Strands.DebugCurveAttributeBuffer, Instance->Strands.Data->GetCurveAttributeSizeInBytes(), FName(Instance->Debug.MeshComponentName));
-				}
-				FRDGImportedBuffer OutRenCurveAttributeBuffer = Register(GraphBuilder, Instance->Strands.DebugCurveAttributeBuffer, ERDGImportedBufferFlags::CreateUAV);
-
-				check(Instance->Strands.Data);
-				AddPatchAttributePass(
-					GraphBuilder,
-					ShaderMap,
-					ActiveCurveCount,
-					PatchMode,
-					bValidGuide,
-					bUseSingleGuide,
-					*Instance->Strands.Data,
-					Register(GraphBuilder, Instance->Strands.RestResource->CurveAttributeBuffer, ERDGImportedBufferFlags::CreateSRV).Buffer,
-					RegisterAsSRV(GraphBuilder, Instance->Strands.RestResource->CurveBuffer),
-					RegisterAsSRV(GraphBuilder, Instance->Strands.ClusterResource->CurveToClusterIdBuffer),
-					bValidGuide ? RegisterAsSRV(GraphBuilder, Instance->Strands.InterpolationResource->InterpolationBuffer) : nullptr,
-					OutRenCurveAttributeBuffer);
-			}
-
-			// 2.2 Update the VF input with the update resources
-			Instance->HairGroupPublicData->VFInput = InternalComputeHairStrandsVertexInputData(&GraphBuilder, Instance, ViewMode);
-
-			// 3. Compute cluster AABBs (used for LODing and voxelization)
-			{
-				// Optim: If an instance does not voxelize it's data, then there is no need for having valid AABB
-				bool bNeedGPUAABB = Instance->Strands.Modifier.bSupportVoxelization && Instance->bCastShadow;
-
-				FHairStrandClusterData::FHairGroup* HairGroupCluster = nullptr;
-				if (CullingData.bCullingResultAvailable)
-				{
-					// Sanity check
-					check(Instance->Strands.bCullingEnable);
-					HairGroupCluster = &InClusterData->HairGroups[Instance->HairGroupPublicData->ClusterDataIndex];
-					if (!HairGroupCluster->bVisible)
-					{
-						bNeedGPUAABB = false;
-					}
-				}
-
-				Instance->HairGroupPublicData->SetClusterAABBValid(false);
-				Instance->HairGroupPublicData->SetGroupAABBValid(false);
-
-
-				if (bNeedGPUAABB)
-				{
-					FRDGImportedBuffer Strands_CulledVertexCount = Register(GraphBuilder, Instance->HairGroupPublicData->GetDrawIndirectRasterComputeBuffer(), ERDGImportedBufferFlags::CreateSRV);
-					AddHairClusterAABBPass(
-						GraphBuilder,
-						ShaderMap,
-						ActiveCurveCount,
-						TranslatedWorldOffset,
-						ShaderPrintData,
-						UpdateType,
-						Instance,
-						Strands_PositionOffsetSRV,
-						HairGroupCluster,
-						CullingData,
-						Strands_PositionSRV,
-						Strands_CulledVertexCount.SRV);
-
-					Instance->HairGroupPublicData->SetClusterAABBValid(UpdateType == EHairAABBUpdateType::UpdateClusterAABB);
-					Instance->HairGroupPublicData->SetGroupAABBValid(true);
-				}
-			}
-
-			// 4. Update raytracing geometry (update only if the view mask and the RT geometry mask match)
-			#if RHI_RAYTRACING
-			if (Instance->Strands.RenRaytracingResource && (Instance->Strands.ViewRayTracingMask & ViewRayTracingMask) != 0)
-			{
-				const float CLODScale = Instance->HairGroupPublicData->ContinuousLODCoverageScale;
-				const float HairRadiusRT	= Instance->HairGroupPublicData->VFInput.Strands.Common.RaytracingRadiusScale * Instance->HairGroupPublicData->VFInput.Strands.Common.Radius * CLODScale;
-				const float HairRootScaleRT = Instance->HairGroupPublicData->VFInput.Strands.Common.RootScale;
-				const float HairTipScaleRT	= Instance->HairGroupPublicData->VFInput.Strands.Common.TipScale;
-				const int ProceduralSplits = GetHairRaytracingProceduralSplits();
-
-				const bool bNeedUpdate = Instance->Strands.DeformedResource != nullptr ||
-					Instance->Strands.CachedHairScaledRadius != HairRadiusRT ||
-					Instance->Strands.CachedHairRootScale != HairRootScaleRT ||
-					Instance->Strands.CachedHairTipScale != HairTipScaleRT ||
-					Instance->Strands.CachedProceduralSplits != ProceduralSplits;
-				const bool bNeedBuild = !Instance->Strands.RenRaytracingResource->bIsRTGeometryInitialized;
-				if (bNeedBuild || bNeedUpdate)
-				{
-					const uint32 HairLODIndex = Instance->HairGroupPublicData->GetIntLODIndex();
-					const bool bProceduralPrimitive = Instance->Strands.RenRaytracingResource->bProceduralPrimitive;
-
-					FRDGImportedBuffer Raytracing_PositionBuffer = Register(GraphBuilder, Instance->Strands.RenRaytracingResource->PositionBuffer, ERDGImportedBufferFlags::CreateViews);
-					FRDGImportedBuffer Raytracing_IndexBuffer = Register(GraphBuilder, Instance->Strands.RenRaytracingResource->IndexBuffer, ERDGImportedBufferFlags::CreateViews);
-					AddGenerateRaytracingGeometryPass(
-						GraphBuilder,
-						ShaderMap,
-						ShaderPrintData,
-						Instance->RegisteredIndex,
-						ActivePointCount,
-						bProceduralPrimitive,
-						ProceduralSplits,
-						HairRadiusRT,
-						HairRootScaleRT,
-						HairTipScaleRT,
-						Strands_PositionOffsetSRV,
-						CullingData,
-						Strands_PositionSRV,
-						Strands_TangentSRV,
-						Raytracing_PositionBuffer.UAV,
-						Raytracing_IndexBuffer.UAV);
-
-					Instance->Strands.CachedHairScaledRadius= HairRadiusRT;
-					Instance->Strands.CachedHairRootScale	= HairRootScaleRT;
-					Instance->Strands.CachedHairTipScale	= HairTipScaleRT;
-
-					FStrandsResourceBLASParameters* Parameters = GraphBuilder.AllocParameters<FStrandsResourceBLASParameters>();
-					Parameters->PositionBuffer = Raytracing_PositionBuffer.Buffer;
-					Parameters->IndexBuffer = Raytracing_IndexBuffer.Buffer;
-
-					GraphBuilder.AddPass(
-						RDG_EVENT_NAME("HairStrands::UpdateBLAS(Strands)"),
-						Parameters,
-						ERDGPassFlags::NeverCull | ERDGPassFlags::Compute,
-					[Instance, HairLODIndex, bNeedUpdate, bProceduralPrimitive, ProceduralSplits](FRHICommandListImmediate& RHICmdList)
-					{
-						SCOPED_GPU_MASK(RHICmdList, FRHIGPUMask::All());
-
-						const bool bLocalNeedBuild = !Instance->Strands.RenRaytracingResource->bIsRTGeometryInitialized || Instance->Strands.CachedProceduralSplits != ProceduralSplits;
-						if (bLocalNeedBuild)
-						{
-							FBufferRHIRef PositionBuffer(Instance->Strands.RenRaytracingResource->PositionBuffer.Buffer->GetRHI());
-
-							// no index buffer needed for procedural primitive
-							FBufferRHIRef IndexBuffer(bProceduralPrimitive ? nullptr : Instance->Strands.RenRaytracingResource->IndexBuffer.Buffer->GetRHI());
-							Instance->Strands.CachedProceduralSplits = ProceduralSplits;
-							BuildHairAccelerationStructure_Strands(RHICmdList,
-								Instance->Strands.RenRaytracingResource->MaxVertexCount,
-								Instance->Strands.RenRaytracingResource->MaxIndexCount,
-								PositionBuffer,
-								IndexBuffer,
-								&Instance->Strands.RenRaytracingResource->RayTracingGeometry,
-								Instance->Debug.GroomAssetName,
-								Instance->Debug.MeshComponentName,
-								HairLODIndex,
-								bProceduralPrimitive,
-								ProceduralSplits
-							);
-						}
-						else if (bNeedUpdate)
-						{
-							// hair strands can move chaotically during simulation which can really tank ray tracing performance
-							// even though a rebuild is more expensive, we can still come out ahead overall, even with a single ray cast per pixel
-							UpdateHairAccelerationStructure(RHICmdList, &Instance->Strands.RenRaytracingResource->RayTracingGeometry, GHairStrands_Raytracing_ForceRebuildBVH > 0? EAccelerationStructureBuildMode::Build : EAccelerationStructureBuildMode::Update);
-						}
-						Instance->Strands.RenRaytracingResource->bIsRTGeometryInitialized = true;
-					});
-				}
-			}
-			#endif
-		}
-
-		// Sanity check
-		check(Instance->HairGroupPublicData->VFInput.Strands.PositionBuffer.Buffer);
-		Instance->Strands.UniformBuffer.UpdateUniformBufferImmediate(GraphBuilder.RHICmdList, Instance->GetHairStandsUniformShaderParameters(ViewMode));
-	}
-	else if (InstanceGeometryType == EHairGeometryType::Cards)
+	if (InstanceGeometryType == EHairGeometryType::Cards)
 	{	
 		DECLARE_GPU_STAT(HairCardsInterpolation);
 		RDG_EVENT_SCOPE(GraphBuilder, "HairInterpolation(Cards)");
@@ -2108,7 +1738,7 @@ void ComputeHairStrandsInterpolation(
 				// 1. Cards are deformed based on guides motion (simulation or RBF applied on guides)
 				if (CardsSimulationType == EHairCardsSimulationType::Guide)
 				{
-					FRDGImportedBuffer Guides_DeformedPositionBuffer = Register(GraphBuilder, LOD.Guides.DeformedResource->GetBuffer(FHairStrandsDeformedResource::Current), ERDGImportedBufferFlags::CreateViews);
+					FRDGBufferUAVRef Guides_DeformedPositionUAV = RegisterAsUAV(GraphBuilder, LOD.Guides.DeformedResource->GetBuffer(FHairStrandsDeformedResource::Current));
 
 					const bool bUseSingleGuide = LOD.Guides.InterpolationResource->UseSingleGuide();
 
@@ -2141,7 +1771,7 @@ void ComputeHairStrandsInterpolation(
 						RegisterAsSRV(GraphBuilder, LOD.Guides.InterpolationResource->SimRootPointIndexBuffer),
 						bValidGuide ? RegisterAsSRV(GraphBuilder, Instance->Guides.RestResource->PointToCurveBuffer) : nullptr,
 						nullptr,
-						Guides_DeformedPositionBuffer,
+						Guides_DeformedPositionUAV,
 						FHairStrandsDeformedRootResource::FLOD::Current); // <- this should be optional
 
 					AddHairCardsDeformationPass(
diff --git a/Engine/Plugins/Runtime/HairStrands/Source/HairStrandsCore/Private/HairStrandsVertexFactory.cpp b/Engine/Plugins/Runtime/HairStrands/Source/HairStrandsCore/Private/HairStrandsVertexFactory.cpp
--- a/Engine/Plugins/Runtime/HairStrands/Source/HairStrandsCore/Private/HairStrandsVertexFactory.cpp
+++ b/Engine/Plugins/Runtime/HairStrands/Source/HairStrandsCore/Private/HairStrandsVertexFactory.cpp
@@ -88,7 +88,7 @@ TGlobalResource<FDummyCulledDispatchVertexIdsBuffer> GDummyCulledDispatchVertexI
 /////////////////////////////////////////////////////////////////////////////////////////
 
 FHairGroupPublicData::FVertexFactoryInput ComputeHairStrandsVertexInputData(const FHairGroupInstance* Instance, EGroomViewMode ViewMode);
-int GetHairRaytracingProceduralSplits();
+uint32 GetHairRaytracingProceduralSplits();
 FHairStrandsVertexFactoryUniformShaderParameters FHairGroupInstance::GetHairStandsUniformShaderParameters(EGroomViewMode ViewMode) const
 {
 	const FHairGroupPublicData::FVertexFactoryInput VFInput = ComputeHairStrandsVertexInputData(this, ViewMode);
diff --git a/Engine/Plugins/Runtime/HairStrands/Source/HairStrandsCore/Public/GroomInstance.h b/Engine/Plugins/Runtime/HairStrands/Source/HairStrandsCore/Public/GroomInstance.h
--- a/Engine/Plugins/Runtime/HairStrands/Source/HairStrandsCore/Public/GroomInstance.h
+++ b/Engine/Plugins/Runtime/HairStrands/Source/HairStrandsCore/Public/GroomInstance.h
@@ -138,7 +138,7 @@ struct HAIRSTRANDSCORE_API FHairGroupInstance : public FHairStrandsInstance
 		float CachedHairScaledRadius = 0;
 		float CachedHairRootScale = 0;
 		float CachedHairTipScale = 0;
-		int   CachedProceduralSplits = 0;
+		uint32 CachedProceduralSplits = 0;
 		#endif
 
 		FRDGExternalBuffer DebugCurveAttributeBuffer;
diff --git a/Engine/Plugins/Runtime/HairStrands/Source/HairStrandsCore/Public/HairStrandsInterpolation.h b/Engine/Plugins/Runtime/HairStrands/Source/HairStrandsCore/Public/HairStrandsInterpolation.h
--- a/Engine/Plugins/Runtime/HairStrands/Source/HairStrandsCore/Public/HairStrandsInterpolation.h
+++ b/Engine/Plugins/Runtime/HairStrands/Source/HairStrandsCore/Public/HairStrandsInterpolation.h
@@ -7,12 +7,14 @@
 #include "Math/MathFwd.h"
 #include "RenderGraphFwd.h"
 
+#include "HairStrandsInterface.h"
+#include "GroomResources.h"
+
 #if UE_ENABLE_INCLUDE_ORDER_DEPRECATED_IN_5_3
 #include "CoreMinimal.h"
 #include "UObject/ObjectMacros.h"
 #include "Engine/EngineTypes.h"
 #include "RenderGraphResources.h"
-#include "HairStrandsInterface.h"
 #include "Shader.h"
 #include "GroomDesc.h"
 #endif
@@ -86,3 +88,128 @@ void AddGroomCacheUpdatePass(
 	FRDGBufferUAVRef OutBuffer);
 
 HAIRSTRANDSCORE_API void ComputeInterpolationWeights(UGroomBindingAsset* BindingAsset, FSkeletalMeshRenderData* TargetRenderData, TArray<FRWBuffer>& TransferedPositions);
+
+struct FRDGHairStrandsCullingData
+{
+	bool bCullingResultAvailable = false;
+	FRDGImportedBuffer HairStrandsVF_CullingIndirectBuffer;
+	FRDGImportedBuffer HairStrandsVF_CullingIndexBuffer;
+	FRDGImportedBuffer HairStrandsVF_CullingRadiusScaleBuffer;
+
+	uint32 ClusterCount = 0;
+	FRDGImportedBuffer ClusterAABBBuffer;
+	FRDGImportedBuffer GroupAABBBuffer;
+};
+
+FRDGHairStrandsCullingData ImportCullingData(FRDGBuilder& GraphBuilder, FHairGroupPublicData* In);
+
+enum class EHairAABBUpdateType
+{
+	UpdateClusterAABB,
+	UpdateGroupAABB
+};
+
+void AddClearClusterAABBPass(
+	FRDGBuilder& GraphBuilder,
+	FGlobalShaderMap* ShaderMap,
+	const EHairAABBUpdateType UpdateType,
+	uint32 ClusterCount,
+	FRDGImportedBuffer& OutClusterAABBBuffer,
+	FRDGImportedBuffer& OutGroupAABBBuffer);
+
+void AddHairStrandsInterpolationPass(
+	FRDGBuilder& GraphBuilder,
+	FGlobalShaderMap* ShaderMap,
+	const FShaderPrintData* ShaderPrintData,
+	const FHairGroupInstance* Instance,
+	const uint32 VertexCount,
+	const int32 MeshLODIndex,
+	const float HairLengthScale,
+	const EHairInterpolationType HairInterpolationType,
+	const EHairGeometryType InstanceGeometryType,
+	const FRDGHairStrandsCullingData& CullingData,
+	const FVector& InRenHairWorldOffset,
+	const FVector& InSimHairWorldOffset,
+	const FRDGBufferSRVRef& OutRenHairPositionOffsetBuffer,
+	const FRDGBufferSRVRef& OutSimHairPositionOffsetBuffer,
+	const FHairStrandsRestRootResource* RenRestRootResources,
+	const FHairStrandsRestRootResource* SimRestRootResources,
+	const FHairStrandsDeformedRootResource* RenDeformedRootResources,
+	const FHairStrandsDeformedRootResource* SimDeformedRootResources,
+	const FRDGBufferSRVRef& RenRestPosePositionBuffer,
+	const FRDGBufferSRVRef& RenPointToCurveBuffer,
+	const bool bUseSingleGuide,
+	const FRDGBufferSRVRef& InterpolationBuffer,
+	const FRDGBufferSRVRef& SimRestPosePositionBuffer,
+	const FRDGBufferSRVRef& SimDeformedPositionBuffer,
+	const FRDGBufferSRVRef& SimRootPointIndexBuffer,
+	const FRDGBufferSRVRef& SimPointToCurveBuffer,
+	const FRDGBufferSRVRef& RenDeformerPositionBuffer,
+	FRDGBufferUAVRef& OutRenPositionBuffer,
+	const FHairStrandsDeformedRootResource::FLOD::EFrameType DeformedFrame);
+
+
+void AddHairTangentPass(
+	FRDGBuilder& GraphBuilder,
+	FGlobalShaderMap* ShaderMap,
+	uint32 PointCount,
+	FHairGroupPublicData* HairGroupPublicData,
+	FRDGBufferSRVRef PositionBuffer,
+	FRDGBufferUAVRef OutTangentBuffer);
+
+
+enum class EHairPatchAttribute : uint8
+{
+	None,
+	GuideInflucence,
+	ClusterInfluence
+};
+
+void AddPatchAttributePass(
+	FRDGBuilder& GraphBuilder,
+	FGlobalShaderMap* ShaderMap,
+	const uint32 CurveCount,
+	const EHairPatchAttribute Mode,
+	const bool bSimulation,
+	const bool bUseSingleGuide,
+	const FHairStrandsBulkData& RenBulkData,
+	const FRDGBufferRef& RenAttributeBuffer,
+	const FRDGBufferSRVRef& RenCurveBuffer,
+	const FRDGBufferSRVRef& RenCurveToClusterIdBuffer,
+	const FRDGBufferSRVRef& InterpolationBuffer,
+	FRDGImportedBuffer& OutRenAttributeBuffer);
+
+void AddTransferPositionPass(
+	FRDGBuilder& GraphBuilder,
+	FGlobalShaderMap* ShaderMap,
+	const uint32 ElementCount,
+	FRDGBufferSRVRef InBuffer,
+	FRDGBufferUAVRef OutBuffer);
+
+#if RHI_RAYTRACING
+void AddGenerateRaytracingGeometryPass(
+	FRDGBuilder& GraphBuilder,
+	FGlobalShaderMap* ShaderMap,
+	const FShaderPrintData* ShaderPrintData,
+	uint32 InstanceRegisteredIndex,
+	uint32 PointCount,
+	bool bProceduralPrimitive,
+	int ProceduralSplits,
+	float HairRadius,
+	float RootScale,
+	float TipScale,
+	const FRDGBufferSRVRef& HairWorldOffsetBuffer,
+	FRDGHairStrandsCullingData& CullingData,
+	const FRDGBufferSRVRef& PositionBuffer,
+	const FRDGBufferSRVRef& TangentBuffer,
+	const FRDGBufferUAVRef& OutPositionBuffer,
+	const FRDGBufferUAVRef& OutIndexBuffer);
+
+void AddBuildStrandsAccelerationStructurePass(
+	FRDGBuilder& GraphBuilder,
+	FHairGroupInstance* Instance,
+	uint32 ProceduralSplits,
+	bool bNeedUpdate,
+	FRDGBufferRef Raytracing_PositionBuffer,
+	FRDGBufferRef Raytracing_IndexBuffer);
+#endif
\ No newline at end of file
