Change hair cluster culling.
* Change counter to structure buffer to scalar load
* Batch clear and make UAV overlap

Save 10% of cluster

diff --git a/Engine/Plugins/Runtime/HairStrands/Source/HairStrandsCore/Private/GroomManager.cpp b/Engine/Plugins/Runtime/HairStrands/Source/HairStrandsCore/Private/GroomManager.cpp
--- a/Engine/Plugins/Runtime/HairStrands/Source/HairStrandsCore/Private/GroomManager.cpp
+++ b/Engine/Plugins/Runtime/HairStrands/Source/HairStrandsCore/Private/GroomManager.cpp
@@ -841,14 +841,14 @@ static void RunHairStrandsInterpolation_Strands(
 	const EGroomViewMode ViewMode = GetGroomViewMode(*View);
 
 	// Early culling
-	FHairStrandClusterData ClusterData;
+	FHairStrandClusterData ClusterDatas;
 	{
 		// Gather culling jobs
 		for (FInstanceData& InstanceData : InstanceDatas)
 		{
 			if (InstanceData.Instance->Strands.bCullingEnable)
 			{
-				AddInstanceToClusterData(InstanceData.Instance, ClusterData);
+				AddInstanceToClusterData(InstanceData.Instance, ClusterDatas);
 			}
 			else
 			{
@@ -857,12 +857,25 @@ static void RunHairStrandsInterpolation_Strands(
 		}
 
 		// Culling pass
-		ComputeHairStrandsClustersCulling(GraphBuilder, *ShaderMap, Views, ShaderPrintData, ClusterData);
+		if (Views.Num() > 0 && ClusterDatas.HairGroups.Num() > 0)
+		{
+			DECLARE_GPU_STAT(HairStrandsClusterCulling);
+			RDG_EVENT_SCOPE(GraphBuilder, "HairStrandsClusterCulling");
+			TRACE_CPUPROFILER_EVENT_SCOPE(ComputeHairStrandsClustersCulling);
+			RDG_GPU_STAT_SCOPE(GraphBuilder, HairStrandsClusterCulling);
+
+			AddClusterCullingPass(
+				GraphBuilder,
+				ShaderMap,
+				Views[0],
+				ShaderPrintData,
+				ClusterDatas);
+		}
 
 		// Run cluster debug view here (instead of GroomDebug.h/.cpp, as we need to have the (transient) cluster data 
 		if (ViewMode == EGroomViewMode::Cluster || ViewMode == EGroomViewMode::ClusterAABB)
 		{
-			AddDrawDebugClusterPass(GraphBuilder, *View, ShaderMap, ShaderPrintData, ViewMode, ClusterData);
+			AddDrawDebugClusterPass(GraphBuilder, *View, ShaderMap, ShaderPrintData, ViewMode, ClusterDatas);
 		}
 	}
 
@@ -1148,7 +1161,7 @@ static void RunHairStrandsInterpolation_Strands(
 		{
 			// Sanity check
 			check(InstanceData.Instance->Strands.bCullingEnable);
-			HairGroupCluster = &ClusterData.HairGroups[InstanceData.Instance->HairGroupPublicData->ClusterDataIndex];
+			HairGroupCluster = &ClusterDatas.HairGroups[InstanceData.Instance->HairGroupPublicData->ClusterDataIndex];
 			if (!HairGroupCluster->bVisible)
 			{
 				bNeedGPUAABB = false;
diff --git a/Engine/Plugins/Runtime/HairStrands/Source/HairStrandsCore/Private/HairStrandsClusterCulling.cpp b/Engine/Plugins/Runtime/HairStrands/Source/HairStrandsCore/Private/HairStrandsClusterCulling.cpp
--- a/Engine/Plugins/Runtime/HairStrands/Source/HairStrandsCore/Private/HairStrandsClusterCulling.cpp
+++ b/Engine/Plugins/Runtime/HairStrands/Source/HairStrandsCore/Private/HairStrandsClusterCulling.cpp
@@ -13,33 +13,6 @@
 
 ///////////////////////////////////////////////////////////////////////////////////////////////////
 
-class FHairIndBufferClearCS : public FGlobalShader
-{
-	DECLARE_GLOBAL_SHADER(FHairIndBufferClearCS);
-	SHADER_USE_PARAMETER_STRUCT(FHairIndBufferClearCS, FGlobalShader);
-
-	class FSetIndirectDraw : SHADER_PERMUTATION_BOOL("PERMUTATION_SETINDIRECTDRAW");
-	using FPermutationDomain = TShaderPermutationDomain<FSetIndirectDraw>;
-
-	BEGIN_SHADER_PARAMETER_STRUCT(FParameters, )
-		SHADER_PARAMETER_RDG_BUFFER_UAV(RWBuffer, DispatchIndirectParametersClusterCount)
-		SHADER_PARAMETER_RDG_BUFFER_UAV(RWBuffer, DrawIndirectParameters)
-		SHADER_PARAMETER(uint32, VertexCountPerInstance)
-	END_SHADER_PARAMETER_STRUCT()
-
-public:
-	static bool ShouldCompilePermutation(const FGlobalShaderPermutationParameters& Parameters) { return IsHairStrandsSupported(EHairStrandsShaderType::Strands, Parameters.Platform); }
-	static void ModifyCompilationEnvironment(const FGlobalShaderPermutationParameters& Parameters, FShaderCompilerEnvironment& OutEnvironment)
-	{
-		FGlobalShader::ModifyCompilationEnvironment(Parameters, OutEnvironment);
-		OutEnvironment.SetDefine(TEXT("SHADER_CLUSTERCULLINGINDCLEAR"), 1);
-	}
-};
-
-IMPLEMENT_GLOBAL_SHADER(FHairIndBufferClearCS, "/Engine/Private/HairStrands/HairStrandsClusterCulling.usf", "MainClusterCullingIndClearCS", SF_Compute);
-
-///////////////////////////////////////////////////////////////////////////////////////////////////
-
 class FHairClusterCullCS: public FGlobalShader
 {
 	DECLARE_GLOBAL_SHADER(FHairClusterCullCS);
@@ -49,6 +22,7 @@ class FHairClusterCullCS: public FGlobalShader
 	using FPermutationDomain = TShaderPermutationDomain<FPointPerCurve>;
 
 	BEGIN_SHADER_PARAMETER_STRUCT(FParameters,)
+		SHADER_PARAMETER(uint32, ClusterGroupIndex)
 		SHADER_PARAMETER(float, LODIndex)
 		SHADER_PARAMETER(float, LODBias)
 		SHADER_PARAMETER(uint32, CurveCount)
@@ -60,7 +34,7 @@ class FHairClusterCullCS: public FGlobalShader
 		SHADER_PARAMETER_RDG_BUFFER_SRV(Buffer<uint>, PointLODBuffer)
 		SHADER_PARAMETER_RDG_BUFFER_SRV(StructuredBuffer<FPackedHairClusterInfo>, ClusterInfoBuffer)
 	
-		SHADER_PARAMETER_RDG_BUFFER_UAV(RWBuffer<uint>, OutPointCounter)
+		SHADER_PARAMETER_RDG_BUFFER_UAV(RWStructuredBuffer<uint>, OutPointCounter)
 		SHADER_PARAMETER_RDG_BUFFER_UAV(RWBuffer<uint>, OutIndexBuffer)
 		SHADER_PARAMETER_RDG_BUFFER_UAV(RWBuffer<float>, OutRadiusScaleBuffer)
 		SHADER_PARAMETER_RDG_BUFFER_UAV(RWBuffer<uint>, OutCulledCurveBuffer)
@@ -90,7 +64,8 @@ class FHairClusterCullArgsCS: public FGlobalShader
 	SHADER_USE_PARAMETER_STRUCT(FHairClusterCullArgsCS, FGlobalShader);
 
 	BEGIN_SHADER_PARAMETER_STRUCT(FParameters,)
-		SHADER_PARAMETER_RDG_BUFFER_SRV(Buffer<uint>, PointCounterBuffer)
+		SHADER_PARAMETER(uint32, ClusterGroupIndex)
+		SHADER_PARAMETER_RDG_BUFFER_SRV(StructuredBuffer<uint>, PointCounterBuffer)
 		SHADER_PARAMETER_RDG_BUFFER_UAV(RWBuffer<uint>, RWIndirectDrawArgsBuffer)
 		SHADER_PARAMETER_RDG_BUFFER_UAV(RWBuffer<uint>, RWIndirectDispatchArgsBuffer)
 		END_SHADER_PARAMETER_STRUCT()
@@ -107,12 +82,12 @@ public:
 IMPLEMENT_GLOBAL_SHADER(FHairClusterCullArgsCS, "/Engine/Private/HairStrands/HairStrandsClusterCulling.usf", "Main", SF_Compute);
 
 ///////////////////////////////////////////////////////////////////////////////////////////////////
-static void AddClusterCullingPass(
+void AddClusterCullingPass(
 	FRDGBuilder& GraphBuilder,
 	FGlobalShaderMap* ShaderMap,
 	const FSceneView* View,
 	const FShaderPrintData* ShaderPrintData,
-	FHairStrandClusterData::FHairGroup& ClusterData)
+	FHairStrandClusterData& ClusterDatas)
 {
 	check(View);
 
@@ -120,14 +95,19 @@ static void AddClusterCullingPass(
 	ShaderPrint::RequestSpaceForCharacters(2048);
 	ShaderPrint::RequestSpaceForLines(2048);
 
+	const uint32 ClusterCount = ClusterDatas.HairGroups.Num();
+
 	// 0. Glogal counter for visible points
-	FRDGBufferRef PointCounter 	= GraphBuilder.CreateBuffer(FRDGBufferDesc::CreateBufferDesc(sizeof(uint32), 1), TEXT("Hair.ClusterPointCounter"));
-	FRDGBufferUAVRef PointCounterUAV = GraphBuilder.CreateUAV(PointCounter, PF_R32_UINT);
-	AddClearUAVPass(GraphBuilder, PointCounterUAV, 0);
+	FRDGBufferRef PointCounter 	= GraphBuilder.CreateBuffer(FRDGBufferDesc::CreateStructuredDesc(sizeof(uint32), ClusterCount), TEXT("Hair.ClusterPointCounter"));
+	AddClearUAVPass(GraphBuilder, GraphBuilder.CreateUAV(PointCounter, PF_R32_UINT), 0);
 
 	// 1. Build culled index buffer
+	FRDGBufferUAVRef PointCounterUAVSkipBarrier = GraphBuilder.CreateUAV(PointCounter, PF_R32_UINT, ERDGUnorderedAccessViewFlags::SkipBarrier);
+	uint32 ClusterGroupIt = 0;
+	for (FHairStrandClusterData::FHairGroup& ClusterData : ClusterDatas.HairGroups)
 	{
 		FHairClusterCullCS::FParameters* Parameters = GraphBuilder.AllocParameters<FHairClusterCullCS::FParameters>();
+		Parameters->ClusterGroupIndex	 = ClusterGroupIt++;
 		Parameters->LODIndex 			 = ClusterData.LODIndex;
 		Parameters->LODBias 			 = ClusterData.LODBias;
 		Parameters->CurveCount 			 = ClusterData.HairGroupPublicPtr->GetActiveStrandsCurveCount();	
@@ -137,7 +117,7 @@ static void AddClusterCullingPass(
 		Parameters->ClusterInfoBuffer 	 = RegisterAsSRV(GraphBuilder, *ClusterData.ClusterInfoBuffer);
 		Parameters->ClusterInfoParameters= ClusterData.ClusterInfoParameters;
 		Parameters->ViewUniformBuffer	 = View->ViewUniformBuffer;
-		Parameters->OutPointCounter 	 = PointCounterUAV;
+		Parameters->OutPointCounter 	 = PointCounterUAVSkipBarrier;
 		Parameters->OutIndexBuffer 		 = RegisterAsUAV(GraphBuilder, *ClusterData.GetCulledVertexIdBuffer());
 		Parameters->OutRadiusScaleBuffer = RegisterAsUAV(GraphBuilder, *ClusterData.GetCulledVertexRadiusScaleBuffer());
 		Parameters->OutCulledCurveBuffer = RegisterAsUAV(GraphBuilder, *ClusterData.GetCulledCurveBuffer()); // TODO: this could be changed to be transient buffer
@@ -177,12 +157,16 @@ static void AddClusterCullingPass(
 	}
 
 	// 2. Prepare indirect draw/dispatch args buffers
+	FRDGBufferSRVRef PointCounterSRV = GraphBuilder.CreateSRV(PointCounter, PF_R32_UINT);
+	ClusterGroupIt = 0;
+	for (FHairStrandClusterData::FHairGroup& ClusterData : ClusterDatas.HairGroups)
 	{
 		FRDGImportedBuffer DrawIndirectParametersBuffer = Register(GraphBuilder, ClusterData.HairGroupPublicPtr->GetDrawIndirectBuffer(), ERDGImportedBufferFlags::CreateViews);
 		FRDGImportedBuffer DrawIndirectParametersRasterComputeBuffer = Register(GraphBuilder, ClusterData.HairGroupPublicPtr->GetDrawIndirectRasterComputeBuffer(), ERDGImportedBufferFlags::CreateViews);
 
 		FHairClusterCullArgsCS::FParameters* Parameters = GraphBuilder.AllocParameters<FHairClusterCullArgsCS::FParameters>();
-		Parameters->PointCounterBuffer = GraphBuilder.CreateSRV(PointCounter, PF_R32_UINT);
+		Parameters->ClusterGroupIndex = ClusterGroupIt++;
+		Parameters->PointCounterBuffer = PointCounterSRV;
 		Parameters->RWIndirectDrawArgsBuffer = DrawIndirectParametersBuffer.UAV;
 		Parameters->RWIndirectDispatchArgsBuffer = DrawIndirectParametersRasterComputeBuffer.UAV;
 
@@ -196,73 +180,9 @@ static void AddClusterCullingPass(
 
 		GraphBuilder.SetBufferAccessFinal(DrawIndirectParametersBuffer.Buffer, ERHIAccess::IndirectArgs | ERHIAccess::SRVMask);
 		GraphBuilder.SetBufferAccessFinal(DrawIndirectParametersRasterComputeBuffer.Buffer, ERHIAccess::IndirectArgs | ERHIAccess::SRVMask);
+		ClusterData.SetCullingResultAvailable(true);
 	}
 
-	ClusterData.SetCullingResultAvailable(true);
-}
-
-///////////////////////////////////////////////////////////////////////////////////////////////////
-
-static void AddClusterResetLod0(
-	FRDGBuilder& GraphBuilder,
-	FGlobalShaderMap* ShaderMap,
-	FHairStrandClusterData::FHairGroup& ClusterData)
-{
-	// Set as culling result not available
-	ClusterData.SetCullingResultAvailable(false);
-
-	FRDGImportedBuffer IndirectBuffer = Register(GraphBuilder, ClusterData.HairGroupPublicPtr->GetDrawIndirectBuffer(), ERDGImportedBufferFlags::CreateViews);
-
-	// Initialise indirect buffers to entire lod 0 dispatch
-	FHairIndBufferClearCS::FParameters* Parameters = GraphBuilder.AllocParameters<FHairIndBufferClearCS::FParameters>();
-	Parameters->DrawIndirectParameters = IndirectBuffer.UAV;
-	Parameters->VertexCountPerInstance = ClusterData.HairGroupPublicPtr->GetActiveStrandsPointCount() * HAIR_POINT_TO_VERTEX;
-
-	FHairIndBufferClearCS::FPermutationDomain Permutation;
-	Permutation.Set<FHairIndBufferClearCS::FSetIndirectDraw>(true);
-	TShaderMapRef<FHairIndBufferClearCS> ComputeShader(ShaderMap, Permutation);
-	FComputeShaderUtils::AddPass(
-		GraphBuilder,
-		RDG_EVENT_NAME("BufferClearCS"),
-		ComputeShader,
-		Parameters,
-		FIntVector(1, 1, 1));
-
-	GraphBuilder.SetBufferAccessFinal(IndirectBuffer.Buffer, ERHIAccess::IndirectArgs);
-}
-
-///////////////////////////////////////////////////////////////////////////////////////////////////
-
-void ComputeHairStrandsClustersCulling(
-	FRDGBuilder& GraphBuilder,
-	FGlobalShaderMap& ShaderMap,
-	const TArray<const FSceneView*>& Views, 
-	const FShaderPrintData* ShaderPrintData,
-	FHairStrandClusterData& ClusterDatas)
-{
-	DECLARE_GPU_STAT(HairStrandsClusterCulling);
-	RDG_EVENT_SCOPE(GraphBuilder, "HairStrandsClusterCulling");
-	TRACE_CPUPROFILER_EVENT_SCOPE(ComputeHairStrandsClustersCulling);
-	RDG_GPU_STAT_SCOPE(GraphBuilder, HairStrandsClusterCulling);
-
-	for (const FSceneView* View : Views)
-	{
-		// TODO use compute overlap (will need to split AddClusterCullingPass)
-		for (FHairStrandClusterData::FHairGroup& ClusterData : ClusterDatas.HairGroups)
-		{
-			AddClusterResetLod0(GraphBuilder, &ShaderMap, ClusterData);
-		}
-
-		for (FHairStrandClusterData::FHairGroup& ClusterData : ClusterDatas.HairGroups)		
-		{
-			AddClusterCullingPass(
-				GraphBuilder,
-				&ShaderMap,
-				View,
-				ShaderPrintData,
-				ClusterData);
-		}
-	}
 }
 
 ///////////////////////////////////////////////////////////////////////////////////////////////////
diff --git a/Engine/Plugins/Runtime/HairStrands/Source/HairStrandsCore/Private/HairStrandsClusterCulling.h b/Engine/Plugins/Runtime/HairStrands/Source/HairStrandsCore/Private/HairStrandsClusterCulling.h
--- a/Engine/Plugins/Runtime/HairStrands/Source/HairStrandsCore/Private/HairStrandsClusterCulling.h
+++ b/Engine/Plugins/Runtime/HairStrands/Source/HairStrandsCore/Private/HairStrandsClusterCulling.h
@@ -53,9 +53,9 @@ void AddInstanceToClusterData(
 	FHairGroupInstance* In,
 	FHairStrandClusterData& Out);
 
-void ComputeHairStrandsClustersCulling(
+void AddClusterCullingPass(
 	FRDGBuilder& GraphBuilder,
-	FGlobalShaderMap& ShaderMap,
-	const TArray<const FSceneView*>& Views,
+	FGlobalShaderMap* ShaderMap,
+	const FSceneView* View,
 	const FShaderPrintData* ShaderPrintData,
 	FHairStrandClusterData& ClusterDatas);
diff --git a/Engine/Shaders/Private/HairStrands/HairStrandsClusterCulling.usf b/Engine/Shaders/Private/HairStrands/HairStrandsClusterCulling.usf
--- a/Engine/Shaders/Private/HairStrands/HairStrandsClusterCulling.usf
+++ b/Engine/Shaders/Private/HairStrands/HairStrandsClusterCulling.usf
@@ -70,37 +70,6 @@ void MainClearClusterAABBCS(uint2 DispatchThreadId : SV_DispatchThreadID)
 
 #endif
 
-/////////////////////////////////////////////////////////////////////////////////////////////////////////////////////////////
-/// Reset indirect buffer args
-#ifdef SHADER_CLUSTERCULLINGINDCLEAR  
-
-RWBuffer<uint> DispatchIndirectParametersClusterCount;
-RWBuffer<uint> DrawIndirectParameters; 
-uint VertexCountPerInstance;
-
-[numthreads(1, 1, 1)]
-void MainClusterCullingIndClearCS(uint3 DispatchThreadId : SV_DispatchThreadID)
-{
-#if PERMUTATION_SETINDIRECTDRAW
-	DrawIndirectParameters[0] = VertexCountPerInstance;
-	DrawIndirectParameters[1] = 1;	// 1 instance
-	DrawIndirectParameters[2] = 0;
-	DrawIndirectParameters[3] = 0;
-#else
-	if (DispatchThreadId.x==0)
-	{
-		WriteDispatchIndirectArgs(DispatchIndirectParametersClusterCount, 0, 0, 1, 1);
-
-		DrawIndirectParameters[0] = 0;	// 0 vertex (to be accumulated)
-		DrawIndirectParameters[1] = 1;	// 1 instance
-		DrawIndirectParameters[2] = 0;
-		DrawIndirectParameters[3] = 0;
-	}
-#endif
-}
-
-#endif
-
 /////////////////////////////////////////////////////////////////////////////////////////////////////////////////////////////
 // Write out culled index buffex & radius scale buffer
 
@@ -112,6 +81,7 @@ void MainClusterCullingIndClearCS(uint3 DispatchThreadId : SV_DispatchThreadID)
 #include "../ThreadGroupPrefixSum.ush"
 
 //uint ForceLOD;
+uint ClusterGroupIndex;
 float LODIndex;
 float LODBias;
 uint CurveCount;
@@ -124,7 +94,7 @@ Buffer<uint> CurveToClusterIdBuffer;
 Buffer<uint> PointLODBuffer;
 StructuredBuffer<FPackedHairClusterInfo> ClusterInfoBuffer;
 
-RWBuffer<uint>  OutPointCounter;
+RWStructuredBuffer<uint>  OutPointCounter;
 RWBuffer<uint>  OutIndexBuffer;
 RWBuffer<float> OutRadiusScaleBuffer;
 RWBuffer<uint>  OutCulledCurveBuffer;
@@ -266,7 +236,7 @@ void Main(uint3 DispatchThreadId : SV_DispatchThreadID, uint ThreadIndex : SV_Gr
 	// 4. Global allocate
 	if (ThreadIndex == 0)
 	{
-		InterlockedAdd(OutPointCounter[0], AllocationCount, AllocationOffset);
+		InterlockedAdd(OutPointCounter[ClusterGroupIndex], AllocationCount, AllocationOffset);
 	}
 	GroupMemoryBarrierWithGroupSync();
 
@@ -406,7 +376,7 @@ void Main(uint3 DispatchThreadId : SV_DispatchThreadID, uint ThreadIndex : SV_Gr
 	// 4. Global allocate
 	if (ThreadIndex == 0)
 	{
-		InterlockedAdd(OutPointCounter[0], AllocationCount, AllocationOffset);
+		InterlockedAdd(OutPointCounter[ClusterGroupIndex], AllocationCount, AllocationOffset);
 	}
 	GroupMemoryBarrierWithGroupSync();
 
@@ -447,14 +417,15 @@ void Main(uint3 DispatchThreadId : SV_DispatchThreadID, uint ThreadIndex : SV_Gr
 
 #if SHADER_CLUSTER_CULL_ARGS
 
-Buffer<uint> PointCounterBuffer;
+uint ClusterGroupIndex;
+StructuredBuffer<uint> PointCounterBuffer;
 RWBuffer<uint> RWIndirectDrawArgsBuffer;
 RWBuffer<uint> RWIndirectDispatchArgsBuffer;
 
 [numthreads(1, 1, 1)]
 void Main(uint DispatchThreadId : SV_DispatchThreadID, uint ThreadIndex : SV_GroupIndex)
 {
-	const uint PointCount = PointCounterBuffer[0];
+	const uint PointCount = PointCounterBuffer[ClusterGroupIndex];
 	const uint DispatchCount = DivideAndRoundUp(PointCount, HAIR_VERTEXCOUNT_GROUP_SIZE);
 
 	RWIndirectDrawArgsBuffer[0] = PointCount * HAIR_POINT_TO_VERTEX;
