Fix missing resource transitions when building follicule texture.

diff --git a/Engine/Plugins/Runtime/HairStrands/Source/HairStrandsCore/Private/GroomTextureBuilder.cpp b/Engine/Plugins/Runtime/HairStrands/Source/HairStrandsCore/Private/GroomTextureBuilder.cpp
--- a/Engine/Plugins/Runtime/HairStrands/Source/HairStrandsCore/Private/GroomTextureBuilder.cpp
+++ b/Engine/Plugins/Runtime/HairStrands/Source/HairStrandsCore/Private/GroomTextureBuilder.cpp
@@ -236,12 +236,17 @@ static void InternalGenerateFollicleTexture_GPU(
 	{
 		if (OutTexture && OutTexture->GetResource() && OutTexture->GetResource()->GetTexture2DRHI())
 		{
+			FRHITexture2D* DstTexture = OutTexture->GetResource()->GetTexture2DRHI();
+			RHICmdList.Transition(FRHITransitionInfo(DstTexture, ERHIAccess::SRVMask, ERHIAccess::CopyDest));
+
 			FRHICopyTextureInfo CopyInfo;
 			CopyInfo.NumMips = MipCount;
 			RHICmdList.CopyTexture(
 				FollicleMaskTexture->GetRHI(),
-				OutTexture->GetResource()->GetTexture2DRHI(),
+				DstTexture,
 				CopyInfo);
+
+			RHICmdList.Transition(FRHITransitionInfo(DstTexture, ERHIAccess::CopyDest, ERHIAccess::SRVMask));
 		}
 	});
 }
