Change hair root mesh update into a single pass instead of iterating over all sections.

diff --git a/Engine/Plugins/Runtime/HairStrands/Source/HairStrandsCore/Private/HairStrandsMeshProjection.cpp b/Engine/Plugins/Runtime/HairStrands/Source/HairStrandsCore/Private/HairStrandsMeshProjection.cpp
--- a/Engine/Plugins/Runtime/HairStrands/Source/HairStrandsCore/Private/HairStrandsMeshProjection.cpp
+++ b/Engine/Plugins/Runtime/HairStrands/Source/HairStrandsCore/Private/HairStrandsMeshProjection.cpp
@@ -168,79 +168,24 @@ void AddSkinUpdatePass(
 
 class FHairUpdateMeshTriangleCS : public FGlobalShader
 {
-public:
-	const static uint32 SectionArrayCount = 16; // This defines the number of sections managed for each iteration pass
 private:
 	DECLARE_GLOBAL_SHADER(FHairUpdateMeshTriangleCS);
 	SHADER_USE_PARAMETER_STRUCT(FHairUpdateMeshTriangleCS, FGlobalShader);
 
-	class FUpdateUVs : SHADER_PERMUTATION_BOOL("PERMUTATION_WITHUV");
 	class FPositionType : SHADER_PERMUTATION_BOOL("PERMUTATION_POSITION_TYPE");
 	class FPrevious : SHADER_PERMUTATION_BOOL("PERMUTATION_PREVIOUS"); 
-	using FPermutationDomain = TShaderPermutationDomain<FUpdateUVs, FPositionType, FPrevious>;
+	using FPermutationDomain = TShaderPermutationDomain<FPositionType, FPrevious>;
 
 	BEGIN_SHADER_PARAMETER_STRUCT(FParameters, )
-		SHADER_PARAMETER(uint32, MaxUniqueTriangleCount)
 		SHADER_PARAMETER(uint32, MaxSectionCount)
-		SHADER_PARAMETER(uint32, Pass_SectionStart)
-		SHADER_PARAMETER(uint32, Pass_SectionCount)
-		
-		SHADER_PARAMETER_ARRAY(FUintVector4, PackedMeshSectionScalars1, [SectionArrayCount])
-		SHADER_PARAMETER_ARRAY(FUintVector4, PackedMeshSectionScalars2, [SectionArrayCount])
-
-		SHADER_PARAMETER_RDG_BUFFER_SRV(Buffer, RDGMeshPositionBuffer0)
-		SHADER_PARAMETER_RDG_BUFFER_SRV(Buffer, RDGMeshPositionBuffer1)
-		SHADER_PARAMETER_RDG_BUFFER_SRV(Buffer, RDGMeshPositionBuffer2)
-		SHADER_PARAMETER_RDG_BUFFER_SRV(Buffer, RDGMeshPositionBuffer3)
-		SHADER_PARAMETER_RDG_BUFFER_SRV(Buffer, RDGMeshPositionBuffer4)
-		SHADER_PARAMETER_RDG_BUFFER_SRV(Buffer, RDGMeshPositionBuffer5)
-		SHADER_PARAMETER_RDG_BUFFER_SRV(Buffer, RDGMeshPositionBuffer6)
-		SHADER_PARAMETER_RDG_BUFFER_SRV(Buffer, RDGMeshPositionBuffer7)
-
-		SHADER_PARAMETER_RDG_BUFFER_SRV(Buffer, RDGMeshPreviousPositionBuffer0)
-		SHADER_PARAMETER_RDG_BUFFER_SRV(Buffer, RDGMeshPreviousPositionBuffer1)
-		SHADER_PARAMETER_RDG_BUFFER_SRV(Buffer, RDGMeshPreviousPositionBuffer2)
-		SHADER_PARAMETER_RDG_BUFFER_SRV(Buffer, RDGMeshPreviousPositionBuffer3)
-		SHADER_PARAMETER_RDG_BUFFER_SRV(Buffer, RDGMeshPreviousPositionBuffer4)
-		SHADER_PARAMETER_RDG_BUFFER_SRV(Buffer, RDGMeshPreviousPositionBuffer5)
-		SHADER_PARAMETER_RDG_BUFFER_SRV(Buffer, RDGMeshPreviousPositionBuffer6)
-		SHADER_PARAMETER_RDG_BUFFER_SRV(Buffer, RDGMeshPreviousPositionBuffer7)
-
-		SHADER_PARAMETER_SRV(Buffer, MeshPositionBuffer0)
-		SHADER_PARAMETER_SRV(Buffer, MeshPositionBuffer1)
-		SHADER_PARAMETER_SRV(Buffer, MeshPositionBuffer2)
-		SHADER_PARAMETER_SRV(Buffer, MeshPositionBuffer3)
-		SHADER_PARAMETER_SRV(Buffer, MeshPositionBuffer4)
-		SHADER_PARAMETER_SRV(Buffer, MeshPositionBuffer5)
-		SHADER_PARAMETER_SRV(Buffer, MeshPositionBuffer6)
-		SHADER_PARAMETER_SRV(Buffer, MeshPositionBuffer7)
-
-		SHADER_PARAMETER_SRV(Buffer, MeshPreviousPositionBuffer0)
-		SHADER_PARAMETER_SRV(Buffer, MeshPreviousPositionBuffer1)
-		SHADER_PARAMETER_SRV(Buffer, MeshPreviousPositionBuffer2)
-		SHADER_PARAMETER_SRV(Buffer, MeshPreviousPositionBuffer3)
-		SHADER_PARAMETER_SRV(Buffer, MeshPreviousPositionBuffer4)
-		SHADER_PARAMETER_SRV(Buffer, MeshPreviousPositionBuffer5)
-		SHADER_PARAMETER_SRV(Buffer, MeshPreviousPositionBuffer6)
-		SHADER_PARAMETER_SRV(Buffer, MeshPreviousPositionBuffer7)
-
-		SHADER_PARAMETER_SRV(Buffer, MeshIndexBuffer0)
-		SHADER_PARAMETER_SRV(Buffer, MeshIndexBuffer1)
-		SHADER_PARAMETER_SRV(Buffer, MeshIndexBuffer2)
-		SHADER_PARAMETER_SRV(Buffer, MeshIndexBuffer3)
-		SHADER_PARAMETER_SRV(Buffer, MeshIndexBuffer4)
-		SHADER_PARAMETER_SRV(Buffer, MeshIndexBuffer5)
-		SHADER_PARAMETER_SRV(Buffer, MeshIndexBuffer6)
-		SHADER_PARAMETER_SRV(Buffer, MeshIndexBuffer7)
-		
-		SHADER_PARAMETER_SRV(Buffer, MeshUVsBuffer0)
-		SHADER_PARAMETER_SRV(Buffer, MeshUVsBuffer1)
-		SHADER_PARAMETER_SRV(Buffer, MeshUVsBuffer2)
-		SHADER_PARAMETER_SRV(Buffer, MeshUVsBuffer3)
-		SHADER_PARAMETER_SRV(Buffer, MeshUVsBuffer4)
-		SHADER_PARAMETER_SRV(Buffer, MeshUVsBuffer5)
-		SHADER_PARAMETER_SRV(Buffer, MeshUVsBuffer6)
-		SHADER_PARAMETER_SRV(Buffer, MeshUVsBuffer7)
+		SHADER_PARAMETER(uint32, MaxUniqueTriangleCount)
+		SHADER_PARAMETER_RDG_BUFFER_SRV(StructuredBuffer, MeshSectionBuffer)
+		SHADER_PARAMETER_RDG_BUFFER_SRV(Buffer, RDGMeshPositionBuffer)
+		SHADER_PARAMETER_RDG_BUFFER_SRV(Buffer, RDGMeshPreviousPositionBuffer)
+		SHADER_PARAMETER_SRV(Buffer, MeshPositionBuffer)
+		SHADER_PARAMETER_SRV(Buffer, MeshPreviousPositionBuffer)
+		SHADER_PARAMETER_SRV(Buffer, MeshIndexBuffer)
+		SHADER_PARAMETER_SRV(Buffer, MeshUVsBuffer)
 
 		SHADER_PARAMETER_RDG_BUFFER_SRV(Buffer, UniqueTriangleIndices)
 		SHADER_PARAMETER_RDG_BUFFER_UAV(RWStructuredBuffer, OutUniqueTrianglePrevPosition)
@@ -248,50 +193,34 @@ private:
 	END_SHADER_PARAMETER_STRUCT()
 
 public:
+	static uint32 GetGroupSize() { return 128; }
 	static bool ShouldCompilePermutation(const FGlobalShaderPermutationParameters& Parameters) { return IsHairStrandsSupported(EHairStrandsShaderType::All, Parameters.Platform); }
 	static void ModifyCompilationEnvironment(const FGlobalShaderPermutationParameters& Parameters, FShaderCompilerEnvironment& OutEnvironment)
 	{
 		FGlobalShader::ModifyCompilationEnvironment(Parameters, OutEnvironment);
-		OutEnvironment.SetDefine(TEXT("MAX_SECTION_COUNT"), SectionArrayCount);
-		OutEnvironment.SetDefine(TEXT("SHADER_MESH_UPDATE"), SectionArrayCount);		
+		OutEnvironment.SetDefine(TEXT("SHADER_MESH_UPDATE"), 1);
+		OutEnvironment.SetDefine(TEXT("GROUP_SIZE"), GetGroupSize());
 	}
 };
 
 IMPLEMENT_GLOBAL_SHADER(FHairUpdateMeshTriangleCS, "/Engine/Private/HairStrands/HairStrandsMesh.usf", "MainCS", SF_Compute);
 
-// Manual packing for mesh section scalars.  MUST MATCH WITH HairStrandsMesh.usf
-// PackedMeshSectionScalars1: MeshSectionIndices, MeshMaxIndexCount, MeshMaxVertexCount, MeshIndexOffset
-// PackedMeshSectionScalars2: MeshUVsChannelOffset, MeshUVsChannelCount, MeshSectionBufferIndex, *free*
-inline void SetMeshSectionIndices(FHairUpdateMeshTriangleCS::FParameters& Params, uint32 SectionIndex, uint32 MeshSectionIndices)         { Params.PackedMeshSectionScalars1[SectionIndex].X = MeshSectionIndices; }
-inline void SetMeshMaxIndexCount(FHairUpdateMeshTriangleCS::FParameters& Params, uint32 SectionIndex, uint32 MeshMaxIndexCount)           { Params.PackedMeshSectionScalars1[SectionIndex].Y = MeshMaxIndexCount; }
-inline void SetMeshMaxVertexCount(FHairUpdateMeshTriangleCS::FParameters& Params, uint32 SectionIndex, uint32 MeshMaxVertexCount)         { Params.PackedMeshSectionScalars1[SectionIndex].Z = MeshMaxVertexCount; }
-inline void SetMeshIndexOffset(FHairUpdateMeshTriangleCS::FParameters& Params, uint32 SectionIndex, uint32 MeshIndexOffset)               { Params.PackedMeshSectionScalars1[SectionIndex].W = MeshIndexOffset; }
-inline void SetMeshUVsChannelOffset(FHairUpdateMeshTriangleCS::FParameters& Params, uint32 SectionIndex, uint32 MeshUVsChannelOffset)     { Params.PackedMeshSectionScalars2[SectionIndex].X = MeshUVsChannelOffset; }
-inline void SetMeshUVsChannelCount(FHairUpdateMeshTriangleCS::FParameters& Params, uint32 SectionIndex, uint32 MeshUVsChannelCount)       { Params.PackedMeshSectionScalars2[SectionIndex].Y = MeshUVsChannelCount; }
-inline void SetMeshSectionBufferIndex(FHairUpdateMeshTriangleCS::FParameters& Params, uint32 SectionIndex, uint32 MeshSectionBufferIndex) { Params.PackedMeshSectionScalars2[SectionIndex].Z = MeshSectionBufferIndex; }
-
-
 void AddHairStrandUpdateMeshTrianglesPass(
 	FRDGBuilder& GraphBuilder,
 	FGlobalShaderMap* ShaderMap,
-	const int32 LODIndex,
+	const int32 MeshLODIndex,
 	const FHairStrandsProjectionMeshData::LOD& MeshData,
 	FHairStrandsRestRootResource* RestResources,
 	FHairStrandsDeformedRootResource* DeformedResources)
 {
 	const uint32 RootCount = RestResources->GetRootCount();
-	if (RootCount == 0 || LODIndex < 0)
-	{
-		return;
-	}
-
-	if ((LODIndex >= RestResources->LODs.Num() || LODIndex >= DeformedResources->LODs.Num()))
+	if (RootCount == 0 || !RestResources->LODs.IsValidIndex(MeshLODIndex))
 	{
 		return;
 	}
 
-	FHairStrandsRestRootResource::FLOD& RestLODData = RestResources->LODs[LODIndex];
-	check(RestLODData.LODIndex == LODIndex);
+	FHairStrandsRestRootResource::FLOD& RestLODData = RestResources->LODs[MeshLODIndex];
+	check(RestLODData.LODIndex == MeshLODIndex);
 
 	const uint32 SectionCount = MeshData.Sections.Num();
 	const uint32 MaxSupportedSectionCount = GetHairStrandsMaxSectionCount();
@@ -301,207 +230,88 @@ void AddHairStrandUpdateMeshTrianglesPass(
 		return;
 	}
 
-	// Update the last known mesh LOD for which the root resources has been updated
-	DeformedResources->MeshLODIndex = LODIndex;
-
-	// When the number of section of a mesh is above FHairUpdateMeshTriangleCS::SectionArrayCount, the update is split into several passes
-	const TArray<uint32>& ValidSectionIndices = RestResources->BulkData.GetValidSectionIndices(LODIndex);
-	const uint32 ValidSectionCount = ValidSectionIndices.Num();
-	const uint32 PassCount = FMath::DivideAndRoundUp(ValidSectionCount, FHairUpdateMeshTriangleCS::SectionArrayCount);
-	const bool bComputePreviousDeformedPosition = IsHairStrandContinuousDecimationReorderingEnabled();
+	// If a skel. mesh is not streaming yet, its SRV will be null
+	const bool bReady = MeshData.Sections[0].IndexBuffer != nullptr && MeshData.Sections[0].UVsBuffer != nullptr && (MeshData.Sections[0].RDGPositionBuffer != nullptr || MeshData.Sections[0].PositionBuffer != nullptr);
+	if (!bReady)
+	{
+		return;
+	}
 
-	const uint32 MaxUniqueTriangleCount = RestLODData.RestUniqueTrianglePositionBuffer.Buffer->Desc.NumElements * 3;
+	// Update the last known mesh LOD for which the root resources has been updated
+	DeformedResources->MeshLODIndex = MeshLODIndex;
 
 	FHairUpdateMeshTriangleCS::FParameters CommonParameters;
+	CommonParameters.MaxUniqueTriangleCount 		= RootCount;
+	CommonParameters.MaxSectionCount 				= SectionCount;
+	CommonParameters.RDGMeshPositionBuffer			= MeshData.Sections[0].RDGPositionBuffer;
+	CommonParameters.RDGMeshPreviousPositionBuffer 	= MeshData.Sections[0].RDGPreviousPositionBuffer;
+	CommonParameters.MeshPositionBuffer				= MeshData.Sections[0].PositionBuffer;
+	CommonParameters.MeshPreviousPositionBuffer		= MeshData.Sections[0].PreviousPositionBuffer;
+	CommonParameters.MeshIndexBuffer				= MeshData.Sections[0].IndexBuffer;
+	CommonParameters.MeshUVsBuffer					= MeshData.Sections[0].UVsBuffer;
+
+	FHairStrandsDeformedRootResource::FLOD& DeformedLODData = DeformedResources->LODs[MeshLODIndex];
+	DeformedLODData.Status = FHairStrandsDeformedRootResource::FLOD::EStatus::Completed;
 
-	FRDGImportedBuffer OutputCurrBuffer;
-	FRDGImportedBuffer OutputPrevBuffer;
-	const bool bEnableUAVOverlap = true;
-	const ERDGUnorderedAccessViewFlags UAVFlags = bEnableUAVOverlap ? ERDGUnorderedAccessViewFlags::SkipBarrier : ERDGUnorderedAccessViewFlags::None;
 	CommonParameters.UniqueTriangleIndices = RegisterAsSRV(GraphBuilder, RestLODData.UniqueTriangleIndexBuffer);
-	{
-		FHairStrandsDeformedRootResource::FLOD& DeformedLODData = DeformedResources->LODs[LODIndex];
-		OutputCurrBuffer = Register(GraphBuilder, DeformedLODData.GetDeformedUniqueTrianglePositionBuffer(FHairStrandsDeformedRootResource::FLOD::Current), ERDGImportedBufferFlags::CreateUAV, UAVFlags);
-
-		if (bComputePreviousDeformedPosition)
-		{
-			OutputPrevBuffer = Register(GraphBuilder, DeformedLODData.GetDeformedUniqueTrianglePositionBuffer(FHairStrandsDeformedRootResource::FLOD::Previous), ERDGImportedBufferFlags::CreateUAV, UAVFlags);
-		}
-
-		DeformedLODData.Status = FHairStrandsDeformedRootResource::FLOD::EStatus::Completed;
-	}
 
+	// Current 
+	FRDGImportedBuffer OutputCurrBuffer = Register(GraphBuilder, DeformedLODData.GetDeformedUniqueTrianglePositionBuffer(FHairStrandsDeformedRootResource::FLOD::Current), ERDGImportedBufferFlags::CreateUAV, ERDGUnorderedAccessViewFlags::SkipBarrier);
 	CommonParameters.OutUniqueTriangleCurrPosition = OutputCurrBuffer.UAV;
 
+	// Previous
+	FRDGImportedBuffer OutputPrevBuffer;
+	const bool bComputePreviousDeformedPosition = IsHairStrandContinuousDecimationReorderingEnabled();
 	if (bComputePreviousDeformedPosition)
 	{
+		OutputPrevBuffer = Register(GraphBuilder, DeformedLODData.GetDeformedUniqueTrianglePositionBuffer(FHairStrandsDeformedRootResource::FLOD::Previous), ERDGImportedBufferFlags::CreateUAV, ERDGUnorderedAccessViewFlags::SkipBarrier);
 		CommonParameters.OutUniqueTrianglePrevPosition = OutputPrevBuffer.UAV;
 	}
 
-	for (uint32 PassIt = 0; PassIt < PassCount; ++PassIt)
+	struct FSectionData
 	{
-		FHairUpdateMeshTriangleCS::FParameters* Parameters = &CommonParameters;
-		Parameters->MaxUniqueTriangleCount = MaxUniqueTriangleCount;
-		Parameters->MaxSectionCount = MeshData.Sections.Num();
-		Parameters->Pass_SectionStart = PassIt * FHairUpdateMeshTriangleCS::SectionArrayCount;
-		Parameters->Pass_SectionCount = FMath::Min(ValidSectionCount - Parameters->Pass_SectionStart, FHairUpdateMeshTriangleCS::SectionArrayCount);
-
-		// Most often, a skeletal mesh will have many sections, but *most* of will share the same vertex buffers. 
-		// So NumSection >> NumBuffer, and the limitation of MaxSectionBufferCount = 8;
-		const uint32 MaxSectionBufferCount = 8;
-		struct FMeshSectionBuffers
-		{
-			uint32 MeshSectionBufferIndex = 0;
-			FRDGBufferSRVRef RDGPositionBuffer = nullptr;
-			FRDGBufferSRVRef RDGPreviousPositionBuffer = nullptr;
-			FRHIShaderResourceView* PositionBuffer = nullptr;
-			FRHIShaderResourceView* PreviousPositionBuffer = nullptr;
-			FRHIShaderResourceView* IndexBuffer = nullptr;
-			FRHIShaderResourceView* UVsBuffer = nullptr;
-		};
-		TMap<FRHIShaderResourceView*, FMeshSectionBuffers>	UniqueMeshSectionBuffers;
-		TMap<FRDGBufferSRVRef, FMeshSectionBuffers>			UniqueMeshSectionBuffersRDG;
-		uint32 UniqueMeshSectionBufferIndex = 0;
-
-		#define SETPARAMETERS(OutParameters, InMeshSectionData, Index) \
-				OutParameters->RDGMeshPositionBuffer##Index			= InMeshSectionData.RDGPositionBuffer; \
-				OutParameters->RDGMeshPreviousPositionBuffer##Index = InMeshSectionData.RDGPreviousPositionBuffer; \
-				OutParameters->MeshPositionBuffer##Index			= InMeshSectionData.PositionBuffer; \
-				OutParameters->MeshPreviousPositionBuffer##Index	= InMeshSectionData.PreviousPositionBuffer; \
-				OutParameters->MeshIndexBuffer##Index				= InMeshSectionData.IndexBuffer; \
-				OutParameters->MeshUVsBuffer##Index					= InMeshSectionData.UVsBuffer;
-
-		auto SetMeshSectionBuffers = [Parameters](uint32 UniqueIndex, const FHairStrandsProjectionMeshData::Section& MeshSectionData)
-		{
-			switch (UniqueIndex)
-			{
-			case 0: SETPARAMETERS(Parameters, MeshSectionData, 0); break;
-			case 1: SETPARAMETERS(Parameters, MeshSectionData, 1); break;
-			case 2: SETPARAMETERS(Parameters, MeshSectionData, 2); break;
-			case 3: SETPARAMETERS(Parameters, MeshSectionData, 3); break;
-			case 4: SETPARAMETERS(Parameters, MeshSectionData, 4); break;
-			case 5: SETPARAMETERS(Parameters, MeshSectionData, 5); break;
-			case 6: SETPARAMETERS(Parameters, MeshSectionData, 6); break;
-			case 7: SETPARAMETERS(Parameters, MeshSectionData, 7); break;
-			}
-		};
-		#undef SETPARAMETERS
-
-		bool bUseRDGPositionBuffer = false;
-		for (int32 SectionStartIt = Parameters->Pass_SectionStart, SectionItEnd = Parameters->Pass_SectionStart + Parameters->Pass_SectionCount; SectionStartIt < SectionItEnd; ++SectionStartIt)
-		{
-			const int32 SectionIt = SectionStartIt - Parameters->Pass_SectionStart;
-			const int32 SectionIndex = ValidSectionIndices[SectionStartIt];
-
-			if (SectionIndex < 0 || SectionIndex >= MeshData.Sections.Num())
-			{
-				continue;
-			}
-
-			const FHairStrandsProjectionMeshData::Section& MeshSectionData = MeshData.Sections[SectionIndex];
-
-			// If a skel. mesh is not streaming yet, its SRV will be null
-			if (MeshSectionData.IndexBuffer == nullptr)
-			{
-				continue;
-			}
-
-			const FMeshSectionBuffers* Buffers = nullptr;
-			if (MeshSectionData.PositionBuffer)
-			{
-				Buffers = UniqueMeshSectionBuffers.Find(MeshSectionData.PositionBuffer);
-			}
-			else if (MeshSectionData.RDGPositionBuffer)
-			{
-				Buffers = UniqueMeshSectionBuffersRDG.Find(MeshSectionData.RDGPositionBuffer);
-			}
-			else
-			{
-				check(false); // Should never happen
-				continue;
-			}
-
-			if (Buffers != nullptr)
-			{
-				// Insure that all buffers actually match
-				check(Buffers->RDGPositionBuffer == MeshSectionData.RDGPositionBuffer);
-				check(Buffers->RDGPreviousPositionBuffer == MeshSectionData.RDGPreviousPositionBuffer);
-				check(Buffers->PositionBuffer == MeshSectionData.PositionBuffer);
-				check(Buffers->PreviousPositionBuffer == MeshSectionData.PreviousPositionBuffer);
-				check(Buffers->IndexBuffer == MeshSectionData.IndexBuffer);
-				check(Buffers->UVsBuffer == MeshSectionData.UVsBuffer);
-
-				SetMeshSectionBufferIndex(*Parameters, SectionIt, Buffers->MeshSectionBufferIndex);
-			}
-			else
-			{
-				// Only support 8 unique different buffer at the moment
-				check(UniqueMeshSectionBufferIndex < MaxSectionBufferCount);
-				SetMeshSectionBuffers(UniqueMeshSectionBufferIndex, MeshSectionData);
-
-				FMeshSectionBuffers Entry;
-				Entry.MeshSectionBufferIndex = UniqueMeshSectionBufferIndex;
-				Entry.RDGPositionBuffer = MeshSectionData.RDGPositionBuffer;
-				Entry.RDGPreviousPositionBuffer = MeshSectionData.RDGPreviousPositionBuffer;
-				Entry.PositionBuffer = MeshSectionData.PositionBuffer;
-				Entry.PreviousPositionBuffer = MeshSectionData.PreviousPositionBuffer;
-				Entry.IndexBuffer = MeshSectionData.IndexBuffer;
-				Entry.UVsBuffer = MeshSectionData.UVsBuffer;
-
-				if (MeshSectionData.PositionBuffer)
-				{
-					UniqueMeshSectionBuffers.Add(MeshSectionData.PositionBuffer, Entry);
-				}
-				else if (MeshSectionData.RDGPositionBuffer)
-				{
-					UniqueMeshSectionBuffersRDG.Add(MeshSectionData.RDGPositionBuffer, Entry);
-				}
-
-				SetMeshSectionBufferIndex(*Parameters, SectionIt, UniqueMeshSectionBufferIndex);
-				++UniqueMeshSectionBufferIndex;
-			}
-
-			SetMeshSectionIndices(*Parameters, SectionIt, MeshSectionData.SectionIndex);
-			SetMeshMaxIndexCount(*Parameters, SectionIt, MeshSectionData.TotalIndexCount);
-			SetMeshMaxVertexCount(*Parameters, SectionIt, MeshSectionData.TotalVertexCount);
-			SetMeshIndexOffset(*Parameters, SectionIt, MeshSectionData.IndexBaseIndex);
-			SetMeshUVsChannelOffset(*Parameters, SectionIt, MeshSectionData.UVsChannelOffset);
-			SetMeshUVsChannelCount(*Parameters, SectionIt, MeshSectionData.UVsChannelCount);
-
-			// Sanity check
-			// If one of the input is using RDG position, we expect all mesh sections to use RDG input
-			if (bUseRDGPositionBuffer)
-			{
-				check(MeshSectionData.RDGPositionBuffer != nullptr);
-			}
-			else if (MeshSectionData.RDGPositionBuffer != nullptr)
-			{
-				bUseRDGPositionBuffer = true;
-			}
-		}
-
-		if (MeshData.Sections.Num() > 0)
-		{
-			for (uint32 Index = UniqueMeshSectionBufferIndex; Index < MaxSectionBufferCount; ++Index)
-			{
-				SetMeshSectionBuffers(Index, MeshData.Sections[0]);
-			}
-		}
-
-		if (UniqueMeshSectionBufferIndex == 0 || Parameters->MeshUVsBuffer0 == nullptr)
-		{
-			return;
-		}
+		uint32 TotalIndexCount;
+		uint32 TotalVertexCount;
+		uint32 IndexBaseIndex;
+		uint32 UVsChannelOffset : 8;
+		uint32 UVsChannelCount : 8;
+		uint32 Pad : 16;
+	};
+	TArray<FSectionData> SectionDatas;
+	SectionDatas.SetNum(SectionCount);
+	for (uint32 SectionIt = 0; SectionIt < SectionCount; ++SectionIt)
+	{
+		const FHairStrandsProjectionMeshData::Section& MeshSectionData = MeshData.Sections[SectionIt];
+		SectionDatas[SectionIt].TotalIndexCount	= MeshSectionData.TotalIndexCount;
+		SectionDatas[SectionIt].TotalVertexCount= MeshSectionData.TotalVertexCount;
+		SectionDatas[SectionIt].IndexBaseIndex	= MeshSectionData.IndexBaseIndex;
+		SectionDatas[SectionIt].UVsChannelOffset= MeshSectionData.UVsChannelOffset;
+		SectionDatas[SectionIt].UVsChannelCount	= MeshSectionData.UVsChannelCount;
+
+		// Sanity check
+		check(MeshSectionData.SectionIndex == SectionIt);
+		check(MeshSectionData.UVsChannelOffset < 255);
+		check(MeshSectionData.UVsChannelCount < 255);
+		check(CommonParameters.RDGMeshPositionBuffer		== MeshData.Sections[SectionIt].RDGPositionBuffer);
+		check(CommonParameters.RDGMeshPreviousPositionBuffer== MeshData.Sections[SectionIt].RDGPreviousPositionBuffer);
+		check(CommonParameters.MeshPositionBuffer			== MeshData.Sections[SectionIt].PositionBuffer);
+		check(CommonParameters.MeshPreviousPositionBuffer	== MeshData.Sections[SectionIt].PreviousPositionBuffer);
+		check(CommonParameters.MeshIndexBuffer				== MeshData.Sections[SectionIt].IndexBuffer);
+		check(CommonParameters.MeshUVsBuffer				== MeshData.Sections[SectionIt].UVsBuffer);
+	}
+	FRDGBufferRef SectionBuffer = CreateStructuredBuffer(GraphBuilder, TEXT("Hair.SkelMeshSectionBuffer"), sizeof(FSectionData),  SectionDatas.Num(), SectionDatas.GetData(), sizeof(FSectionData) * SectionDatas.Num());
+	CommonParameters.MeshSectionBuffer = GraphBuilder.CreateSRV(SectionBuffer);
 
+	const bool bUseRDGPositionBuffer = CommonParameters.RDGMeshPositionBuffer != nullptr;
+	{
 		FHairUpdateMeshTriangleCS::FParameters* PassParameters = GraphBuilder.AllocParameters<FHairUpdateMeshTriangleCS::FParameters>();
-		*PassParameters = *Parameters;
+		*PassParameters = CommonParameters;
 
 		FHairUpdateMeshTriangleCS::FPermutationDomain PermutationVector;
-		PermutationVector.Set<FHairUpdateMeshTriangleCS::FUpdateUVs>(1);
 		PermutationVector.Set<FHairUpdateMeshTriangleCS::FPrevious>(bComputePreviousDeformedPosition);
 		PermutationVector.Set<FHairUpdateMeshTriangleCS::FPositionType>(bUseRDGPositionBuffer ? 1 : 0);
 
-		const FIntVector DispatchGroupCount = FComputeShaderUtils::GetGroupCount(RootCount, 128);
+		const FIntVector DispatchGroupCount = FComputeShaderUtils::GetGroupCount(RootCount, FHairUpdateMeshTriangleCS::GetGroupSize());
 		check(DispatchGroupCount.X <= GRHIMaxDispatchThreadGroupsPerDimension.X);
 		TShaderMapRef<FHairUpdateMeshTriangleCS> ComputeShader(ShaderMap, PermutationVector);
 		FComputeShaderUtils::AddPass(
diff --git a/Engine/Shaders/Private/HairStrands/HairStrandsMesh.usf b/Engine/Shaders/Private/HairStrands/HairStrandsMesh.usf
--- a/Engine/Shaders/Private/HairStrands/HairStrandsMesh.usf
+++ b/Engine/Shaders/Private/HairStrands/HairStrandsMesh.usf
@@ -211,224 +211,102 @@ void MainCS(uint3 DispatchThreadId : SV_DispatchThreadID)
 
 #if SHADER_MESH_UPDATE
 
-#ifndef MAX_SECTION_COUNT
-#error MAX_SECTION_COUNT needs to be defined
+#ifndef GROUP_SIZE
+	#error GROUP_SIZE needs to be defined 
 #endif
 
-uint Pass_SectionStart;
-uint Pass_SectionCount;
+struct FMeshSectionData
+{
+	uint MaxIndexCount;
+	uint MaxVertexCount;
+	uint IndexOffset;
+	uint UVsChannelOffset;
+	uint UVsChannelCount;
+};
 
 uint MaxSectionCount;
-
-// Manual packing for mesh section scalars.  MUST MATCH WITH HairStrandsMeshProjection.cpp
-// PackedMeshSectionScalars1: MeshSectionIndices, MeshMaxIndexCount, MeshMaxVertexCount, MeshIndexOffset
-// PackedMeshSectionScalars2: MeshUVsChannelOffset, MeshUVsChannelCount, MeshSectionBufferIndex, *free*
-uint4 PackedMeshSectionScalars1[MAX_SECTION_COUNT];
-uint4 PackedMeshSectionScalars2[MAX_SECTION_COUNT];
-#define MeshSectionIndices(SectionIndex) PackedMeshSectionScalars1[SectionIndex].x
-#define MeshMaxIndexCount(SectionIndex) PackedMeshSectionScalars1[SectionIndex].y
-#define MeshMaxVertexCount(SectionIndex) PackedMeshSectionScalars1[SectionIndex].z
-#define MeshIndexOffset(SectionIndex) PackedMeshSectionScalars1[SectionIndex].w
-#define MeshUVsChannelOffset(SectionIndex) PackedMeshSectionScalars2[SectionIndex].x
-#define MeshUVsChannelCount(SectionIndex) PackedMeshSectionScalars2[SectionIndex].y
-#define MeshSectionBufferIndex(SectionIndex) PackedMeshSectionScalars2[SectionIndex].z
-
-Buffer<uint> MeshIndexBuffer0;
-Buffer<uint> MeshIndexBuffer1;
-Buffer<uint> MeshIndexBuffer2;
-Buffer<uint> MeshIndexBuffer3;
-Buffer<uint> MeshIndexBuffer4;
-Buffer<uint> MeshIndexBuffer5;
-Buffer<uint> MeshIndexBuffer6;
-Buffer<uint> MeshIndexBuffer7;
+uint MaxUniqueTriangleCount;
+Buffer<float2> MeshUVsBuffer;
+Buffer<uint> MeshIndexBuffer;
+Buffer<uint> UniqueTriangleIndices;
+StructuredBuffer<uint4> MeshSectionBuffer;
 
 // Note on buffer input type
 // There are two different input buffer permutation, as if skin cache is enabled, we are using raw RHI SRV 
 // (while the skin cache system is not converted to RDG). If skin cache is disable, a manual skin compute 
 // update happens in the hair system, and its output are RDG buffers.
 #if PERMUTATION_POSITION_TYPE == 0
-
-Buffer<float> MeshPositionBuffer0;
-Buffer<float> MeshPositionBuffer1;
-Buffer<float> MeshPositionBuffer2;
-Buffer<float> MeshPositionBuffer3;
-Buffer<float> MeshPositionBuffer4;
-Buffer<float> MeshPositionBuffer5;
-Buffer<float> MeshPositionBuffer6;
-Buffer<float> MeshPositionBuffer7;
-
-Buffer<float> MeshPreviousPositionBuffer0;
-Buffer<float> MeshPreviousPositionBuffer1;
-Buffer<float> MeshPreviousPositionBuffer2;
-Buffer<float> MeshPreviousPositionBuffer3;
-Buffer<float> MeshPreviousPositionBuffer4;
-Buffer<float> MeshPreviousPositionBuffer5;
-Buffer<float> MeshPreviousPositionBuffer6;
-Buffer<float> MeshPreviousPositionBuffer7;
-
+Buffer<float> MeshPositionBuffer;
+Buffer<float> MeshPreviousPositionBuffer;
 #else
-
-Buffer<float> RDGMeshPositionBuffer0;
-Buffer<float> RDGMeshPositionBuffer1;
-Buffer<float> RDGMeshPositionBuffer2;
-Buffer<float> RDGMeshPositionBuffer3;
-Buffer<float> RDGMeshPositionBuffer4;
-Buffer<float> RDGMeshPositionBuffer5;
-Buffer<float> RDGMeshPositionBuffer6;
-Buffer<float> RDGMeshPositionBuffer7;
-
-Buffer<float> RDGMeshPreviousPositionBuffer0;
-Buffer<float> RDGMeshPreviousPositionBuffer1;
-Buffer<float> RDGMeshPreviousPositionBuffer2;
-Buffer<float> RDGMeshPreviousPositionBuffer3;
-Buffer<float> RDGMeshPreviousPositionBuffer4;
-Buffer<float> RDGMeshPreviousPositionBuffer5;
-Buffer<float> RDGMeshPreviousPositionBuffer6;
-Buffer<float> RDGMeshPreviousPositionBuffer7;
-
-#define MeshPositionBuffer0 RDGMeshPositionBuffer0
-#define MeshPositionBuffer1 RDGMeshPositionBuffer1
-#define MeshPositionBuffer2 RDGMeshPositionBuffer2
-#define MeshPositionBuffer3 RDGMeshPositionBuffer3
-#define MeshPositionBuffer4 RDGMeshPositionBuffer4
-#define MeshPositionBuffer5 RDGMeshPositionBuffer5
-#define MeshPositionBuffer6 RDGMeshPositionBuffer6
-#define MeshPositionBuffer7 RDGMeshPositionBuffer7
-
-#define MeshPreviousPositionBuffer0 RDGMeshPreviousPositionBuffer0
-#define MeshPreviousPositionBuffer1 RDGMeshPreviousPositionBuffer1
-#define MeshPreviousPositionBuffer2 RDGMeshPreviousPositionBuffer2
-#define MeshPreviousPositionBuffer3 RDGMeshPreviousPositionBuffer3
-#define MeshPreviousPositionBuffer4 RDGMeshPreviousPositionBuffer4
-#define MeshPreviousPositionBuffer5 RDGMeshPreviousPositionBuffer5
-#define MeshPreviousPositionBuffer6 RDGMeshPreviousPositionBuffer6
-#define MeshPreviousPositionBuffer7 RDGMeshPreviousPositionBuffer7
-
+Buffer<float> RDGMeshPositionBuffer;
+Buffer<float> RDGMeshPreviousPositionBuffer;
+#define MeshPositionBuffer RDGMeshPositionBuffer
+#define MeshPreviousPositionBuffer RDGMeshPreviousPositionBuffer
 #endif 
 
-#if PERMUTATION_WITHUV == 1
-Buffer<float2> MeshUVsBuffer0;
-Buffer<float2> MeshUVsBuffer1;
-Buffer<float2> MeshUVsBuffer2;
-Buffer<float2> MeshUVsBuffer3;
-Buffer<float2> MeshUVsBuffer4;
-Buffer<float2> MeshUVsBuffer5;
-Buffer<float2> MeshUVsBuffer6;
-Buffer<float2> MeshUVsBuffer7;
-Buffer<float2> MeshUVsBuffer8;
-#endif
-
-uint MaxUniqueTriangleCount;
-Buffer<uint> UniqueTriangleIndices;
-
 RWBuffer<float4> OutUniqueTriangleCurrPosition;
-
 #if PERMUTATION_PREVIOUS
 RWBuffer<float4> OutUniqueTrianglePrevPosition;
 #endif
 
-#if PERMUTATION_WITHUV == 1
-FHairMeshTriangleWithUVs GetTriangleData(FHairTriangleIndex TriangleIndex, uint SectionIndex, bool bPrevious)
+FMeshSectionData ReadMeshSectionData(uint InSectionIndex)
 {
-	#define GetTriangleDataAtIndex(Index) GetTriangleData(TriangleIndex.TriangleIndex, bPrevious ? MeshPreviousPositionBuffer##Index : MeshPositionBuffer##Index, MeshUVsBuffer##Index, MeshIndexBuffer##Index, MeshIndexOffset(SectionIndex), MeshMaxIndexCount(SectionIndex), MeshMaxVertexCount(SectionIndex), MeshUVsChannelOffset(SectionIndex), MeshUVsChannelCount(SectionIndex));
-
-	const uint MeshSectionIndex = MeshSectionBufferIndex(SectionIndex);
-		 if (MeshSectionIndex == 0) { return GetTriangleDataAtIndex(0); } 
-	else if (MeshSectionIndex == 1) { return GetTriangleDataAtIndex(1); }
-	else if (MeshSectionIndex == 2) { return GetTriangleDataAtIndex(2); }
-	else if (MeshSectionIndex == 3) { return GetTriangleDataAtIndex(3); }
-	else if (MeshSectionIndex == 4) { return GetTriangleDataAtIndex(4); }
-	else if (MeshSectionIndex == 5) { return GetTriangleDataAtIndex(5); }
-	else if (MeshSectionIndex == 6) { return GetTriangleDataAtIndex(6); }
-	else if (MeshSectionIndex == 7) { return GetTriangleDataAtIndex(7); }
-	
-	#undef GetTriangleDataAtIndex
-
-	return (FHairMeshTriangleWithUVs)0;
+	const uint4 Packed = MeshSectionBuffer[InSectionIndex];
+
+	FMeshSectionData Out;
+	Out.MaxIndexCount = Packed.x;
+	Out.MaxVertexCount = Packed.y;
+	Out.IndexOffset = Packed.z;
+	Out.UVsChannelOffset = BitFieldExtractU32(Packed.w, 8, 0);
+	Out.UVsChannelCount = BitFieldExtractU32(Packed.w, 8, 8);
+	return Out;
 }
-#else
-FHairMeshTriangle GetTriangleData(FHairTriangleIndex TriangleIndex, uint SectionIndex, bool bPrevious)
-{
-	#define GetTriangleDataAtIndex(Index) GetTriangleData(TriangleIndex.TriangleIndex, bPrevious ? MeshPreviousPositionBuffer##Index : MeshPositionBuffer##Index, MeshIndexBuffer##Index, MeshIndexOffset(SectionIndex), MeshMaxIndexCount(SectionIndex), MeshMaxVertexCount(SectionIndex));
 
-	const uint MeshSectionIndex = MeshSectionBufferIndex(SectionIndex);
-		 if (MeshSectionIndex == 0) { return GetTriangleDataAtIndex(0); }
-	else if (MeshSectionIndex == 1) { return GetTriangleDataAtIndex(1); }
-	else if (MeshSectionIndex == 2) { return GetTriangleDataAtIndex(2); }
-	else if (MeshSectionIndex == 3) { return GetTriangleDataAtIndex(3); }
-	else if (MeshSectionIndex == 4) { return GetTriangleDataAtIndex(4); }
-	else if (MeshSectionIndex == 5) { return GetTriangleDataAtIndex(5); }
-	else if (MeshSectionIndex == 6) { return GetTriangleDataAtIndex(6); }
-	else if (MeshSectionIndex == 7) { return GetTriangleDataAtIndex(7); }
-
-	#undef GetTriangleDataAtIndex
-
-	return (FHairMeshTriangle)0;
+FHairMeshTriangleWithUVs GetTriangleData(FHairTriangleIndex TriangleIndex, FMeshSectionData MeshSectionData, bool bPrevious)
+{
+	return GetTriangleData(
+		TriangleIndex.TriangleIndex, 
+		bPrevious ? MeshPreviousPositionBuffer : MeshPositionBuffer, 
+		MeshUVsBuffer, 
+		MeshIndexBuffer, 
+		MeshSectionData.IndexOffset,
+		MeshSectionData.MaxIndexCount,
+		MeshSectionData.MaxVertexCount,
+		MeshSectionData.UVsChannelOffset, 
+		MeshSectionData.UVsChannelCount);
 }
-#endif
 
-[numthreads(128, 1, 1)]
+[numthreads(GROUP_SIZE, 1, 1)]
 void MainCS(uint3 DispatchThreadId : SV_DispatchThreadID)
 { 
 	const uint UniqueTriangleIndex = DispatchThreadId.x;
 	if (UniqueTriangleIndex >= MaxUniqueTriangleCount)
+	{
 		return;
+	}
 	
 	const FHairTriangleIndex TriangleIndex = UnpackTriangleIndex(UniqueTriangleIndices[UniqueTriangleIndex]);
-	if (TriangleIndex.SectionIndex >= MaxSectionCount)
-		return;
-
-	const uint SectionIndex = TriangleIndex.SectionIndex;
-
-	bool bValid = false;
-	uint SectionSlotIndex = 0;
-	for (uint It=0;It<Pass_SectionCount;++It)
+	if (TriangleIndex.SectionIndex < MaxSectionCount)
 	{
-		if (MeshSectionIndices(It) == SectionIndex)
+		const FMeshSectionData MeshSectionData = ReadMeshSectionData(TriangleIndex.SectionIndex);
+	
+		// Compute the slot index of the resources for the section index (this is done in multiple pass
 		{
-			SectionSlotIndex = It;
-			bValid = true;
-			break;
+			const FHairMeshTriangleWithUVs Tri = GetTriangleData(TriangleIndex, MeshSectionData, false);
+			OutUniqueTriangleCurrPosition[UniqueTriangleIndex * 3 + 0] = float4(Tri.P0, PackUVsToFloat(Tri.UV0));
+			OutUniqueTriangleCurrPosition[UniqueTriangleIndex * 3 + 1] = float4(Tri.P1, PackUVsToFloat(Tri.UV1));
+			OutUniqueTriangleCurrPosition[UniqueTriangleIndex * 3 + 2] = float4(Tri.P2, PackUVsToFloat(Tri.UV2));
 		}
+		#if PERMUTATION_PREVIOUS
+		{
+			const FHairMeshTriangleWithUVs Tri = GetTriangleData(TriangleIndex, MeshSectionData, true);
+			OutUniqueTrianglePrevPosition[UniqueTriangleIndex * 3 + 0] = float4(Tri.P0, PackUVsToFloat(Tri.UV0));
+			OutUniqueTrianglePrevPosition[UniqueTriangleIndex * 3 + 1] = float4(Tri.P1, PackUVsToFloat(Tri.UV1));
+			OutUniqueTrianglePrevPosition[UniqueTriangleIndex * 3 + 2] = float4(Tri.P2, PackUVsToFloat(Tri.UV2));
+		}
+		#endif
 	}
-	
-	if (!bValid)
-		return;
-
-	// Compute the slod index of the resources for the section index (this is done in multiple pass
-#if PERMUTATION_WITHUV == 1
-	{
-		const FHairMeshTriangleWithUVs Tri = GetTriangleData(TriangleIndex, SectionSlotIndex, false);
-		OutUniqueTriangleCurrPosition[UniqueTriangleIndex * 3 + 0] = float4(Tri.P0, PackUVsToFloat(Tri.UV0));
-		OutUniqueTriangleCurrPosition[UniqueTriangleIndex * 3 + 1] = float4(Tri.P1, PackUVsToFloat(Tri.UV1));
-		OutUniqueTriangleCurrPosition[UniqueTriangleIndex * 3 + 2] = float4(Tri.P2, PackUVsToFloat(Tri.UV2));
-	}
-
-	#if PERMUTATION_PREVIOUS
-	{
-		const FHairMeshTriangleWithUVs Tri = GetTriangleData(TriangleIndex, SectionSlotIndex, true);
-		OutUniqueTrianglePrevPosition[UniqueTriangleIndex * 3 + 0] = float4(Tri.P0, PackUVsToFloat(Tri.UV0));
-		OutUniqueTrianglePrevPosition[UniqueTriangleIndex * 3 + 1] = float4(Tri.P1, PackUVsToFloat(Tri.UV1));
-		OutUniqueTrianglePrevPosition[UniqueTriangleIndex * 3 + 2] = float4(Tri.P2, PackUVsToFloat(Tri.UV2));
-	}
-	#endif
-#else
-	{
-		const FHairMeshTriangle Tri = GetTriangleData(TriangleIndex, SectionSlotIndex, false);
-		OutUniqueTriangleCurrPosition[UniqueTriangleIndex * 3 + 0] = (float4(Tri.P0, 1));
-		OutUniqueTriangleCurrPosition[UniqueTriangleIndex * 3 + 1] = (float4(Tri.P1, 1));
-		OutUniqueTriangleCurrPosition[UniqueTriangleIndex * 3 + 2] = (float4(Tri.P2, 1));
-	}
-
-	#if PERMUTATION_PREVIOUS
-	{
-		const FHairMeshTriangle Tri = GetTriangleData(TriangleIndex, SectionSlotIndex, true);
-		OutUniqueTrianglePrevPosition[UniqueTriangleIndex * 3 + 0] = (float4(Tri.P0, 1));
-		OutUniqueTrianglePrevPosition[UniqueTriangleIndex * 3 + 1] = (float4(Tri.P1, 1));
-		OutUniqueTrianglePrevPosition[UniqueTriangleIndex * 3 + 2] = (float4(Tri.P2, 1));
-	}
-	#endif
-#endif
 }
 
 #endif // SHADER_MESH_UPDATE
