Fix incorrect root dispatch.

diff --git a/Engine/Plugins/Runtime/HairStrands/Source/HairStrandsCore/Private/HairStrandsMeshProjection.cpp b/Engine/Plugins/Runtime/HairStrands/Source/HairStrandsCore/Private/HairStrandsMeshProjection.cpp
--- a/Engine/Plugins/Runtime/HairStrands/Source/HairStrandsCore/Private/HairStrandsMeshProjection.cpp
+++ b/Engine/Plugins/Runtime/HairStrands/Source/HairStrandsCore/Private/HairStrandsMeshProjection.cpp
@@ -205,43 +205,27 @@ public:
 
 IMPLEMENT_GLOBAL_SHADER(FHairUpdateMeshTriangleCS, "/Engine/Private/HairStrands/HairStrandsMesh.usf", "MainCS", SF_Compute);
 
-void AddHairStrandUpdateMeshTrianglesPass(
+bool AddHairStrandUpdateMeshTrianglesPass(
 	FRDGBuilder& GraphBuilder,
 	FGlobalShaderMap* ShaderMap,
-	const int32 MeshLODIndex,
+	const uint32 UniqueTriangleCount, 
 	const FHairStrandsProjectionMeshData::LOD& MeshData,
-	FHairStrandsRestRootResource* RestResources,
-	FHairStrandsDeformedRootResource* DeformedResources)
+	FRDGBufferSRVRef UniqueTriangleIndexSRV,
+	FRDGBufferUAVRef OutputCurrUAV,
+	FRDGBufferUAVRef OutputPrevUAV)
 {
-	const uint32 RootCount = RestResources->GetRootCount();
-	if (RootCount == 0 || !RestResources->LODs.IsValidIndex(MeshLODIndex))
-	{
-		return;
-	}
-
-	FHairStrandsRestRootResource::FLOD& RestLODData = RestResources->LODs[MeshLODIndex];
-	check(RestLODData.LODIndex == MeshLODIndex);
-
 	const uint32 SectionCount = MeshData.Sections.Num();
-	const uint32 MaxSupportedSectionCount = GetHairStrandsMaxSectionCount();
-	check(SectionCount < MaxSupportedSectionCount);
-	if (SectionCount == 0 || SectionCount >= MaxSupportedSectionCount)
-	{
-		return;
-	}
 
 	// If a skel. mesh is not streaming yet, its SRV will be null
+	const bool bValid = SectionCount > 0 && SectionCount < GetHairStrandsMaxSectionCount();
 	const bool bReady = MeshData.Sections[0].IndexBuffer != nullptr && MeshData.Sections[0].UVsBuffer != nullptr && (MeshData.Sections[0].RDGPositionBuffer != nullptr || MeshData.Sections[0].PositionBuffer != nullptr);
-	if (!bReady)
+	if (!bReady || !bValid)
 	{
-		return;
+		return false;
 	}
 
-	// Update the last known mesh LOD for which the root resources has been updated
-	DeformedResources->MeshLODIndex = MeshLODIndex;
-
 	FHairUpdateMeshTriangleCS::FParameters CommonParameters;
-	CommonParameters.MaxUniqueTriangleCount 		= RootCount;
+	CommonParameters.MaxUniqueTriangleCount 		= UniqueTriangleCount;
 	CommonParameters.MaxSectionCount 				= SectionCount;
 	CommonParameters.RDGMeshPositionBuffer			= MeshData.Sections[0].RDGPositionBuffer;
 	CommonParameters.RDGMeshPreviousPositionBuffer 	= MeshData.Sections[0].RDGPreviousPositionBuffer;
@@ -249,24 +233,9 @@ void AddHairStrandUpdateMeshTrianglesPass(
 	CommonParameters.MeshPreviousPositionBuffer		= MeshData.Sections[0].PreviousPositionBuffer;
 	CommonParameters.MeshIndexBuffer				= MeshData.Sections[0].IndexBuffer;
 	CommonParameters.MeshUVsBuffer					= MeshData.Sections[0].UVsBuffer;
-
-	FHairStrandsDeformedRootResource::FLOD& DeformedLODData = DeformedResources->LODs[MeshLODIndex];
-	DeformedLODData.Status = FHairStrandsDeformedRootResource::FLOD::EStatus::Completed;
-
-	CommonParameters.UniqueTriangleIndices = RegisterAsSRV(GraphBuilder, RestLODData.UniqueTriangleIndexBuffer);
-
-	// Current 
-	FRDGImportedBuffer OutputCurrBuffer = Register(GraphBuilder, DeformedLODData.GetDeformedUniqueTrianglePositionBuffer(FHairStrandsDeformedRootResource::FLOD::Current), ERDGImportedBufferFlags::CreateUAV, ERDGUnorderedAccessViewFlags::SkipBarrier);
-	CommonParameters.OutUniqueTriangleCurrPosition = OutputCurrBuffer.UAV;
-
-	// Previous
-	FRDGImportedBuffer OutputPrevBuffer;
-	const bool bComputePreviousDeformedPosition = IsHairStrandContinuousDecimationReorderingEnabled();
-	if (bComputePreviousDeformedPosition)
-	{
-		OutputPrevBuffer = Register(GraphBuilder, DeformedLODData.GetDeformedUniqueTrianglePositionBuffer(FHairStrandsDeformedRootResource::FLOD::Previous), ERDGImportedBufferFlags::CreateUAV, ERDGUnorderedAccessViewFlags::SkipBarrier);
-		CommonParameters.OutUniqueTrianglePrevPosition = OutputPrevBuffer.UAV;
-	}
+	CommonParameters.UniqueTriangleIndices 			= UniqueTriangleIndexSRV;
+	CommonParameters.OutUniqueTriangleCurrPosition	= OutputCurrUAV;
+	CommonParameters.OutUniqueTrianglePrevPosition	= OutputPrevUAV;
 
 	struct FSectionData
 	{
@@ -303,6 +272,7 @@ void AddHairStrandUpdateMeshTrianglesPass(
 	CommonParameters.MeshSectionBuffer = GraphBuilder.CreateSRV(SectionBuffer);
 
 	const bool bUseRDGPositionBuffer = CommonParameters.RDGMeshPositionBuffer != nullptr;
+	const bool bComputePreviousDeformedPosition = OutputPrevUAV != nullptr;
 	{
 		FHairUpdateMeshTriangleCS::FParameters* PassParameters = GraphBuilder.AllocParameters<FHairUpdateMeshTriangleCS::FParameters>();
 		*PassParameters = CommonParameters;
@@ -311,7 +281,7 @@ void AddHairStrandUpdateMeshTrianglesPass(
 		PermutationVector.Set<FHairUpdateMeshTriangleCS::FPrevious>(bComputePreviousDeformedPosition);
 		PermutationVector.Set<FHairUpdateMeshTriangleCS::FPositionType>(bUseRDGPositionBuffer ? 1 : 0);
 
-		const FIntVector DispatchGroupCount = FComputeShaderUtils::GetGroupCount(RootCount, FHairUpdateMeshTriangleCS::GetGroupSize());
+		const FIntVector DispatchGroupCount = FComputeShaderUtils::GetGroupCount(UniqueTriangleCount, FHairUpdateMeshTriangleCS::GetGroupSize());
 		check(DispatchGroupCount.X <= GRHIMaxDispatchThreadGroupsPerDimension.X);
 		TShaderMapRef<FHairUpdateMeshTriangleCS> ComputeShader(ShaderMap, PermutationVector);
 		FComputeShaderUtils::AddPass(
@@ -322,13 +292,59 @@ void AddHairStrandUpdateMeshTrianglesPass(
 			DispatchGroupCount);
 	}
 
-	GraphBuilder.SetBufferAccessFinal(OutputCurrBuffer.Buffer, ERHIAccess::SRVMask);
+	return true;
+}
+
+void AddHairStrandUpdateMeshTrianglesPass(
+	FRDGBuilder& GraphBuilder,
+	FGlobalShaderMap* ShaderMap,
+	const int32 MeshLODIndex,
+	const FHairStrandsProjectionMeshData::LOD& MeshData,
+	FHairStrandsRestRootResource* RestResources,
+	FHairStrandsDeformedRootResource* DeformedResources)
+{	
+	if (RestResources->GetRootCount() == 0 || !RestResources->LODs.IsValidIndex(MeshLODIndex))
+	{
+		return;
+	}
+
+	FHairStrandsRestRootResource::FLOD& RestLODData = RestResources->LODs[MeshLODIndex];
+	check(RestLODData.LODIndex == MeshLODIndex);
+	check(RestResources->BulkData.Header.LODs.IsValidIndex(MeshLODIndex));
+	const uint32 UniqueTriangleCount = RestResources->BulkData.Header.LODs[MeshLODIndex].UniqueTriangleCount;
+	if (UniqueTriangleCount == 0)
+	{
+		return;
+	}
 
+	FRDGImportedBuffer OutputCurrBuffer = Register(GraphBuilder, DeformedResources->LODs[MeshLODIndex].GetDeformedUniqueTrianglePositionBuffer(FHairStrandsDeformedRootResource::FLOD::Current), ERDGImportedBufferFlags::CreateUAV, ERDGUnorderedAccessViewFlags::SkipBarrier);
+	FRDGImportedBuffer OutputPrevBuffer;
+	const bool bComputePreviousDeformedPosition = IsHairStrandContinuousDecimationReorderingEnabled();
 	if (bComputePreviousDeformedPosition)
 	{
-		GraphBuilder.SetBufferAccessFinal(OutputPrevBuffer.Buffer, ERHIAccess::SRVMask);
+		Register(GraphBuilder, DeformedResources->LODs[MeshLODIndex].GetDeformedUniqueTrianglePositionBuffer(FHairStrandsDeformedRootResource::FLOD::Previous), ERDGImportedBufferFlags::CreateUAV, ERDGUnorderedAccessViewFlags::SkipBarrier);
 	}
 
+	if (AddHairStrandUpdateMeshTrianglesPass(
+		GraphBuilder,
+		ShaderMap,
+		UniqueTriangleCount,
+		MeshData,
+		RegisterAsSRV(GraphBuilder, RestLODData.UniqueTriangleIndexBuffer),
+		OutputCurrBuffer.UAV,
+		OutputPrevBuffer.UAV))
+	{
+
+		GraphBuilder.SetBufferAccessFinal(OutputCurrBuffer.Buffer, ERHIAccess::SRVMask);
+		if (OutputPrevBuffer.Buffer)
+		{
+			GraphBuilder.SetBufferAccessFinal(OutputPrevBuffer.Buffer, ERHIAccess::SRVMask);
+		}
+	
+		// Update the last known mesh LOD for which the root resources has been updated
+		DeformedResources->MeshLODIndex = MeshLODIndex;
+		DeformedResources->LODs[MeshLODIndex].Status = FHairStrandsDeformedRootResource::FLOD::EStatus::Completed;	
+	}
 }
 
 ///////////////////////////////////////////////////////////////////////////////////////////////////
