Fix RHI validation error in hair code.

diff --git a/Engine/Plugins/Runtime/HairStrands/Source/HairStrandsCore/Private/GroomManager.cpp b/Engine/Plugins/Runtime/HairStrands/Source/HairStrandsCore/Private/GroomManager.cpp
--- a/Engine/Plugins/Runtime/HairStrands/Source/HairStrandsCore/Private/GroomManager.cpp
+++ b/Engine/Plugins/Runtime/HairStrands/Source/HairStrandsCore/Private/GroomManager.cpp
@@ -572,6 +572,7 @@ static void RunHairStrandsInterpolation_Guide(
 		AddDeformSimHairStrandsPass(
 			GraphBuilder,
 			ShaderMap,
+			InstanceData.Instance->RegisteredIndex,
 			InstanceData.Instance->Debug.MeshLODIndex,
 			InstanceData.Instance->Guides.RestResource->GetPointCount(),
 			InstanceData.Instance->Guides.RestRootResource,
@@ -606,6 +607,7 @@ static void RunHairStrandsInterpolation_Guide(
 						AddDeformSimHairStrandsPass(
 							GraphBuilder,
 							ShaderMap,
+							InstanceData.Instance->RegisteredIndex,
 							InstanceData.MeshLODIndex,
 							InstanceData.Instance->Guides.RestResource->GetPointCount(),
 							InstanceData.Instance->Guides.RestRootResource,
diff --git a/Engine/Plugins/Runtime/HairStrands/Source/HairStrandsCore/Private/GroomResources.cpp b/Engine/Plugins/Runtime/HairStrands/Source/HairStrandsCore/Private/GroomResources.cpp
--- a/Engine/Plugins/Runtime/HairStrands/Source/HairStrandsCore/Private/GroomResources.cpp
+++ b/Engine/Plugins/Runtime/HairStrands/Source/HairStrandsCore/Private/GroomResources.cpp
@@ -545,6 +545,23 @@ FRDGBufferRef InternalCreateStructuredBufferRDG(FRDGBuilder& GraphBuilder, uint3
 	return Buffer;
 }
 
+template<typename DataType>
+FRDGBufferRef InternalCreateStructuredBufferRDG(FRDGBuilder& GraphBuilder, const TArray<DataType>& InData, FRDGExternalBuffer& Out, const TCHAR* DebugName, const FName& OwnerName, EHairResourceUsageType UsageType, ERDGInitialDataFlags InitialDataFlags)
+{
+	const uint32 DataCount = InData.Num();
+	const uint32 DataSizeInBytes = sizeof(DataType) * DataCount;
+	if (DataSizeInBytes == 0)
+	{
+		Out.Buffer = nullptr;
+		return nullptr;
+	}
+
+	FRDGBufferRef Buffer = CreateStructuredBuffer(GraphBuilder, DebugName, sizeof(DataType), InData.Num(), InData.GetData(), sizeof(DataType) * InData.Num());
+	Buffer->SetOwnerName(OwnerName);
+	ConvertToExternalBufferWithViews(GraphBuilder, Buffer, Out, EPixelFormat::PF_Unknown);
+	return Buffer;
+}
+
 template<typename DataType>
 FRDGBufferRef InternalCreateByteAddressBufferRDG(FRDGBuilder& GraphBuilder, const TArray<DataType>& InData, EPixelFormat Format, FRDGExternalBuffer& Out, const TCHAR* DebugName, const FName& OwnerName, EHairResourceUsageType UsageType)
 {
@@ -1058,7 +1075,7 @@ void FHairStrandsRestResource::InternalAllocate(FRDGBuilder& GraphBuilder)
 	{	
 		TArray<FVector4f> RestOffset;
 		RestOffset.Add((FVector3f)BulkData.GetPositionOffset());// LWC_TODO: precision loss
-		InternalCreateVertexBufferRDG<FHairStrandsPositionOffsetFormat>(GraphBuilder, RestOffset, PositionOffsetBuffer, ToHairResourceDebugName(HAIRSTRANDS_RESOUCE_NAME(CurveType, Hair.StrandsRest_PositionOffsetBuffer), ResourceName), OwnerName, EHairResourceUsageType::Static, ERDGInitialDataFlags::None);
+		InternalCreateStructuredBufferRDG<FHairStrandsPositionOffsetFormat::Type>(GraphBuilder, RestOffset, PositionOffsetBuffer, ToHairResourceDebugName(HAIRSTRANDS_RESOUCE_NAME(CurveType, Hair.StrandsRest_PositionOffsetBuffer), ResourceName), OwnerName, EHairResourceUsageType::Static, ERDGInitialDataFlags::None);
 		GraphBuilder.UseExternalAccessMode(Register(GraphBuilder, PositionOffsetBuffer, ERDGImportedBufferFlags::CreateSRV).Buffer, ERHIAccess::SRVMask);
 	}
 }
@@ -1152,8 +1169,8 @@ void FHairStrandsDeformedResource::InternalAllocate(FRDGBuilder& GraphBuilder, u
 	{
 		TArray<FVector4f> DefaultOffsets;
 		DefaultOffsets.Add((FVector3f)BulkData.GetPositionOffset()); // LWC_TODO: precision loss
-		InternalCreateVertexBufferRDG<FHairStrandsPositionOffsetFormat>(GraphBuilder, DefaultOffsets, DeformedOffsetBuffer[0], ToHairResourceDebugName(HAIRSTRANDS_RESOUCE_NAME(CurveType, Hair.StrandsDeformed_DeformedOffsetBuffer0), ResourceName), OwnerName, EHairResourceUsageType::Dynamic, ERDGInitialDataFlags::None);
-		InternalCreateVertexBufferRDG<FHairStrandsPositionOffsetFormat>(GraphBuilder, DefaultOffsets, DeformedOffsetBuffer[1], ToHairResourceDebugName(HAIRSTRANDS_RESOUCE_NAME(CurveType, Hair.StrandsDeformed_DeformedOffsetBuffer1), ResourceName), OwnerName, EHairResourceUsageType::Dynamic, ERDGInitialDataFlags::None);
+		InternalCreateStructuredBufferRDG<FHairStrandsPositionOffsetFormat::Type>(GraphBuilder, DefaultOffsets, DeformedOffsetBuffer[0], ToHairResourceDebugName(HAIRSTRANDS_RESOUCE_NAME(CurveType, Hair.StrandsDeformed_DeformedOffsetBuffer0), ResourceName), OwnerName, EHairResourceUsageType::Dynamic, ERDGInitialDataFlags::None);
+		InternalCreateStructuredBufferRDG<FHairStrandsPositionOffsetFormat::Type>(GraphBuilder, DefaultOffsets, DeformedOffsetBuffer[1], ToHairResourceDebugName(HAIRSTRANDS_RESOUCE_NAME(CurveType, Hair.StrandsDeformed_DeformedOffsetBuffer1), ResourceName), OwnerName, EHairResourceUsageType::Dynamic, ERDGInitialDataFlags::None);
 	}
 	// Note: DeformerBuffer is optionally/lazily allocated by a mesh-deformer graph
 }
@@ -1502,13 +1519,13 @@ void FHairStrandsDeformedRootResource::InternalAllocate(FRDGBuilder& GraphBuilde
 				const uint32 RBFEntryCount = FGroomRBFDeformer::GetEntryCount(LOD.SampleCount);
 
 				InternalCreateVertexBufferRDG<FHairStrandsMeshTrianglePositionFormat>(GraphBuilder, LOD.SampleCount, LOD.DeformedSamplePositionsBuffer[0], ToHairResourceDebugName(HAIRSTRANDS_RESOUCE_NAME(CurveType, Hair.StrandsRootDeformed_DeformedSamplePositionsBuffer0), ResourceName), OwnerName, EHairResourceUsageType::Dynamic);
-				InternalCreateVertexBufferRDG<FHairStrandsMeshTrianglePositionFormat>(GraphBuilder, RBFEntryCount, LOD.MeshSampleWeightsBuffer[0], ToHairResourceDebugName(HAIRSTRANDS_RESOUCE_NAME(CurveType, Hair.StrandsRootDeformed_MeshSampleWeightsBuffer0), ResourceName), OwnerName, EHairResourceUsageType::Dynamic);
+				InternalCreateStructuredBufferRDG<FHairStrandsMeshTrianglePositionFormat>(GraphBuilder, RBFEntryCount, LOD.MeshSampleWeightsBuffer[0], ToHairResourceDebugName(HAIRSTRANDS_RESOUCE_NAME(CurveType, Hair.StrandsRootDeformed_MeshSampleWeightsBuffer0), ResourceName), OwnerName, EHairResourceUsageType::Dynamic);
 
 				// Double buffering is disabled by default unless the read-only cvar r.HairStrands.ContinuousDecimationReordering is set
 				if (IsHairStrandContinuousDecimationReorderingEnabled())
 				{
 					InternalCreateVertexBufferRDG<FHairStrandsMeshTrianglePositionFormat>(GraphBuilder, LOD.SampleCount, LOD.DeformedSamplePositionsBuffer[1], ToHairResourceDebugName(HAIRSTRANDS_RESOUCE_NAME(CurveType, Hair.StrandsRootDeformed_DeformedSamplePositionsBuffer1), ResourceName), OwnerName, EHairResourceUsageType::Dynamic);
-					InternalCreateVertexBufferRDG<FHairStrandsMeshTrianglePositionFormat>(GraphBuilder, RBFEntryCount, LOD.MeshSampleWeightsBuffer[1], ToHairResourceDebugName(HAIRSTRANDS_RESOUCE_NAME(CurveType, Hair.StrandsRootDeformed_MeshSampleWeightsBuffer1), ResourceName), OwnerName, EHairResourceUsageType::Dynamic);
+					InternalCreateStructuredBufferRDG<FHairStrandsMeshTrianglePositionFormat>(GraphBuilder, RBFEntryCount, LOD.MeshSampleWeightsBuffer[1], ToHairResourceDebugName(HAIRSTRANDS_RESOUCE_NAME(CurveType, Hair.StrandsRootDeformed_MeshSampleWeightsBuffer1), ResourceName), OwnerName, EHairResourceUsageType::Dynamic);
 				}
 			}
 
diff --git a/Engine/Plugins/Runtime/HairStrands/Source/HairStrandsCore/Private/HairStrandsInterpolation.cpp b/Engine/Plugins/Runtime/HairStrands/Source/HairStrandsCore/Private/HairStrandsInterpolation.cpp
--- a/Engine/Plugins/Runtime/HairStrands/Source/HairStrandsCore/Private/HairStrandsInterpolation.cpp
+++ b/Engine/Plugins/Runtime/HairStrands/Source/HairStrandsCore/Private/HairStrandsInterpolation.cpp
@@ -302,6 +302,7 @@ class FDeformGuideCS : public FGlobalShader
 	using FPermutationDomain = TShaderPermutationDomain<FDeformationType>;
 
 	BEGIN_SHADER_PARAMETER_STRUCT(FParameters, )
+		SHADER_PARAMETER(uint32, InstanceRegisteredIndex)
 		SHADER_PARAMETER(uint32, VertexCount)
 		SHADER_PARAMETER(FVector3f, SimRestOffset)
 		
@@ -312,7 +313,7 @@ class FDeformGuideCS : public FGlobalShader
 		SHADER_PARAMETER_RDG_BUFFER_SRV(Buffer<uint>, SimPointToCurveBuffer)
 		SHADER_PARAMETER_RDG_BUFFER_SRV(Buffer<uint>, SimRootToUniqueTriangleIndexBuffer)
 
-		SHADER_PARAMETER_RDG_BUFFER_SRV(Buffer, SimDeformedOffsetBuffer)
+		SHADER_PARAMETER_RDG_BUFFER_SRV(StructuredBuffer, SimDeformedOffsetBuffer)
 		SHADER_PARAMETER_RDG_BUFFER_SRV(Buffer, SimRestPosePositionBuffer)
 		SHADER_PARAMETER_RDG_BUFFER_UAV(RWBuffer, OutSimDeformedPositionBuffer)
 
@@ -339,6 +340,7 @@ IMPLEMENT_GLOBAL_SHADER(FDeformGuideCS, "/Engine/Private/HairStrands/HairStrands
 void AddDeformSimHairStrandsPass(
 	FRDGBuilder& GraphBuilder,
 	FGlobalShaderMap* ShaderMap,
+	const uint32 InstanceRegisteredIndex,
 	const uint32 MeshLODIndex,
 	const uint32 VertexCount,
 	FHairStrandsRestRootResource* SimRestRootResources,
@@ -364,6 +366,7 @@ void AddDeformSimHairStrandsPass(
 	EInternalDeformationType InternalDeformationType = InternalDeformationType_Offset;
 
 	FDeformGuideCS::FParameters* Parameters = GraphBuilder.AllocParameters<FDeformGuideCS::FParameters>();
+	Parameters->InstanceRegisteredIndex = InstanceRegisteredIndex;
 	Parameters->SimRestPosePositionBuffer = SimRestPosePositionBuffer;
 	Parameters->OutSimDeformedPositionBuffer = OutSimDeformedPositionBuffer.UAV;
 	Parameters->VertexCount = VertexCount;
@@ -800,7 +803,7 @@ class FHairClusterAABBCS : public FGlobalShader
 		SHADER_PARAMETER_RDG_BUFFER_SRV(Buffer<uint>, RenCurveBuffer)
 		SHADER_PARAMETER_RDG_BUFFER_SRV(Buffer<uint>, RenPointLODBuffer)
 		SHADER_PARAMETER_RDG_BUFFER_SRV(Buffer, RenderDeformedPositionBuffer)
-		SHADER_PARAMETER_RDG_BUFFER_SRV(Buffer, RenderDeformedOffsetBuffer)
+		SHADER_PARAMETER_RDG_BUFFER_SRV(StructuredBuffer, RenderDeformedOffsetBuffer)
 		SHADER_PARAMETER_RDG_BUFFER_UAV(RWBuffer, OutClusterAABBBuffer)
 		SHADER_PARAMETER_RDG_BUFFER_UAV(RWBuffer, OutGroupAABBBuffer)
 	END_SHADER_PARAMETER_STRUCT()
diff --git a/Engine/Plugins/Runtime/HairStrands/Source/HairStrandsCore/Public/HairStrandsInterpolation.h b/Engine/Plugins/Runtime/HairStrands/Source/HairStrandsCore/Public/HairStrandsInterpolation.h
--- a/Engine/Plugins/Runtime/HairStrands/Source/HairStrandsCore/Public/HairStrandsInterpolation.h
+++ b/Engine/Plugins/Runtime/HairStrands/Source/HairStrandsCore/Public/HairStrandsInterpolation.h
@@ -44,6 +44,7 @@ enum class EGroomViewMode : uint8;
 void AddDeformSimHairStrandsPass(
 	FRDGBuilder& GraphBuilder,
 	FGlobalShaderMap* ShaderMap,
+	const uint32 InstanceRegisteredIndex,
 	const uint32 MeshLODIndex,
 	const uint32 VertexCount,
 	FHairStrandsRestRootResource* SimRestRootResources,
diff --git a/Engine/Shaders/Private/HairStrands/HairStrandsClusterCulling.usf b/Engine/Shaders/Private/HairStrands/HairStrandsClusterCulling.usf
--- a/Engine/Shaders/Private/HairStrands/HairStrandsClusterCulling.usf
+++ b/Engine/Shaders/Private/HairStrands/HairStrandsClusterCulling.usf
@@ -466,7 +466,7 @@ float4x4 LocalToTranslatedWorldMatrix;
 Buffer<uint> RenCurveBuffer;
 Buffer<uint> RenPointLODBuffer;
 Buffer<uint4> RenderDeformedPositionBuffer;
-Buffer<float4> RenderDeformedOffsetBuffer;
+StructuredBuffer<float4> RenderDeformedOffsetBuffer;
 
 RWBuffer<int> OutClusterAABBBuffer; 	// Cluster data packed as {uint3 AABBMin, uint3 AABBMax}
 RWBuffer<int> OutGroupAABBBuffer; 		// Group data packed as {uint3 AABBMin, uint3 AABBMax}
@@ -506,7 +506,7 @@ groupshared float3 OutHairPositionOffset;
 			CurvePointIndex = RenCurve.PointIndex;
 			CurvePointCount = RenCurve.PointCount;
 		}
-		OutHairPositionOffset = RenderDeformedOffsetBuffer[0].xyz;
+		OutHairPositionOffset = ReadRenPositionOffset(RenderDeformedOffsetBuffer, InstanceRegisteredIndex);
 	}
 	GroupMemoryBarrierWithGroupSync();
 
diff --git a/Engine/Shaders/Private/HairStrands/HairStrandsGuideDeform.usf b/Engine/Shaders/Private/HairStrands/HairStrandsGuideDeform.usf
--- a/Engine/Shaders/Private/HairStrands/HairStrandsGuideDeform.usf
+++ b/Engine/Shaders/Private/HairStrands/HairStrandsGuideDeform.usf
@@ -3,17 +3,18 @@
 #include "../Common.ush"
 #include "HairStrandsClusterCommon.ush"
 #include "HairStrandsBindingCommon.ush"
+#include "HairStrandsVertexFactoryCommon.ush"
 
 #if SHADER_DEFORM_GUIDE
 
 float3 SimRestOffset;
 uint VertexCount;
+uint InstanceRegisteredIndex;
 
-Buffer<float4> SimDeformedOffsetBuffer;
+StructuredBuffer<float4> SimDeformedOffsetBuffer;
 Buffer<uint4> SimRestPosePositionBuffer;
 RWBuffer<uint4> OutSimDeformedPositionBuffer;
 
-
 #define DEFORMATION_BYPASS 0
 #define DEFORMATION_OFFSET_GUIDE 1
 #define DEFORMATION_OFFSET_GUIDE_DYNAMIC 2
@@ -60,7 +61,7 @@ StructuredBuffer<float4>  MeshSampleWeightsBuffer;
 
 float3 DisplacePosition(uint VertexIndex, float3 Pos)
 {
-	const float3 OutSimHairPositionOffset = SimDeformedOffsetBuffer[0].xyz;
+	const float3 OutSimHairPositionOffset = ReadSimPositionOffset(SimDeformedOffsetBuffer, InstanceRegisteredIndex);
 
 // By pass
 #if PERMUTATION_DEFORMATION == DEFORMATION_BYPASS
