* Add resource transition pass for forcing hair passes to be batched avoid avoid unecessary wait.
* Add skip UAV barrier on cluster computation.

Reduce HairInterpolation cost from 0.48ms -> 0.26ms

diff --git a/Engine/Plugins/Runtime/HairStrands/Source/HairStrandsCore/Private/GroomManager.cpp b/Engine/Plugins/Runtime/HairStrands/Source/HairStrandsCore/Private/GroomManager.cpp
--- a/Engine/Plugins/Runtime/HairStrands/Source/HairStrandsCore/Private/GroomManager.cpp
+++ b/Engine/Plugins/Runtime/HairStrands/Source/HairStrandsCore/Private/GroomManager.cpp
@@ -849,12 +849,17 @@ static void RunHairStrandsInterpolation_Strands(
 	// Early culling
 	FHairStrandClusterData ClusterDatas;
 	{
+		TArray<FRDGBufferSRVRef> Transitions;
+		Transitions.Reserve(InstanceDatas.Num());
+
 		// Gather culling jobs
 		for (FInstanceData& InstanceData : InstanceDatas)
 		{
 			if (InstanceData.Instance->Strands.bCullingEnable)
 			{
 				AddInstanceToClusterData(InstanceData.Instance, ClusterDatas);
+				Transitions.Add(RegisterAsSRV(GraphBuilder, InstanceData.Instance->HairGroupPublicData->GetCulledVertexIdBuffer()));
+				Transitions.Add(RegisterAsSRV(GraphBuilder, InstanceData.Instance->HairGroupPublicData->GetDrawIndirectRasterComputeBuffer()));
 			}
 			else
 			{
@@ -877,6 +882,8 @@ static void RunHairStrandsInterpolation_Strands(
 				ShaderPrintData,
 				ClusterDatas,
 				TransientResources.IndirectDispatchArgsUAV);
+
+			AddTransitionPass(GraphBuilder, ShaderMap, Transitions);
 		}
 
 		// Run cluster debug view here (instead of GroomDebug.h/.cpp, as we need to have the (transient) cluster data 
@@ -887,78 +894,98 @@ static void RunHairStrandsInterpolation_Strands(
 	}
 
 	// Update dynamic mesh triangles
-	for (FInstanceData& InstanceData : InstanceDatas)
 	{
-		if (InstanceData.MeshLODIndex >= 0 && (InstanceData.BindingType == EHairBindingType::Skinning || InstanceData.bGlobalDeformationEnable))
+		TArray<FRDGBufferSRVRef> Transitions;
+		Transitions.Reserve(InstanceDatas.Num());
+		for (FInstanceData& InstanceData : InstanceDatas)
 		{
-			AddHairStrandUpdateMeshTrianglesPass(
-				GraphBuilder,
-				ShaderMap,
-				InstanceData.MeshLODIndex,
-				InstanceData.MeshDataLOD,
-				InstanceData.Instance->Strands.RestRootResource,
-				InstanceData.Instance->Strands.DeformedRootResource);
+			if (InstanceData.MeshLODIndex >= 0 && (InstanceData.BindingType == EHairBindingType::Skinning || InstanceData.bGlobalDeformationEnable))
+			{
+				AddHairStrandUpdateMeshTrianglesPass(
+					GraphBuilder,
+					ShaderMap,
+					InstanceData.MeshLODIndex,
+					InstanceData.MeshDataLOD,
+					InstanceData.Instance->Strands.RestRootResource,
+					InstanceData.Instance->Strands.DeformedRootResource);
+				Transitions.Add(RegisterAsSRV(GraphBuilder, InstanceData.Instance->Strands.DeformedRootResource->LODs[InstanceData.MeshLODIndex].GetDeformedUniqueTrianglePositionBuffer(FHairStrandsDeformedRootResource::FLOD::Current)));
+			}
 		}
+		AddTransitionPass(GraphBuilder, ShaderMap, Transitions);
 	}
 
-	// Update position offset (GPU)
-	for (FInstanceData& InstanceData : InstanceDatas)
+	// Update position offset 
 	{
-		if (InstanceData.MeshLODIndex >= 0 && (InstanceData.BindingType == EHairBindingType::Skinning || InstanceData.bGlobalDeformationEnable))
-		{
-			AddHairStrandUpdatePositionOffsetPass(
-				GraphBuilder,
-				ShaderMap,
-				EHairPositionUpdateType::Strands,
-				InstanceData.Instance->RegisteredIndex,
-				InstanceData.HairLODIndex,
-				InstanceData.MeshLODIndex,
-				InstanceData.Instance->Strands.DeformedRootResource,
-				InstanceData.Instance->Strands.DeformedResource);
-		}
-	}
+		TArray<FRDGBufferSRVRef> Transitions;
+		Transitions.Reserve(InstanceDatas.Num());
 
-	// Update position offset (CPU)
-	for (FInstanceData& InstanceData : InstanceDatas)
-	{
-		if (InstanceData.MeshLODIndex >= 0 && (InstanceData.bSimulationEnable || InstanceData.bDeformationEnable))
+		// From GPU root triangle
 		{
-			AddHairStrandUpdatePositionOffsetPass(
-				GraphBuilder,
-				ShaderMap,
-				EHairPositionUpdateType::Strands,
-				InstanceData.Instance->RegisteredIndex,
-				InstanceData.HairLODIndex,
-				InstanceData.MeshLODIndex,
-				nullptr,
-				InstanceData.Instance->Strands.DeformedResource);
+			for (FInstanceData& InstanceData : InstanceDatas)
+			{
+				if (InstanceData.MeshLODIndex >= 0 && (InstanceData.BindingType == EHairBindingType::Skinning || InstanceData.bGlobalDeformationEnable))
+				{
+					AddHairStrandUpdatePositionOffsetPass(
+						GraphBuilder,
+						ShaderMap,
+						EHairPositionUpdateType::Strands,
+						InstanceData.Instance->RegisteredIndex,
+						InstanceData.HairLODIndex,
+						InstanceData.MeshLODIndex,
+						InstanceData.Instance->Strands.DeformedRootResource,
+						InstanceData.Instance->Strands.DeformedResource);
+					Transitions.Add(RegisterAsSRV(GraphBuilder, InstanceData.Instance->Strands.DeformedResource->GetPositionOffsetBuffer(FHairStrandsDeformedResource::EFrameType::Current)));
+				}
+			}
 		}
-	}
-
-	// Update position offset (CPU) - With Sim Cache
-	if (bHasAnySimCacheInstances)
-	{
+	
+		// From CPU AABB
 		for (FInstanceData& InstanceData : InstanceDatas)
 		{
-			if (InstanceData.ActiveGroomCacheType == EGroomCacheType::Guides && InstanceData.bNeedDeformation)
+			if (InstanceData.MeshLODIndex >= 0 && (InstanceData.bSimulationEnable || InstanceData.bDeformationEnable))
 			{
-				// Apply the same offset to the render strands for proper rendering
-				const FVector OffsetPosition = InstanceData.Instance->Guides.DeformedResource->GetPositionOffset(FHairStrandsDeformedResource::EFrameType::Current);
-				InstanceData.Instance->Strands.DeformedResource->SetPositionOffset(FHairStrandsDeformedResource::EFrameType::Current, OffsetPosition);
-		
 				AddHairStrandUpdatePositionOffsetPass(
 					GraphBuilder,
 					ShaderMap,
 					EHairPositionUpdateType::Strands,
 					InstanceData.Instance->RegisteredIndex,
-					0,
+					InstanceData.HairLODIndex,
 					InstanceData.MeshLODIndex,
-					InstanceData.Instance->Strands.DeformedRootResource,
+					nullptr,
 					InstanceData.Instance->Strands.DeformedResource);
+				Transitions.Add(RegisterAsSRV(GraphBuilder, InstanceData.Instance->Strands.DeformedResource->GetPositionOffsetBuffer(FHairStrandsDeformedResource::EFrameType::Current)));
+			}
+		}
+		
+		// From Sim Cache
+		if (bHasAnySimCacheInstances)
+		{
+			for (FInstanceData& InstanceData : InstanceDatas)
+			{
+				if (InstanceData.ActiveGroomCacheType == EGroomCacheType::Guides && InstanceData.bNeedDeformation)
+				{
+					// Apply the same offset to the render strands for proper rendering
+					const FVector OffsetPosition = InstanceData.Instance->Guides.DeformedResource->GetPositionOffset(FHairStrandsDeformedResource::EFrameType::Current);
+					InstanceData.Instance->Strands.DeformedResource->SetPositionOffset(FHairStrandsDeformedResource::EFrameType::Current, OffsetPosition);
+		
+					AddHairStrandUpdatePositionOffsetPass(
+						GraphBuilder,
+						ShaderMap,
+						EHairPositionUpdateType::Strands,
+						InstanceData.Instance->RegisteredIndex,
+						0,
+						InstanceData.MeshLODIndex,
+						InstanceData.Instance->Strands.DeformedRootResource,
+						InstanceData.Instance->Strands.DeformedResource);
+					Transitions.Add(RegisterAsSRV(GraphBuilder, InstanceData.Instance->Strands.DeformedResource->GetPositionOffsetBuffer(FHairStrandsDeformedResource::EFrameType::Current)));
+				}
 			}
 		}
+
+		AddTransitionPass(GraphBuilder, ShaderMap, Transitions);
 	}
 
+
 	// Clear cluster AABBs (used optionally  for voxel allocation & for culling)
 	for (FInstanceData& InstanceData : InstanceDatas)
 	{
@@ -973,47 +1000,55 @@ static void RunHairStrandsInterpolation_Strands(
 			TransientResources.GroupAABBUAV);
 	}
 
-	for (FInstanceData& InstanceData : InstanceDatas)
 	{
-		if (InstanceData.ActiveGroomCacheType != EGroomCacheType::Strands && InstanceData.bNeedDeformation)
+		TArray<FRDGBufferSRVRef> Transitions;
+		Transitions.Reserve(InstanceDatas.Num());
+		for (FInstanceData& InstanceData : InstanceDatas)
 		{
-			const bool bValidGuide		= InstanceData.bNeedDeformation && (InstanceData.bSimulationEnable || InstanceData.bGlobalDeformationEnable || InstanceData.bDeformationEnable || InstanceData.Instance->Guides.bIsSimulationCacheEnable);// || (WITH_EDITOR && PatchMode == EHairPatchAttribute::GuideInflucence);
-			const bool bUseSingleGuide	= InstanceData.bNeedDeformation && bValidGuide && InstanceData.Instance->Strands.InterpolationResource->UseSingleGuide();
-			const bool bHasSkinning		= InstanceData.bNeedDeformation && InstanceData.BindingType == EHairBindingType::Skinning;
+			if (InstanceData.ActiveGroomCacheType != EGroomCacheType::Strands && InstanceData.bNeedDeformation)
+			{
+				const bool bValidGuide		= InstanceData.bNeedDeformation && (InstanceData.bSimulationEnable || InstanceData.bGlobalDeformationEnable || InstanceData.bDeformationEnable || InstanceData.Instance->Guides.bIsSimulationCacheEnable);// || (WITH_EDITOR && PatchMode == EHairPatchAttribute::GuideInflucence);
+				const bool bUseSingleGuide	= InstanceData.bNeedDeformation && bValidGuide && InstanceData.Instance->Strands.InterpolationResource->UseSingleGuide();
+				const bool bHasSkinning		= InstanceData.bNeedDeformation && InstanceData.BindingType == EHairBindingType::Skinning;
+	
+				AddHairStrandsInterpolationPass(
+					GraphBuilder,
+					ShaderMap,
+					ShaderPrintData,
+					InstanceData.Instance,
+					InstanceData.ActivePointCount,
+					InstanceData.MeshLODIndex,
+					InstanceData.Instance->Strands.Modifier.HairLengthScale,
+					InstanceData.Instance->Strands.HairInterpolationType,
+					EHairGeometryType::Strands,
+					InstanceData.CullingData,
+					InstanceData.Instance->Strands.RestResource->GetPositionOffset(),
+					bValidGuide ? InstanceData.Instance->Guides.RestResource->GetPositionOffset() : FVector::ZeroVector,
+					InstanceData.RDGResources.PositionOffsetSRV,
+					bValidGuide ? RegisterAsSRV(GraphBuilder, InstanceData.Instance->Guides.DeformedResource->GetPositionOffsetBuffer(FHairStrandsDeformedResource::Current)) : nullptr,
+					bHasSkinning ? InstanceData.Instance->Strands.RestRootResource : nullptr,
+					bHasSkinning && bValidGuide ? InstanceData.Instance->Guides.RestRootResource : nullptr,
+					bHasSkinning ? InstanceData.Instance->Strands.DeformedRootResource : nullptr,
+					bHasSkinning && bValidGuide ? InstanceData.Instance->Guides.DeformedRootResource : nullptr,
+					RegisterAsSRV(GraphBuilder, InstanceData.Instance->Strands.RestResource->PositionBuffer),
+					RegisterAsSRV(GraphBuilder, InstanceData.Instance->Strands.RestResource->PointToCurveBuffer),
+					bUseSingleGuide,
+					bValidGuide ? RegisterAsSRV(GraphBuilder, InstanceData.Instance->Strands.InterpolationResource->InterpolationBuffer) : nullptr,
+					bValidGuide ? RegisterAsSRV(GraphBuilder, InstanceData.Instance->Guides.RestResource->PositionBuffer) : nullptr,
+					bValidGuide ? RegisterAsSRV(GraphBuilder, InstanceData.Instance->Guides.DeformedResource->GetBuffer(FHairStrandsDeformedResource::Current)) : nullptr,
+					bValidGuide ? RegisterAsSRV(GraphBuilder, InstanceData.Instance->Strands.InterpolationResource->SimRootPointIndexBuffer) : nullptr,
+					bValidGuide ? RegisterAsSRV(GraphBuilder, InstanceData.Instance->Guides.RestResource->PointToCurveBuffer) : nullptr,
+					InstanceData.RDGResources.DeformerPositionSRV,
+					InstanceData.RDGResources.PositionUAV,
+					FHairStrandsDeformedRootResource::FLOD::Current);
+	
+				Transitions.Add(InstanceData.RDGResources.PositionSRV);
 
-			AddHairStrandsInterpolationPass(
-				GraphBuilder,
-				ShaderMap,
-				ShaderPrintData,
-				InstanceData.Instance,
-				InstanceData.ActivePointCount,
-				InstanceData.MeshLODIndex,
-				InstanceData.Instance->Strands.Modifier.HairLengthScale,
-				InstanceData.Instance->Strands.HairInterpolationType,
-				EHairGeometryType::Strands,
-				InstanceData.CullingData,
-				InstanceData.Instance->Strands.RestResource->GetPositionOffset(),
-				bValidGuide ? InstanceData.Instance->Guides.RestResource->GetPositionOffset() : FVector::ZeroVector,
-				InstanceData.RDGResources.PositionOffsetSRV,
-				bValidGuide ? RegisterAsSRV(GraphBuilder, InstanceData.Instance->Guides.DeformedResource->GetPositionOffsetBuffer(FHairStrandsDeformedResource::Current)) : nullptr,
-				bHasSkinning ? InstanceData.Instance->Strands.RestRootResource : nullptr,
-				bHasSkinning && bValidGuide ? InstanceData.Instance->Guides.RestRootResource : nullptr,
-				bHasSkinning ? InstanceData.Instance->Strands.DeformedRootResource : nullptr,
-				bHasSkinning && bValidGuide ? InstanceData.Instance->Guides.DeformedRootResource : nullptr,
-				RegisterAsSRV(GraphBuilder, InstanceData.Instance->Strands.RestResource->PositionBuffer),
-				RegisterAsSRV(GraphBuilder, InstanceData.Instance->Strands.RestResource->PointToCurveBuffer),
-				bUseSingleGuide,
-				bValidGuide ? RegisterAsSRV(GraphBuilder, InstanceData.Instance->Strands.InterpolationResource->InterpolationBuffer) : nullptr,
-				bValidGuide ? RegisterAsSRV(GraphBuilder, InstanceData.Instance->Guides.RestResource->PositionBuffer) : nullptr,
-				bValidGuide ? RegisterAsSRV(GraphBuilder, InstanceData.Instance->Guides.DeformedResource->GetBuffer(FHairStrandsDeformedResource::Current)) : nullptr,
-				bValidGuide ? RegisterAsSRV(GraphBuilder, InstanceData.Instance->Strands.InterpolationResource->SimRootPointIndexBuffer) : nullptr,
-				bValidGuide ? RegisterAsSRV(GraphBuilder, InstanceData.Instance->Guides.RestResource->PointToCurveBuffer) : nullptr,
-				InstanceData.RDGResources.DeformerPositionSRV,
-				InstanceData.RDGResources.PositionUAV,
-				FHairStrandsDeformedRootResource::FLOD::Current);
-
-			GraphBuilder.SetBufferAccessFinal(InstanceData.RDGResources.DeformedPosition.Buffer, ERHIAccess::SRVMask);
+				GraphBuilder.SetBufferAccessFinal(InstanceData.RDGResources.DeformedPosition.Buffer, ERHIAccess::SRVMask);
+			}
 		}
+
+		AddTransitionPass(GraphBuilder, ShaderMap, Transitions);
 	}
 
 	// Previous Position
@@ -1154,48 +1189,52 @@ static void RunHairStrandsInterpolation_Strands(
 	}
 
 	// 3. Compute cluster AABBs (used for LODing and voxelization)
-	const FVector& TranslatedWorldOffset = View->ViewMatrices.GetPreViewTranslation();
-	for (FInstanceData& InstanceData : InstanceDatas)
 	{
-		// Optim: If an instance does not voxelize it's data, then there is no need for having valid AABB
-		bool bNeedGPUAABB = InstanceData.Instance->Strands.Modifier.bSupportVoxelization && InstanceData.Instance->bCastShadow;
-		const EHairAABBUpdateType UpdateType = InstanceData.bNeedDeformation ? EHairAABBUpdateType::UpdateClusterAABB : EHairAABBUpdateType::UpdateGroupAABB;
-			
-		FHairStrandClusterData::FHairGroup* HairGroupCluster = nullptr;
-		if (InstanceData.CullingData.bCullingResultAvailable)
+		FRDGBufferUAVRef GroupAABUAVSkipBarrier = GraphBuilder.CreateUAV(TransientResources.GroupAABBBuffer, PF_R32_SINT, ERDGUnorderedAccessViewFlags::SkipBarrier);
+
+		const FVector& TranslatedWorldOffset = View->ViewMatrices.GetPreViewTranslation();
+		for (FInstanceData& InstanceData : InstanceDatas)
 		{
-			// Sanity check
-			check(InstanceData.Instance->Strands.bCullingEnable);
-			HairGroupCluster = &ClusterDatas.HairGroups[InstanceData.Instance->HairGroupPublicData->ClusterDataIndex];
-			if (!HairGroupCluster->bVisible)
+			// Optim: If an instance does not voxelize it's data, then there is no need for having valid AABB
+			bool bNeedGPUAABB = InstanceData.Instance->Strands.Modifier.bSupportVoxelization && InstanceData.Instance->bCastShadow;
+			const EHairAABBUpdateType UpdateType = InstanceData.bNeedDeformation ? EHairAABBUpdateType::UpdateClusterAABB : EHairAABBUpdateType::UpdateGroupAABB;
+				
+			FHairStrandClusterData::FHairGroup* HairGroupCluster = nullptr;
+			if (InstanceData.CullingData.bCullingResultAvailable)
 			{
-				bNeedGPUAABB = false;
+				// Sanity check
+				check(InstanceData.Instance->Strands.bCullingEnable);
+				HairGroupCluster = &ClusterDatas.HairGroups[InstanceData.Instance->HairGroupPublicData->ClusterDataIndex];
+				if (!HairGroupCluster->bVisible)
+				{
+					bNeedGPUAABB = false;
+				}
+			}
+				
+			InstanceData.Instance->HairGroupPublicData->SetClusterAABBValid(false);
+				
+			if (bNeedGPUAABB)
+			{
+				FRDGImportedBuffer Strands_CulledVertexCount = Register(GraphBuilder, InstanceData.Instance->HairGroupPublicData->GetDrawIndirectRasterComputeBuffer(), ERDGImportedBufferFlags::CreateSRV);
+				AddHairClusterAABBPass(
+					GraphBuilder,
+					ShaderMap,
+					InstanceData.ActiveCurveCount,
+					TranslatedWorldOffset,
+					ShaderPrintData,
+					UpdateType,
+					InstanceData.Instance,
+					InstanceData.RDGResources.PositionOffsetSRV,
+					HairGroupCluster,
+					InstanceData.CullingData,
+					InstanceData.RDGResources.PositionSRV,
+					Strands_CulledVertexCount.SRV,
+					GroupAABUAVSkipBarrier);
+				
+				InstanceData.Instance->HairGroupPublicData->SetClusterAABBValid(UpdateType == EHairAABBUpdateType::UpdateClusterAABB);
+	
+				TransientResources.bIsGroupAABBValid[InstanceData.RegisteredIndex] = true;
 			}
-		}
-			
-		InstanceData.Instance->HairGroupPublicData->SetClusterAABBValid(false);
-			
-		if (bNeedGPUAABB)
-		{
-			FRDGImportedBuffer Strands_CulledVertexCount = Register(GraphBuilder, InstanceData.Instance->HairGroupPublicData->GetDrawIndirectRasterComputeBuffer(), ERDGImportedBufferFlags::CreateSRV);
-			AddHairClusterAABBPass(
-				GraphBuilder,
-				ShaderMap,
-				InstanceData.ActiveCurveCount,
-				TranslatedWorldOffset,
-				ShaderPrintData,
-				UpdateType,
-				InstanceData.Instance,
-				InstanceData.RDGResources.PositionOffsetSRV,
-				HairGroupCluster,
-				InstanceData.CullingData,
-				InstanceData.RDGResources.PositionSRV,
-				Strands_CulledVertexCount.SRV,
-				TransientResources.GroupAABBUAV);
-			
-			InstanceData.Instance->HairGroupPublicData->SetClusterAABBValid(UpdateType == EHairAABBUpdateType::UpdateClusterAABB);
-
-			TransientResources.bIsGroupAABBValid[InstanceData.RegisteredIndex] = true;
 		}
 	}
 
diff --git a/Engine/Plugins/Runtime/HairStrands/Source/HairStrandsCore/Private/HairStrandsMeshProjection.cpp b/Engine/Plugins/Runtime/HairStrands/Source/HairStrandsCore/Private/HairStrandsMeshProjection.cpp
--- a/Engine/Plugins/Runtime/HairStrands/Source/HairStrandsCore/Private/HairStrandsMeshProjection.cpp
+++ b/Engine/Plugins/Runtime/HairStrands/Source/HairStrandsCore/Private/HairStrandsMeshProjection.cpp
@@ -10,6 +10,7 @@
 #include "RayTracingDynamicGeometryCollection.h"
 #include "RenderGraphUtils.h"
 #include "GroomResources.h"
+#include "SystemTextures.h"
 
 static int32 GHairProjectionMaxTrianglePerProjectionIteration = 8;
 static FAutoConsoleVariableRef CVarHairProjectionMaxTrianglePerProjectionIteration(TEXT("r.HairStrands.Projection.MaxTrianglePerIteration"), GHairProjectionMaxTrianglePerProjectionIteration, TEXT("Change the number of triangles which are iterated over during one projection iteration step. In kilo triangle (e.g., 8 == 8000 triangles). Default is 8."));
@@ -1205,4 +1206,79 @@ void AddHairStrandUpdatePositionOffsetPass(
 
 	GraphBuilder.SetBufferAccessFinal(OutCurrPositionOffsetBuffer.Buffer, ERHIAccess::SRVMask);
 	GraphBuilder.SetBufferAccessFinal(OutPrevPositionOffsetBuffer.Buffer, ERHIAccess::SRVMask);
+}
+
+///////////////////////////////////////////////////////////////////////////////////////////////////
+
+class FHairResourceTransitionPass : public FGlobalShader
+{
+public:
+private:
+	DECLARE_GLOBAL_SHADER(FHairResourceTransitionPass);
+	SHADER_USE_PARAMETER_STRUCT(FHairResourceTransitionPass, FGlobalShader);
+
+	using FPermutationDomain = TShaderPermutationDomain<>;
+
+	BEGIN_SHADER_PARAMETER_STRUCT(FParameters, )
+		SHADER_PARAMETER(uint32, DummyValue)
+		SHADER_PARAMETER_RDG_BUFFER_SRV_ARRAY(Buffer, Buffers, [16])
+		SHADER_PARAMETER_RDG_BUFFER_UAV(Buffer, DummyOutput)
+	END_SHADER_PARAMETER_STRUCT()
+
+public:
+	static bool ShouldCompilePermutation(const FGlobalShaderPermutationParameters& Parameters)
+	{
+		return IsHairStrandsSupported(EHairStrandsShaderType::All, Parameters.Platform);
+	}
+	static void ModifyCompilationEnvironment(const FGlobalShaderPermutationParameters& Parameters, FShaderCompilerEnvironment& OutEnvironment)
+	{
+		FGlobalShader::ModifyCompilationEnvironment(Parameters, OutEnvironment);
+		OutEnvironment.SetDefine(TEXT("SHADER_RESOURCE_TRANSITION"), 1);
+	}
+};
+
+IMPLEMENT_GLOBAL_SHADER(FHairResourceTransitionPass, "/Engine/Private/HairStrands/HairStrandsMesh.usf", "MainCS", SF_Compute);
+
+void AddTransitionPass(
+	FRDGBuilder& GraphBuilder,
+	FGlobalShaderMap* ShaderMap,
+	const TArray<FRDGBufferSRVRef>& Transitions)
+{
+	const uint32 ResourceCount = Transitions.Num();
+	if (ResourceCount == 0)
+	{
+		return;
+	}
+
+	FRDGBufferSRVRef DummyInput = GraphBuilder.CreateSRV(GSystemTextures.GetDefaultBuffer(GraphBuilder, 4u, 1u), PF_R32_UINT);
+	FRDGBufferRef DummyOutput = GraphBuilder.CreateBuffer(FRDGBufferDesc::CreateStructuredDesc(4,1),TEXT("DummyOutput"));
+	FRDGBufferUAVRef DummyOutputUAV = GraphBuilder.CreateUAV(DummyOutput, PF_R32_UINT, ERDGUnorderedAccessViewFlags::SkipBarrier);
+
+	const uint32 MaxBufferCount = 16;
+	const uint32 PassCount = FMath::DivideAndRoundUp(ResourceCount, MaxBufferCount);
+	for (uint32 PassIt=0; PassIt< PassCount; ++PassIt)
+	{
+		FHairResourceTransitionPass::FParameters* PassParameters = GraphBuilder.AllocParameters<FHairResourceTransitionPass::FParameters>();
+		PassParameters->DummyValue = 0;
+		PassParameters->DummyOutput = DummyOutputUAV;
+
+		const uint32 PassResourceOffset = PassIt * MaxBufferCount;
+		const uint32 PassResourceCount = FMath::Min(int32(MaxBufferCount), int32(ResourceCount) - int32(PassResourceOffset));
+		for (uint32 ResourceIt = 0; ResourceIt < PassResourceCount; ++ResourceIt)
+		{
+			PassParameters->Buffers[ResourceIt] = Transitions[PassResourceOffset + ResourceIt];
+		}
+		for (uint32 ResourceIt = PassResourceCount; ResourceIt < MaxBufferCount; ++ResourceIt)
+		{
+			PassParameters->Buffers[ResourceIt] = DummyInput;
+		}
+		TShaderMapRef<FHairResourceTransitionPass> ComputeShader(ShaderMap);
+		FComputeShaderUtils::AddPass(
+			GraphBuilder,
+			RDG_EVENT_NAME("HairStrands::ResourceTransitions"),
+			ERDGPassFlags::Compute | ERDGPassFlags::NeverCull,
+			ComputeShader,
+			PassParameters,
+			FIntVector(1, 1, 1));
+	}
 }
\ No newline at end of file
diff --git a/Engine/Plugins/Runtime/HairStrands/Source/HairStrandsCore/Public/HairStrandsMeshProjection.h b/Engine/Plugins/Runtime/HairStrands/Source/HairStrandsCore/Public/HairStrandsMeshProjection.h
--- a/Engine/Plugins/Runtime/HairStrands/Source/HairStrandsCore/Public/HairStrandsMeshProjection.h
+++ b/Engine/Plugins/Runtime/HairStrands/Source/HairStrandsCore/Public/HairStrandsMeshProjection.h
@@ -153,4 +153,9 @@ void AddHairStrandUpdatePositionOffsetPass(
 	FHairStrandsDeformedRootResource* DeformedRootResources,
 	FHairStrandsDeformedResource* DeformedResources);
 
-FHairStrandsProjectionMeshData ExtractMeshData(FSkeletalMeshRenderData* RenderData);
\ No newline at end of file
+void AddTransitionPass(
+	FRDGBuilder& GraphBuilder,
+	FGlobalShaderMap* ShaderMap,
+	const TArray<FRDGBufferSRVRef>& Transitions);
+
+FHairStrandsProjectionMeshData ExtractMeshData(FSkeletalMeshRenderData* RenderData);
diff --git a/Engine/Shaders/Private/HairStrands/HairStrandsMesh.usf b/Engine/Shaders/Private/HairStrands/HairStrandsMesh.usf
--- a/Engine/Shaders/Private/HairStrands/HairStrandsMesh.usf
+++ b/Engine/Shaders/Private/HairStrands/HairStrandsMesh.usf
@@ -353,4 +353,58 @@ void MainCS(uint3 DispatchThreadId : SV_DispatchThreadID)
 #endif
 }
 
-#endif // SHADER_OFFSET_UPDATE
\ No newline at end of file
+#endif // SHADER_OFFSET_UPDATE
+
+
+////////////////////////////////////////////////////////////////////////
+// Update hair strands offset position
+
+#if SHADER_RESOURCE_TRANSITION
+
+Buffer<uint> Buffers_0;
+Buffer<uint> Buffers_1;
+Buffer<uint> Buffers_2;
+Buffer<uint> Buffers_3;
+Buffer<uint> Buffers_4;
+Buffer<uint> Buffers_5;
+Buffer<uint> Buffers_6;
+Buffer<uint> Buffers_7;
+Buffer<uint> Buffers_8;
+Buffer<uint> Buffers_9;
+Buffer<uint> Buffers_10;
+Buffer<uint> Buffers_11;
+Buffer<uint> Buffers_12;
+Buffer<uint> Buffers_13;
+Buffer<uint> Buffers_14;
+Buffer<uint> Buffers_15;
+
+RWBuffer<uint> DummyOutput;
+uint DummyValue;
+
+[numthreads(1, 1, 1)]
+void MainCS(uint3 DispatchThreadId : SV_DispatchThreadID)
+{
+	if (DummyValue > 0)
+	{
+		uint Acc = 0;
+		Acc += Buffers_0[0];
+		Acc += Buffers_1[0];
+		Acc += Buffers_2[0];
+		Acc += Buffers_3[0];
+		Acc += Buffers_4[0];
+		Acc += Buffers_5[0];
+		Acc += Buffers_6[0];
+		Acc += Buffers_7[0];
+		Acc += Buffers_8[0];
+		Acc += Buffers_9[0];
+		Acc += Buffers_10[0];
+		Acc += Buffers_11[0];
+		Acc += Buffers_12[0];
+		Acc += Buffers_13[0];
+		Acc += Buffers_14[0];
+		Acc += Buffers_15[0];
+		DummyOutput[0] = Acc;
+	}
+}
+
+#endif // SHADER_RESOURCE_TRANSITION
\ No newline at end of file
diff --git a/Engine/Source/Runtime/Renderer/Private/HairStrands/HairStrandsCluster.cpp b/Engine/Source/Runtime/Renderer/Private/HairStrands/HairStrandsCluster.cpp
--- a/Engine/Source/Runtime/Renderer/Private/HairStrands/HairStrandsCluster.cpp
+++ b/Engine/Source/Runtime/Renderer/Private/HairStrands/HairStrandsCluster.cpp
@@ -293,8 +293,8 @@ void CreateHairStrandsMacroGroups(
 		MacroGroupResources.MacroGroupAABBsBuffer = GraphBuilder.CreateBuffer(FRDGBufferDesc::CreateBufferDesc(4, 6 * MacroGroupCount), TEXT("Hair.MacroGroupAABBBuffer"));
 		MacroGroupResources.MacroGroupVoxelAlignedAABBsBuffer = GraphBuilder.CreateBuffer(FRDGBufferDesc::CreateBufferDesc(4, 6 * MacroGroupCount), TEXT("Hair.MacroGroupVoxelAlignedAABBBuffer"));
 		MacroGroupResources.MacroGroupVoxelSizeBuffer = GraphBuilder.CreateBuffer(FRDGBufferDesc::CreateBufferDesc(2, MacroGroupCount), TEXT("Hair.MacroGroupVoxelSize"));
-		FRDGBufferUAVRef MacroGroupAABBBufferUAV = GraphBuilder.CreateUAV(MacroGroupResources.MacroGroupAABBsBuffer, PF_R32_SINT);
-		FRDGBufferUAVRef MacroGroupVoxelSizeBufferUAV = GraphBuilder.CreateUAV(MacroGroupResources.MacroGroupVoxelSizeBuffer, PF_R16F);
+		FRDGBufferUAVRef MacroGroupAABBBufferUAV = GraphBuilder.CreateUAV(MacroGroupResources.MacroGroupAABBsBuffer, PF_R32_SINT, ERDGUnorderedAccessViewFlags::SkipBarrier);
+		FRDGBufferUAVRef MacroGroupVoxelSizeBufferUAV = GraphBuilder.CreateUAV(MacroGroupResources.MacroGroupVoxelSizeBuffer, PF_R16F, ERDGUnorderedAccessViewFlags::SkipBarrier);
 		for (FHairStrandsMacroGroupData& MacroGroup : MacroGroups)
 		{				
 			AddHairMacroGroupAABBPass(GraphBuilder, View, *Scene->HairStrandsSceneData.TransientResources, MacroGroup, MacroGroupAABBBufferUAV, MacroGroupVoxelSizeBufferUAV);
