Fix CIS static analysis issue.

diff --git a/Engine/Plugins/Runtime/HairStrands/Source/HairStrandsCore/Private/GroomManager.cpp b/Engine/Plugins/Runtime/HairStrands/Source/HairStrandsCore/Private/GroomManager.cpp
--- a/Engine/Plugins/Runtime/HairStrands/Source/HairStrandsCore/Private/GroomManager.cpp
+++ b/Engine/Plugins/Runtime/HairStrands/Source/HairStrandsCore/Private/GroomManager.cpp
@@ -514,6 +514,8 @@ static void RunHairStrandsInterpolation_Guide(
 	{
 		if (InstanceData.NeedsMeshUpdate())
 		{
+			check(InstanceData.MeshLODIndex >= 0);
+
 			AddHairStrandUpdatePositionOffsetPass(
 				GraphBuilder,
 				ShaderMap,
