Fix incrrect shader parameter type for a UAV buffer.

diff --git a/Engine/Plugins/Runtime/HairStrands/Source/HairStrandsCore/Private/HairStrandsMeshProjection.cpp b/Engine/Plugins/Runtime/HairStrands/Source/HairStrandsCore/Private/HairStrandsMeshProjection.cpp
--- a/Engine/Plugins/Runtime/HairStrands/Source/HairStrandsCore/Private/HairStrandsMeshProjection.cpp
+++ b/Engine/Plugins/Runtime/HairStrands/Source/HairStrandsCore/Private/HairStrandsMeshProjection.cpp
@@ -1222,7 +1222,7 @@ private:
 	BEGIN_SHADER_PARAMETER_STRUCT(FParameters, )
 		SHADER_PARAMETER(uint32, DummyValue)
 		SHADER_PARAMETER_RDG_BUFFER_SRV_ARRAY(Buffer, Buffers, [16])
-		SHADER_PARAMETER_RDG_BUFFER_UAV(Buffer, DummyOutput)
+		SHADER_PARAMETER_RDG_BUFFER_UAV(RWBuffer, DummyOutput)
 	END_SHADER_PARAMETER_STRUCT()
 
 public:
