Fix RHI validation error.

diff --git a/Engine/Plugins/Runtime/HairStrands/Source/HairStrandsCore/Private/GroomManager.cpp b/Engine/Plugins/Runtime/HairStrands/Source/HairStrandsCore/Private/GroomManager.cpp
--- a/Engine/Plugins/Runtime/HairStrands/Source/HairStrandsCore/Private/GroomManager.cpp
+++ b/Engine/Plugins/Runtime/HairStrands/Source/HairStrandsCore/Private/GroomManager.cpp
@@ -656,6 +656,7 @@ static void RunHairStrandsInterpolation_Guide(
 			AddGroomCacheUpdatePass(
 				GraphBuilder,
 				ShaderMap,
+				InstanceData.Instance->RegisteredIndex,
 				InstanceData.Instance->Guides.RestResource->GetPointCount(),
 				InterpolationFactor,
 				CacheResources0,
@@ -1113,6 +1114,7 @@ static void RunHairStrandsInterpolation_Strands(
 				AddGroomCacheUpdatePass(
 					GraphBuilder,
 					ShaderMap,
+					InstanceData.Instance->RegisteredIndex,
 					InstanceData.ActivePointCount,
 					InterpolationFactor,
 					CacheResources0,
diff --git a/Engine/Plugins/Runtime/HairStrands/Source/HairStrandsCore/Private/HairStrandsInterpolation.cpp b/Engine/Plugins/Runtime/HairStrands/Source/HairStrandsCore/Private/HairStrandsInterpolation.cpp
--- a/Engine/Plugins/Runtime/HairStrands/Source/HairStrandsCore/Private/HairStrandsInterpolation.cpp
+++ b/Engine/Plugins/Runtime/HairStrands/Source/HairStrandsCore/Private/HairStrandsInterpolation.cpp
@@ -229,6 +229,7 @@ class FGroomCacheUpdatePassCS : public FGlobalShader
 	SHADER_USE_PARAMETER_STRUCT(FGroomCacheUpdatePassCS, FGlobalShader);
 
 	BEGIN_SHADER_PARAMETER_STRUCT(FParameters, )
+		SHADER_PARAMETER(uint32, InstanceRegisteredIndex)
 		SHADER_PARAMETER(uint32, ElementCount)
 		SHADER_PARAMETER(uint32, bHasRadiusData)
 		SHADER_PARAMETER(float, InterpolationFactor)
@@ -237,7 +238,7 @@ class FGroomCacheUpdatePassCS : public FGlobalShader
 		SHADER_PARAMETER_RDG_BUFFER_SRV(Buffer, InRadius0Buffer)
 		SHADER_PARAMETER_RDG_BUFFER_SRV(Buffer, InRadius1Buffer)
 		SHADER_PARAMETER_RDG_BUFFER_SRV(Buffer, InRestPoseBuffer)
-		SHADER_PARAMETER_RDG_BUFFER_SRV(Buffer, InDeformedOffsetBuffer)
+		SHADER_PARAMETER_RDG_BUFFER_SRV(StructuredBuffer, InDeformedOffsetBuffer)
 		SHADER_PARAMETER_RDG_BUFFER_UAV(RWBuffer, OutDeformedBuffer)
 	END_SHADER_PARAMETER_STRUCT()
 
@@ -257,6 +258,7 @@ IMPLEMENT_GLOBAL_SHADER(FGroomCacheUpdatePassCS, "/Engine/Private/HairStrands/Ha
 void AddGroomCacheUpdatePass(
 	FRDGBuilder& GraphBuilder,
 	FGlobalShaderMap* ShaderMap,
+	uint32 InstanceRegisteredIndex,
 	uint32 PointCount,
 	float InterpolationFactor,
 	FGroomCacheResources CacheResources0,
@@ -270,6 +272,7 @@ void AddGroomCacheUpdatePass(
 	if (PointCount == 0 || !CacheResources0.PositionBuffer || !CacheResources1.PositionBuffer) return;
 
 	FGroomCacheUpdatePassCS::FParameters* Parameters = GraphBuilder.AllocParameters<FGroomCacheUpdatePassCS::FParameters>();
+	Parameters->InstanceRegisteredIndex = InstanceRegisteredIndex;
 	Parameters->ElementCount = PointCount;
 	Parameters->InPosition0Buffer = CacheResources0.PositionBuffer;
 	Parameters->InPosition1Buffer = CacheResources1.PositionBuffer;
diff --git a/Engine/Plugins/Runtime/HairStrands/Source/HairStrandsCore/Private/HairStrandsMeshProjection.cpp b/Engine/Plugins/Runtime/HairStrands/Source/HairStrandsCore/Private/HairStrandsMeshProjection.cpp
--- a/Engine/Plugins/Runtime/HairStrands/Source/HairStrandsCore/Private/HairStrandsMeshProjection.cpp
+++ b/Engine/Plugins/Runtime/HairStrands/Source/HairStrandsCore/Private/HairStrandsMeshProjection.cpp
@@ -1217,12 +1217,14 @@ private:
 	DECLARE_GLOBAL_SHADER(FHairResourceTransitionPass);
 	SHADER_USE_PARAMETER_STRUCT(FHairResourceTransitionPass, FGlobalShader);
 
-	using FPermutationDomain = TShaderPermutationDomain<>;
+	class FStructuredBuffer : SHADER_PERMUTATION_BOOL("PERMUTATION_STRUCTURED_BUFFER");
+	using FPermutationDomain = TShaderPermutationDomain<FStructuredBuffer>;
 
 	BEGIN_SHADER_PARAMETER_STRUCT(FParameters, )
 		SHADER_PARAMETER(uint32, DummyValue)
-		SHADER_PARAMETER_RDG_BUFFER_SRV_ARRAY(Buffer, Buffers, [16])
-		SHADER_PARAMETER_RDG_BUFFER_UAV(RWBuffer, DummyOutput)
+		SHADER_PARAMETER_RDG_BUFFER_SRV_ARRAY(Buffer, VertexBuffers, [16])
+		SHADER_PARAMETER_RDG_BUFFER_SRV_ARRAY(StructuredBuffer, StructuredBuffers, [16])
+		SHADER_PARAMETER_RDG_BUFFER_UAV(RWStructuredBuffer, DummyOutput)
 	END_SHADER_PARAMETER_STRUCT()
 
 public:
@@ -1250,10 +1252,16 @@ void AddTransitionPass(
 		return;
 	}
 
-	FRDGBufferSRVRef DummyInput = GraphBuilder.CreateSRV(GSystemTextures.GetDefaultBuffer(GraphBuilder, 4u, 1u), PF_R32_UINT);
+	FRDGBufferSRVRef DummyVertexInput = GraphBuilder.CreateSRV(GSystemTextures.GetDefaultBuffer(GraphBuilder, 4u, 1u), PF_R32_UINT);
+	FRDGBufferSRVRef DummyStructuredInput= GraphBuilder.CreateSRV(GSystemTextures.GetDefaultStructuredBuffer(GraphBuilder, 16u));
+
 	FRDGBufferRef DummyOutput = GraphBuilder.CreateBuffer(FRDGBufferDesc::CreateStructuredDesc(4,1),TEXT("DummyOutput"));
 	FRDGBufferUAVRef DummyOutputUAV = GraphBuilder.CreateUAV(DummyOutput, PF_R32_UINT, ERDGUnorderedAccessViewFlags::SkipBarrier);
 
+	check(Transitions[0]);
+	check(Transitions[0]->Desc.Buffer);
+	const bool bStructuredBuffer = EnumHasAnyFlags(Transitions[0]->Desc.Buffer->Desc.Usage, EBufferUsageFlags::StructuredBuffer);
+
 	const uint32 MaxBufferCount = 16;
 	const uint32 PassCount = FMath::DivideAndRoundUp(ResourceCount, MaxBufferCount);
 	for (uint32 PassIt=0; PassIt< PassCount; ++PassIt)
@@ -1264,15 +1272,31 @@ void AddTransitionPass(
 
 		const uint32 PassResourceOffset = PassIt * MaxBufferCount;
 		const uint32 PassResourceCount = FMath::Min(int32(MaxBufferCount), int32(ResourceCount) - int32(PassResourceOffset));
-		for (uint32 ResourceIt = 0; ResourceIt < PassResourceCount; ++ResourceIt)
+		if (bStructuredBuffer)
 		{
-			PassParameters->Buffers[ResourceIt] = Transitions[PassResourceOffset + ResourceIt];
+			for (uint32 ResourceIt = 0; ResourceIt < PassResourceCount; ++ResourceIt)
+			{
+				PassParameters->StructuredBuffers[ResourceIt] = Transitions[PassResourceOffset + ResourceIt];
+			}
+			for (uint32 ResourceIt = PassResourceCount; ResourceIt < MaxBufferCount; ++ResourceIt)
+			{
+				PassParameters->StructuredBuffers[ResourceIt] = DummyStructuredInput;
+			}
 		}
-		for (uint32 ResourceIt = PassResourceCount; ResourceIt < MaxBufferCount; ++ResourceIt)
+		else
 		{
-			PassParameters->Buffers[ResourceIt] = DummyInput;
+			for (uint32 ResourceIt = 0; ResourceIt < PassResourceCount; ++ResourceIt)
+			{
+				PassParameters->VertexBuffers[ResourceIt] = Transitions[PassResourceOffset + ResourceIt];
+			}
+			for (uint32 ResourceIt = PassResourceCount; ResourceIt < MaxBufferCount; ++ResourceIt)
+			{
+				PassParameters->VertexBuffers[ResourceIt] = DummyVertexInput;
+			}
 		}
-		TShaderMapRef<FHairResourceTransitionPass> ComputeShader(ShaderMap);
+		FHairResourceTransitionPass::FPermutationDomain PermutationVector;
+		PermutationVector.Set<FHairResourceTransitionPass::FStructuredBuffer>(bStructuredBuffer);
+		TShaderMapRef<FHairResourceTransitionPass> ComputeShader(ShaderMap, PermutationVector);
 		FComputeShaderUtils::AddPass(
 			GraphBuilder,
 			RDG_EVENT_NAME("HairStrands::ResourceTransitions"),
diff --git a/Engine/Plugins/Runtime/HairStrands/Source/HairStrandsCore/Public/HairStrandsInterpolation.h b/Engine/Plugins/Runtime/HairStrands/Source/HairStrandsCore/Public/HairStrandsInterpolation.h
--- a/Engine/Plugins/Runtime/HairStrands/Source/HairStrandsCore/Public/HairStrandsInterpolation.h
+++ b/Engine/Plugins/Runtime/HairStrands/Source/HairStrandsCore/Public/HairStrandsInterpolation.h
@@ -80,6 +80,7 @@ FGroomCacheResources CreateGroomCacheBuffer(FRDGBuilder& GraphBuilder, FGroomCac
 void AddGroomCacheUpdatePass(
 	FRDGBuilder& GraphBuilder,
 	FGlobalShaderMap* ShaderMap,
+	uint32 InstanceRegisteredIndex,
 	uint32 PointCount,
 	float InterpolationFactor,
 	FGroomCacheResources CacheResources0,
diff --git a/Engine/Shaders/Private/HairStrands/HairStrandsInterpolation.usf b/Engine/Shaders/Private/HairStrands/HairStrandsInterpolation.usf
--- a/Engine/Shaders/Private/HairStrands/HairStrandsInterpolation.usf
+++ b/Engine/Shaders/Private/HairStrands/HairStrandsInterpolation.usf
@@ -614,6 +614,7 @@ void MainCS(uint DispatchThreadId : SV_DispatchThreadID)
 
 #if SHADER_GROOMCACHE_UPDATE
 
+uint InstanceRegisteredIndex;
 uint ElementCount;
 uint bHasRadiusData;
 float InterpolationFactor;
@@ -624,7 +625,7 @@ Buffer<float> InRadius0Buffer;
 Buffer<float> InRadius1Buffer;
 
 Buffer<uint4>  InRestPoseBuffer;
-Buffer<float4> InDeformedOffsetBuffer;
+StructuredBuffer<float4> InDeformedOffsetBuffer;
 RWBuffer<uint4> OutDeformedBuffer;
 
 [numthreads(GROUP_SIZE, 1, 1)]
@@ -648,7 +649,6 @@ void MainCS(uint DispatchThreadId : SV_DispatchThreadID)
 		InPosition1Buffer[Index3+2]);
 	
 	const float3 WorldVertexPosition = lerp(InPosition0, InPosition1, InterpolationFactor);
-
 	FHairControlPoint CP = UnpackHairControlPoint(InRestPoseBuffer[Index], float3(0, 0, 0), 1, 1, 1);
 	CP.Position = WorldVertexPosition;
 
@@ -657,7 +657,8 @@ void MainCS(uint DispatchThreadId : SV_DispatchThreadID)
 		CP.WorldRadius = lerp(InRadius0Buffer[Index], InRadius1Buffer[Index], InterpolationFactor);
 	}
 
-	OutDeformedBuffer[Index] = PackHairControlPoint(CP, InDeformedOffsetBuffer[0].xyz, 1);
+	const float3 DeformedPositionOffset = ReadRenPositionOffset(InDeformedOffsetBuffer, InstanceRegisteredIndex);
+	OutDeformedBuffer[Index] = PackHairControlPoint(CP, DeformedPositionOffset, 1);
 }
 
 #endif // SHADER_GROOMCACHE_UPDATE
diff --git a/Engine/Shaders/Private/HairStrands/HairStrandsMesh.usf b/Engine/Shaders/Private/HairStrands/HairStrandsMesh.usf
--- a/Engine/Shaders/Private/HairStrands/HairStrandsMesh.usf
+++ b/Engine/Shaders/Private/HairStrands/HairStrandsMesh.usf
@@ -361,24 +361,28 @@ void MainCS(uint3 DispatchThreadId : SV_DispatchThreadID)
 
 #if SHADER_RESOURCE_TRANSITION
 
-Buffer<uint> Buffers_0;
-Buffer<uint> Buffers_1;
-Buffer<uint> Buffers_2;
-Buffer<uint> Buffers_3;
-Buffer<uint> Buffers_4;
-Buffer<uint> Buffers_5;
-Buffer<uint> Buffers_6;
-Buffer<uint> Buffers_7;
-Buffer<uint> Buffers_8;
-Buffer<uint> Buffers_9;
-Buffer<uint> Buffers_10;
-Buffer<uint> Buffers_11;
-Buffer<uint> Buffers_12;
-Buffer<uint> Buffers_13;
-Buffer<uint> Buffers_14;
-Buffer<uint> Buffers_15;
-
-RWBuffer<uint> DummyOutput;
+#define TEMPLATE_BUFFERS(TPrefix, TSuffix) \
+	TPrefix##Buffers_0##TSuffix;\
+	TPrefix##Buffers_1##TSuffix;\
+	TPrefix##Buffers_2##TSuffix;\
+	TPrefix##Buffers_3##TSuffix;\
+	TPrefix##Buffers_4##TSuffix;\
+	TPrefix##Buffers_5##TSuffix;\
+	TPrefix##Buffers_6##TSuffix;\
+	TPrefix##Buffers_7##TSuffix;\
+	TPrefix##Buffers_8##TSuffix;\
+	TPrefix##Buffers_9##TSuffix;\
+	TPrefix##Buffers_10##TSuffix;\
+	TPrefix##Buffers_11##TSuffix;\
+	TPrefix##Buffers_12##TSuffix;\
+	TPrefix##Buffers_13##TSuffix;\
+	TPrefix##Buffers_14##TSuffix;\
+	TPrefix##Buffers_15##TSuffix;
+
+TEMPLATE_BUFFERS(StructuredBuffer<float4> Structured, )
+TEMPLATE_BUFFERS(Buffer<uint> Vertex, )
+
+RWStructuredBuffer<uint> DummyOutput;
 uint DummyValue;
 
 [numthreads(1, 1, 1)]
@@ -387,22 +391,11 @@ void MainCS(uint3 DispatchThreadId : SV_DispatchThreadID)
 	if (DummyValue > 0)
 	{
 		uint Acc = 0;
-		Acc += Buffers_0[0];
-		Acc += Buffers_1[0];
-		Acc += Buffers_2[0];
-		Acc += Buffers_3[0];
-		Acc += Buffers_4[0];
-		Acc += Buffers_5[0];
-		Acc += Buffers_6[0];
-		Acc += Buffers_7[0];
-		Acc += Buffers_8[0];
-		Acc += Buffers_9[0];
-		Acc += Buffers_10[0];
-		Acc += Buffers_11[0];
-		Acc += Buffers_12[0];
-		Acc += Buffers_13[0];
-		Acc += Buffers_14[0];
-		Acc += Buffers_15[0];
+	#if PERMUTATION_STRUCTURED_BUFFER
+		TEMPLATE_BUFFERS(Acc += Structured, [0].x)
+	#else
+		TEMPLATE_BUFFERS(Acc += Vertex, [0])
+	#endif
 		DummyOutput[0] = Acc;
 	}
 }
