* Change cluster AABB to be transient
* Allocate a unique cluster AABB buffer for all instances
* Batched cluster & group AABB clear

Reduce hair interpolation from 0.26 -> 0.22ms

diff --git a/Engine/Plugins/Runtime/HairStrands/Source/HairStrandsCore/Private/GroomDebug.cpp b/Engine/Plugins/Runtime/HairStrands/Source/HairStrandsCore/Private/GroomDebug.cpp
--- a/Engine/Plugins/Runtime/HairStrands/Source/HairStrandsCore/Private/GroomDebug.cpp
+++ b/Engine/Plugins/Runtime/HairStrands/Source/HairStrandsCore/Private/GroomDebug.cpp
@@ -408,9 +408,10 @@ class FDrawDebugClusterAABBCS : public FGlobalShader
 
 	BEGIN_SHADER_PARAMETER_STRUCT(FParameters,)
 		SHADER_PARAMETER_STRUCT_REF(FViewUniformShaderParameters, ViewUniformBuffer)
-		SHADER_PARAMETER_SRV(Buffer, ClusterAABBBuffer)
+		SHADER_PARAMETER_RDG_BUFFER_SRV(Buffer, ClusterAABBBuffer)
 		SHADER_PARAMETER_RDG_BUFFER_SRV(Buffer, GroupAABBBuffer)
 		SHADER_PARAMETER(uint32, InstanceRegisteredIndex)
+		SHADER_PARAMETER(uint32, ClusterOffset)
 		SHADER_PARAMETER(uint32, ClusterCount)
 		SHADER_PARAMETER(uint32, PointCount)
 		SHADER_PARAMETER(uint32, CurveCount)
@@ -463,12 +464,13 @@ void AddDrawDebugClusterPass(
 		FDrawDebugClusterAABBCS::FParameters* Parameters = GraphBuilder.AllocParameters<FDrawDebugClusterAABBCS::FParameters>();
 		Parameters->InstanceRegisteredIndex = HairGroupClusters.InstanceRegisteredIndex;
 		Parameters->ViewUniformBuffer = View.ViewUniformBuffer;
-		Parameters->ClusterCount = HairGroupClusters.ClusterCount;
+		Parameters->ClusterOffset = TransientResources.GetClusterOffset(HairGroupClusters.InstanceRegisteredIndex);
+		Parameters->ClusterCount  = TransientResources.GetClusterCount(HairGroupClusters.InstanceRegisteredIndex);
 		Parameters->PointCount = HairGroupClusters.HairGroupPublicPtr->GetActiveStrandsPointCount();
 		Parameters->CurveCount = HairGroupClusters.HairGroupPublicPtr->GetActiveStrandsCurveCount();
 		Parameters->HairGroupId = DataIndex++;
 		Parameters->bDrawAABB = ViewMode == EGroomViewMode::ClusterAABB ? 1 : 0;
-		Parameters->ClusterAABBBuffer = HairGroupClusters.ClusterAABBBuffer->SRV;
+		Parameters->ClusterAABBBuffer = TransientResources.ClusterAABBSRV;
 		Parameters->GroupAABBBuffer = TransientResources.GroupAABBSRV;
 		ShaderPrint::SetParameters(GraphBuilder, *ShaderPrintData, Parameters->ShaderPrintParameters);
 
diff --git a/Engine/Plugins/Runtime/HairStrands/Source/HairStrandsCore/Private/GroomManager.cpp b/Engine/Plugins/Runtime/HairStrands/Source/HairStrandsCore/Private/GroomManager.cpp
--- a/Engine/Plugins/Runtime/HairStrands/Source/HairStrandsCore/Private/GroomManager.cpp
+++ b/Engine/Plugins/Runtime/HairStrands/Source/HairStrandsCore/Private/GroomManager.cpp
@@ -80,10 +80,12 @@ void AddHairClusterAABBPass(
 	FHairGroupInstance* Instance,
 	FRDGBufferSRVRef RenderDeformedOffsetBuffer,
 	FHairStrandClusterData::FHairGroup* ClusterData,
-	FRDGHairStrandsCullingData& ClusterAABBData,
+	const uint32 ClusterOffset,
+	const uint32 CluesterCount,
 	FRDGBufferSRVRef RenderPositionBufferSRV,
 	FRDGBufferSRVRef& DrawIndirectRasterComputeBuffer,
-	FRDGBufferUAVRef GroupAABBUAV);
+	FRDGBufferUAVRef ClusterAABBUAV,
+	FRDGBufferUAVRef GroupAABBBUAV);
 
 ///////////////////////////////////////////////////////////////////////////////////////////////////
 // Utils
@@ -684,6 +686,7 @@ static void RunHairStrandsInterpolation_Strands(
 	const TArray<const FSceneView*>& Views, 
 	const FSceneView* View,
 	const uint32 ViewUniqueID,
+	const uint32 TotalRegisteredInstanceCount,
 	const FHairStrandsInstances& Instances,
 	const FShaderPrintData* ShaderPrintData,
 	FHairTransientResources& TransientResources,
@@ -878,13 +881,14 @@ static void RunHairStrandsInterpolation_Strands(
 			TRACE_CPUPROFILER_EVENT_SCOPE(ComputeHairStrandsClustersCulling);
 			RDG_GPU_STAT_SCOPE(GraphBuilder, HairStrandsClusterCulling);
 
+			FRDGBufferUAVRef IndirectDispatchArgsUAVWithSkipBarrier = GraphBuilder.CreateUAV(TransientResources.IndirectDispatchArgsBuffer, ERDGUnorderedAccessViewFlags::SkipBarrier);
 			AddClusterCullingPass(
 				GraphBuilder,
 				ShaderMap,
 				Views[0],
 				ShaderPrintData,
 				ClusterDatas,
-				TransientResources.IndirectDispatchArgsUAV);
+				IndirectDispatchArgsUAVWithSkipBarrier);
 
 			AddTransitionPass(GraphBuilder, ShaderMap, Transitions);
 		}
@@ -990,17 +994,31 @@ static void RunHairStrandsInterpolation_Strands(
 
 
 	// Clear cluster AABBs (used optionally  for voxel allocation & for culling)
-	for (FInstanceData& InstanceData : InstanceDatas)
 	{
-		InstanceData.CullingData = ImportCullingData(GraphBuilder, InstanceData.Instance->HairGroupPublicData);
-		AddClearClusterAABBPass(
-			GraphBuilder,
-			ShaderMap,
-			InstanceData.bNeedDeformation ? EHairAABBUpdateType::UpdateClusterAABB : EHairAABBUpdateType::UpdateGroupAABB,
-			InstanceData.RegisteredIndex,
-			InstanceData.CullingData.ClusterCount,
-			InstanceData.CullingData.ClusterAABBBuffer,
-			TransientResources.GroupAABBUAV);
+		uint32 TotalClusterAABBCount = 0;
+		TransientResources.ClusterAABBOffetAndCounts.Init(FUintVector2::ZeroValue, TotalRegisteredInstanceCount); // Use TotalRegisteredInstanceCount because this buffer is indexed by RegisteredIndex
+		for (FInstanceData& InstanceData : InstanceDatas)
+		{
+			if (InstanceData.bNeedDeformation)
+			{
+				TransientResources.ClusterAABBOffetAndCounts[InstanceData.RegisteredIndex].X = TotalClusterAABBCount;
+				TransientResources.ClusterAABBOffetAndCounts[InstanceData.RegisteredIndex].Y = InstanceData.Instance->HairGroupPublicData->ClusterCount;
+				TotalClusterAABBCount += InstanceData.Instance->HairGroupPublicData->ClusterCount;
+			}
+		}
+		TotalClusterAABBCount = FMath::Max(TotalClusterAABBCount, 1u);
+		TransientResources.ClusterAABBBuffer = GraphBuilder.CreateBuffer(FRDGBufferDesc::CreateBufferDesc(4, 6u * TotalClusterAABBCount), TEXT("Hair.Transient.ClusterAABBs"));
+		TransientResources.ClusterAABBSRV = GraphBuilder.CreateSRV(TransientResources.ClusterAABBBuffer, PF_R32_SINT);
+
+		FRDGBufferUAVRef ClusterAABBUAVSkipBarrier = GraphBuilder.CreateUAV(TransientResources.ClusterAABBBuffer, PF_R32_SINT, ERDGUnorderedAccessViewFlags::SkipBarrier);
+		FRDGBufferUAVRef GroupAABBUAVSkipBarrier = GraphBuilder.CreateUAV(TransientResources.GroupAABBBuffer, PF_R32_SINT, ERDGUnorderedAccessViewFlags::SkipBarrier);
+		AddClearAABBPass(GraphBuilder, ShaderMap, TotalClusterAABBCount, ClusterAABBUAVSkipBarrier);
+		AddClearAABBPass(GraphBuilder, ShaderMap, TotalRegisteredInstanceCount, GroupAABBUAVSkipBarrier);
+
+		for (FInstanceData& InstanceData : InstanceDatas)
+		{
+			InstanceData.CullingData = ImportCullingData(GraphBuilder, InstanceData.Instance->HairGroupPublicData);
+		}
 	}
 
 	{
@@ -1195,6 +1213,7 @@ static void RunHairStrandsInterpolation_Strands(
 	// 3. Compute cluster AABBs (used for LODing and voxelization)
 	{
 		FRDGBufferUAVRef GroupAABUAVSkipBarrier = GraphBuilder.CreateUAV(TransientResources.GroupAABBBuffer, PF_R32_SINT, ERDGUnorderedAccessViewFlags::SkipBarrier);
+		FRDGBufferUAVRef ClusterAABBUAVSkipBarrier = GraphBuilder.CreateUAV(TransientResources.ClusterAABBBuffer, PF_R32_SINT, ERDGUnorderedAccessViewFlags::SkipBarrier);
 
 		const FVector& TranslatedWorldOffset = View->ViewMatrices.GetPreViewTranslation();
 		for (FInstanceData& InstanceData : InstanceDatas)
@@ -1215,8 +1234,6 @@ static void RunHairStrandsInterpolation_Strands(
 				}
 			}
 				
-			InstanceData.Instance->HairGroupPublicData->SetClusterAABBValid(false);
-				
 			if (bNeedGPUAABB)
 			{
 				FRDGImportedBuffer Strands_CulledVertexCount = Register(GraphBuilder, InstanceData.Instance->HairGroupPublicData->GetDrawIndirectRasterComputeBuffer(), ERDGImportedBufferFlags::CreateSRV);
@@ -1230,12 +1247,12 @@ static void RunHairStrandsInterpolation_Strands(
 					InstanceData.Instance,
 					InstanceData.RDGResources.PositionOffsetSRV,
 					HairGroupCluster,
-					InstanceData.CullingData,
+					TransientResources.GetClusterOffset(InstanceData.Instance->RegisteredIndex),
+					TransientResources.GetClusterCount(InstanceData.Instance->RegisteredIndex),
 					InstanceData.RDGResources.PositionSRV,
 					Strands_CulledVertexCount.SRV,
+					ClusterAABBUAVSkipBarrier,
 					GroupAABUAVSkipBarrier);
-				
-				InstanceData.Instance->HairGroupPublicData->SetClusterAABBValid(UpdateType == EHairAABBUpdateType::UpdateClusterAABB);
 	
 				TransientResources.bIsGroupAABBValid[InstanceData.RegisteredIndex] = true;
 			}
@@ -2158,6 +2175,7 @@ void ProcessHairStrandsBookmark(
 			Parameters.AllViews,
 			Parameters.View,
 			Parameters.ViewUniqueID,
+			Parameters.Instances->Num(),
 			*Instances,
 			Parameters.ShaderPrintData,
 			*Parameters.TransientResources,
diff --git a/Engine/Plugins/Runtime/HairStrands/Source/HairStrandsCore/Private/GroomResources.cpp b/Engine/Plugins/Runtime/HairStrands/Source/HairStrandsCore/Private/GroomResources.cpp
--- a/Engine/Plugins/Runtime/HairStrands/Source/HairStrandsCore/Private/GroomResources.cpp
+++ b/Engine/Plugins/Runtime/HairStrands/Source/HairStrandsCore/Private/GroomResources.cpp
@@ -1326,7 +1326,6 @@ void FHairStrandsCullingResource::InternalAllocate(FRDGBuilder& GraphBuilder, ui
 		// Only allocated once, as cluster count does not change
 		if (Resources.DrawIndirectBuffer.Buffer == nullptr)
 		{
-			InternalCreateVertexBufferRDG<FHairStrandsSintFormat>(GraphBuilder, ClusterCount * 6, Resources.ClusterAABBBuffer, TEXT("Hair.Cluster_ClusterAABBBuffer"), GetOwnerName(), EHairResourceUsageType::Dynamic);
 			InternalCreateIndirectBufferRDG(GraphBuilder, Resources.DrawIndirectBuffer, TEXT("Hair.Cluster_DrawIndirectBuffer"), GetOwnerName());
 			InternalCreateIndirectBufferRDG(GraphBuilder, Resources.DrawIndirectRasterComputeBuffer, TEXT("Hair.Cluster_DrawIndirectRasterComputeBuffer"), GetOwnerName());
 			GraphBuilder.SetBufferAccessFinal(Register(GraphBuilder, Resources.DrawIndirectBuffer, ERDGImportedBufferFlags::None).Buffer, ERHIAccess::IndirectArgs);
@@ -1338,7 +1337,6 @@ void FHairStrandsCullingResource::InternalRelease()
 {
 	Resources.DrawIndirectBuffer.Release();
 	Resources.DrawIndirectRasterComputeBuffer.Release();
-	Resources.ClusterAABBBuffer.Release();
 	Resources.CulledCurveBuffer.Release();
 	Resources.CulledVertexIdBuffer.Release();
 	Resources.CulledVertexRadiusScaleBuffer.Release();
diff --git a/Engine/Plugins/Runtime/HairStrands/Source/HairStrandsCore/Private/HairStrandsClusterCulling.cpp b/Engine/Plugins/Runtime/HairStrands/Source/HairStrandsCore/Private/HairStrandsClusterCulling.cpp
--- a/Engine/Plugins/Runtime/HairStrands/Source/HairStrandsCore/Private/HairStrandsClusterCulling.cpp
+++ b/Engine/Plugins/Runtime/HairStrands/Source/HairStrandsCore/Private/HairStrandsClusterCulling.cpp
@@ -197,9 +197,7 @@ void AddInstanceToClusterData(FHairGroupInstance* In, FHairStrandClusterData& Ou
 	const int32 ClusterDataGroupIndex = Out.HairGroups.Num();
 	FHairStrandClusterData::FHairGroup& HairGroupCluster = Out.HairGroups.Emplace_GetRef();
 	HairGroupCluster.MaxPointPerCurve = In->Strands.Data ? In->Strands.Data->Header.MaxPointPerCurve : 0;
-	HairGroupCluster.ClusterScale = In->HairGroupPublicData->GetClusterScale();
-	HairGroupCluster.ClusterCount = In->HairGroupPublicData->GetClusterCount();
-	HairGroupCluster.ClusterAABBBuffer = &In->HairGroupPublicData->GetClusterAABBBuffer();
+	HairGroupCluster.ClusterScale = In->HairGroupPublicData->ClusterScale;
 
 	HairGroupCluster.CurveBuffer = &In->Strands.RestResource->CurveBuffer;
 	HairGroupCluster.PointLODBuffer = &In->Strands.ClusterResource->PointLODBuffer;
diff --git a/Engine/Plugins/Runtime/HairStrands/Source/HairStrandsCore/Private/HairStrandsClusterCulling.h b/Engine/Plugins/Runtime/HairStrands/Source/HairStrandsCore/Private/HairStrandsClusterCulling.h
--- a/Engine/Plugins/Runtime/HairStrands/Source/HairStrandsCore/Private/HairStrandsClusterCulling.h
+++ b/Engine/Plugins/Runtime/HairStrands/Source/HairStrandsCore/Private/HairStrandsClusterCulling.h
@@ -16,7 +16,6 @@ struct FHairStrandClusterData
 	struct FHairGroup
 	{
 		uint32 InstanceRegisteredIndex = ~0;
-		uint32 ClusterCount = 0;
 		float ClusterScale = 0;
 		uint32 MaxPointPerCurve = 0;
 		FVector4f ClusterInfoParameters = FVector4f::Zero();
@@ -28,7 +27,6 @@ struct FHairStrandClusterData
 		FRDGExternalBuffer* PointLODBuffer = nullptr;
 
 		// See FHairStrandsClusterResource fro details about those buffers.
-		FRDGExternalBuffer* ClusterAABBBuffer = nullptr;
 		FRDGExternalBuffer* ClusterInfoBuffer = nullptr; // SRV
 		FRDGExternalBuffer* CurveToClusterIdBuffer = nullptr; // SRV
 
@@ -47,7 +45,6 @@ struct FHairStrandClusterData
 	};
 
 	TArray<FHairGroup> HairGroups;
-	FRDGBufferRef SharedIndirectDispatchArgsBuffer = nullptr;
 };
 
 void AddInstanceToClusterData(
diff --git a/Engine/Plugins/Runtime/HairStrands/Source/HairStrandsCore/Private/HairStrandsInterpolation.cpp b/Engine/Plugins/Runtime/HairStrands/Source/HairStrandsCore/Private/HairStrandsInterpolation.cpp
--- a/Engine/Plugins/Runtime/HairStrands/Source/HairStrandsCore/Private/HairStrandsInterpolation.cpp
+++ b/Engine/Plugins/Runtime/HairStrands/Source/HairStrandsCore/Private/HairStrandsInterpolation.cpp
@@ -455,9 +455,6 @@ FRDGHairStrandsCullingData ImportCullingData(FRDGBuilder& GraphBuilder, FHairGro
 		Out.HairStrandsVF_CullingRadiusScaleBuffer	= Register(GraphBuilder, In->GetCulledVertexRadiusScaleBuffer(), ERDGImportedBufferFlags::CreateViews);
 	}
 
-	Out.ClusterCount		= In->GetClusterCount();
-	Out.ClusterAABBBuffer	= Register(GraphBuilder, In->GetClusterAABBBuffer(), ERDGImportedBufferFlags::CreateViews);
-
 	return Out;
 }
 
@@ -798,6 +795,7 @@ class FHairClusterAABBCS : public FGlobalShader
 		SHADER_PARAMETER(uint32, InstanceRegisteredIndex)
 		SHADER_PARAMETER(float, LODIndex)
 		SHADER_PARAMETER(float, ClusterScale)
+		SHADER_PARAMETER(uint32, ClusterOffset)
 		SHADER_PARAMETER(uint32, ClusterCount)
 		SHADER_PARAMETER(uint32, CurveCount)
 		SHADER_PARAMETER(FVector3f, CPUBoundMin)
@@ -834,9 +832,11 @@ void AddHairClusterAABBPass(
 	FHairGroupInstance* Instance,
 	FRDGBufferSRVRef RenderDeformedOffsetBuffer,
 	FHairStrandClusterData::FHairGroup* ClusterData,
-	FRDGHairStrandsCullingData& ClusterAABBData,
+	const uint32 ClusterOffset,
+	const uint32 ClusterCount,
 	FRDGBufferSRVRef RenderPositionBufferSRV,
 	FRDGBufferSRVRef& DrawIndirectRasterComputeBuffer,
+	FRDGBufferUAVRef ClusterAABBUAV,
 	FRDGBufferUAVRef GroupAABBBUAV)
 {
 	// Clusters AABB are only update if the groom is deformed.
@@ -844,7 +844,6 @@ void AddHairClusterAABBPass(
 	FTransform InRenLocalToTranslatedWorld = Instance->LocalToWorld;
 	InRenLocalToTranslatedWorld.AddToTranslation(TranslatedWorldOffset);
 	const FBoxSphereBounds TransformedBounds = Bounds.TransformBy(InRenLocalToTranslatedWorld);
-	Instance->HairGroupPublicData->SetClusterAABBValid(UpdateType == EHairAABBUpdateType::UpdateClusterAABB);
 
 	FHairClusterAABBCS::FParameters* Parameters = GraphBuilder.AllocParameters<FHairClusterAABBCS::FParameters>();
 	Parameters->InstanceRegisteredIndex = Instance->RegisteredIndex;
@@ -855,11 +854,12 @@ void AddHairClusterAABBPass(
 	Parameters->RenderDeformedPositionBuffer = RenderPositionBufferSRV;
 	Parameters->RenderDeformedOffsetBuffer = RenderDeformedOffsetBuffer;
 	Parameters->CurveCount = ActiveCurveCount;
-	Parameters->ClusterCount = ClusterData ? ClusterData->ClusterCount : 1;
+	Parameters->ClusterOffset = ClusterData ? ClusterOffset : 0;
+	Parameters->ClusterCount = ClusterData ? ClusterCount : 1;
 	Parameters->ClusterScale = ClusterData ? ClusterData->ClusterScale : 1.f;
 	Parameters->RenCurveBuffer = RegisterAsSRV(GraphBuilder, Instance->Strands.RestResource->CurveBuffer);
 	Parameters->RenPointLODBuffer = ClusterData ? RegisterAsSRV(GraphBuilder, *ClusterData->PointLODBuffer) : nullptr;
-	Parameters->OutClusterAABBBuffer = ClusterAABBData.ClusterAABBBuffer.UAV;
+	Parameters->OutClusterAABBBuffer = ClusterAABBUAV;
 	Parameters->OutGroupAABBBuffer = GroupAABBBUAV;
 
 	if (ShaderPrintData)
@@ -890,8 +890,6 @@ void AddHairClusterAABBPass(
 			Parameters,
 			FIntVector(1,1,1));
 	}
-
-	GraphBuilder.SetBufferAccessFinal(ClusterAABBData.ClusterAABBBuffer.Buffer, ERHIAccess::SRVMask);
 }
 
 ///////////////////////////////////////////////////////////////////////////////////////////////////
@@ -1242,17 +1240,14 @@ void AddGenerateRaytracingGeometryPass(
 
 ///////////////////////////////////////////////////////////////////////////////////////////////////
 
-class FClearClusterAABBCS : public FGlobalShader
+class FClearAABBCS : public FGlobalShader
 {
-	DECLARE_GLOBAL_SHADER(FClearClusterAABBCS);
-	SHADER_USE_PARAMETER_STRUCT(FClearClusterAABBCS, FGlobalShader);
+	DECLARE_GLOBAL_SHADER(FClearAABBCS);
+	SHADER_USE_PARAMETER_STRUCT(FClearAABBCS, FGlobalShader);
 
 	BEGIN_SHADER_PARAMETER_STRUCT(FParameters, )
-		SHADER_PARAMETER_RDG_BUFFER_UAV(RWBuffer, OutClusterAABBBuffer)
-		SHADER_PARAMETER_RDG_BUFFER_UAV(RWBuffer, OutGroupAABBBuffer)
-		SHADER_PARAMETER(uint32, InstanceRegisteredIndex)
-		SHADER_PARAMETER(uint32, ClusterCount)
-		SHADER_PARAMETER(uint32, bClearClusterAABBs)
+		SHADER_PARAMETER_RDG_BUFFER_UAV(RWBuffer, OutAABBBuffer)
+		SHADER_PARAMETER(uint32, AABBCount)
 	END_SHADER_PARAMETER_STRUCT()
 
 public:
@@ -1260,45 +1255,31 @@ public:
 	static void ModifyCompilationEnvironment(const FGlobalShaderPermutationParameters& Parameters, FShaderCompilerEnvironment& OutEnvironment)
 	{
 		FGlobalShader::ModifyCompilationEnvironment(Parameters, OutEnvironment);
-		OutEnvironment.SetDefine(TEXT("SHADER_CLEARCLUSTERAABB"), 1);
+		OutEnvironment.SetDefine(TEXT("SHADER_CLEARAABB"), 1);
 	}
 };
 
-IMPLEMENT_GLOBAL_SHADER(FClearClusterAABBCS, "/Engine/Private/HairStrands/HairStrandsClusterCulling.usf", "MainClearClusterAABBCS", SF_Compute);
+IMPLEMENT_GLOBAL_SHADER(FClearAABBCS, "/Engine/Private/HairStrands/HairStrandsClusterCulling.usf", "MainClearAABBCS", SF_Compute);
 
-void AddClearClusterAABBPass(
+void AddClearAABBPass(
 	FRDGBuilder& GraphBuilder,
 	FGlobalShaderMap* ShaderMap,
-	const EHairAABBUpdateType UpdateType,
-	uint32 InstanceRegisteredIndex,
-	uint32 ClusterCount,
-	FRDGImportedBuffer& OutClusterAABBBuffer,
-	FRDGBufferUAVRef& OutGroupAABBUAV)
+	uint32 AABBCount,
+	FRDGBufferUAVRef& OutAABBUAV)
 {
-	check(OutClusterAABBBuffer.Buffer);
-
-	FClearClusterAABBCS::FParameters* Parameters = GraphBuilder.AllocParameters<FClearClusterAABBCS::FParameters>();
-	Parameters->InstanceRegisteredIndex = InstanceRegisteredIndex;
-	Parameters->ClusterCount = ClusterCount;
-	Parameters->bClearClusterAABBs = UpdateType == EHairAABBUpdateType::UpdateClusterAABB ? 1 : 0;
-	Parameters->OutClusterAABBBuffer = OutClusterAABBBuffer.UAV;
-	Parameters->OutGroupAABBBuffer = OutGroupAABBUAV;
+	check(OutAABBUAV);
 
-	TShaderMapRef<FClearClusterAABBCS> ComputeShader(ShaderMap);
-
-	const FIntVector DispatchCount = 
-		UpdateType == EHairAABBUpdateType::UpdateClusterAABB ? 
-		FIntVector(FMath::DivideAndRoundUp(ClusterCount * 6u, 64u), 1, 1) :
-		FIntVector(1,1,1);
+	FClearAABBCS::FParameters* Parameters = GraphBuilder.AllocParameters<FClearAABBCS::FParameters>();
+	Parameters->AABBCount = AABBCount;
+	Parameters->OutAABBBuffer = OutAABBUAV;
 
+	TShaderMapRef<FClearAABBCS> ComputeShader(ShaderMap);
 	FComputeShaderUtils::AddPass(
 		GraphBuilder,
-		RDG_EVENT_NAME("HairStrands::ClearClusterAABB(%s)", UpdateType == EHairAABBUpdateType::UpdateClusterAABB ? TEXT("Cluster,Group") : TEXT("Group Only")),
+		RDG_EVENT_NAME("HairStrands::ClearAABBs"),
 		ComputeShader,
 		Parameters,
-		DispatchCount);
-
-	GraphBuilder.SetBufferAccessFinal(OutClusterAABBBuffer.Buffer, ERHIAccess::SRVMask);
+		FIntVector(FMath::DivideAndRoundUp(AABBCount * 6u, 64u), 1, 1));
 }
 
 ///////////////////////////////////////////////////////////////////////////////////////////////////
diff --git a/Engine/Plugins/Runtime/HairStrands/Source/HairStrandsCore/Public/GroomResources.h b/Engine/Plugins/Runtime/HairStrands/Source/HairStrandsCore/Public/GroomResources.h
--- a/Engine/Plugins/Runtime/HairStrands/Source/HairStrandsCore/Public/GroomResources.h
+++ b/Engine/Plugins/Runtime/HairStrands/Source/HairStrandsCore/Public/GroomResources.h
@@ -369,7 +369,6 @@ struct FHairStrandsCullingResource : public FHairCommonResource
 		uint32 Total = 0;
 		Total += GetBufferTotalNumBytes(Resources.DrawIndirectBuffer);
 		Total += GetBufferTotalNumBytes(Resources.DrawIndirectRasterComputeBuffer);
-		Total += GetBufferTotalNumBytes(Resources.ClusterAABBBuffer);
 		Total += GetBufferTotalNumBytes(Resources.CulledCurveBuffer);
 		Total += GetBufferTotalNumBytes(Resources.CulledVertexIdBuffer);
 		Total += GetBufferTotalNumBytes(Resources.CulledVertexRadiusScaleBuffer);
diff --git a/Engine/Plugins/Runtime/HairStrands/Source/HairStrandsCore/Public/HairStrandsInterpolation.h b/Engine/Plugins/Runtime/HairStrands/Source/HairStrandsCore/Public/HairStrandsInterpolation.h
--- a/Engine/Plugins/Runtime/HairStrands/Source/HairStrandsCore/Public/HairStrandsInterpolation.h
+++ b/Engine/Plugins/Runtime/HairStrands/Source/HairStrandsCore/Public/HairStrandsInterpolation.h
@@ -97,9 +97,6 @@ struct FRDGHairStrandsCullingData
 	FRDGImportedBuffer HairStrandsVF_CullingIndirectBuffer;
 	FRDGImportedBuffer HairStrandsVF_CullingIndexBuffer;
 	FRDGImportedBuffer HairStrandsVF_CullingRadiusScaleBuffer;
-
-	uint32 ClusterCount = 0;
-	FRDGImportedBuffer ClusterAABBBuffer;
 };
 
 FRDGHairStrandsCullingData ImportCullingData(FRDGBuilder& GraphBuilder, FHairGroupPublicData* In);
@@ -110,14 +107,11 @@ enum class EHairAABBUpdateType
 	UpdateGroupAABB
 };
 
-void AddClearClusterAABBPass(
+void AddClearAABBPass(
 	FRDGBuilder& GraphBuilder,
 	FGlobalShaderMap* ShaderMap,
-	const EHairAABBUpdateType UpdateType,
-	uint32 InstanceRegisteredIndex,
-	uint32 ClusterCount,
-	FRDGImportedBuffer& OutClusterAABBBuffer,
-	FRDGBufferUAVRef& OutGroupAABBUAV);
+	uint32 AABBCount,
+	FRDGBufferUAVRef& OutAABBUAV);
 
 void AddHairStrandsInterpolationPass(
 	FRDGBuilder& GraphBuilder,
diff --git a/Engine/Shaders/Private/HairStrands/HairStrandsClusterCulling.usf b/Engine/Shaders/Private/HairStrands/HairStrandsClusterCulling.usf
--- a/Engine/Shaders/Private/HairStrands/HairStrandsClusterCulling.usf
+++ b/Engine/Shaders/Private/HairStrands/HairStrandsClusterCulling.usf
@@ -34,38 +34,20 @@ uint3 ComputeDispatchGroupCount2D(uint GroupCount)
 }
 
 /////////////////////////////////////////////////////////////////////////////////////////////////////////////////////////////
-/// Clear clusters AABB
-#ifdef SHADER_CLEARCLUSTERAABB
-
-uint InstanceRegisteredIndex;
-uint ClusterCount;
-uint bClearClusterAABBs;
-
-RWBuffer<int> OutClusterAABBBuffer; // Cluster data packed as {uint3 AABBMin, uint3 AABBMax}
-RWBuffer<int> OutGroupAABBBuffer; 	// Group data packed as {uint3 AABBMin, uint3 AABBMax}
+/// Clear group/cluster AABBs
+#ifdef SHADER_CLEARAABB
+uint AABBCount;
+RWBuffer<int> OutAABBBuffer;
 
 [numthreads(64, 1, 1)]
-void MainClearClusterAABBCS(uint2 DispatchThreadId : SV_DispatchThreadID)
+void MainClearAABBCS(uint2 DispatchThreadId : SV_DispatchThreadID)
 {
-	// Clear 'clusters' AABB
-	if (bClearClusterAABBs > 0 && DispatchThreadId.x < ClusterCount*6)
+	if (DispatchThreadId.x < AABBCount*6)
 	{
 		const int ClusterId = DispatchThreadId.x / 6;
 		const bool bIsMin = (DispatchThreadId.x - ClusterId * 6) < 3 ? 1 : 0;
-		const int Value = bIsMin ? 2147483647 : -2147483648;
-		checkBufferAccessSlow(OutClusterAABBBuffer, DispatchThreadId.x);
-		OutClusterAABBBuffer[DispatchThreadId.x] = Value;
-	}
-
-	// Clear 'group' AABB
-	if (DispatchThreadId.x == 0)
-	{
-		OutGroupAABBBuffer[InstanceRegisteredIndex*6 + 0] =  2147483647;
-		OutGroupAABBBuffer[InstanceRegisteredIndex*6 + 1] =  2147483647;
-		OutGroupAABBBuffer[InstanceRegisteredIndex*6 + 2] =  2147483647;
-		OutGroupAABBBuffer[InstanceRegisteredIndex*6 + 3] = -2147483647;
-		OutGroupAABBBuffer[InstanceRegisteredIndex*6 + 4] = -2147483647;
-		OutGroupAABBBuffer[InstanceRegisteredIndex*6 + 5] = -2147483647;
+		checkBufferAccessSlow(OutAABBBuffer, DispatchThreadId.x);
+		OutAABBBuffer[DispatchThreadId.x] = bIsMin ? 2147483647 : -2147483648;
 	}
 }
 
@@ -456,6 +438,7 @@ void Main(uint DispatchThreadId : SV_DispatchThreadID, uint ThreadIndex : SV_Gro
 
 uint InstanceRegisteredIndex;
 uint CurveCount;
+uint ClusterOffset;
 uint ClusterCount;
 float ClusterScale;
 float3 CPUBoundMin;
@@ -585,7 +568,7 @@ groupshared float3 OutHairPositionOffset;
 	// Write out hair cluster AABB
 	if (ThreadIndex == 0)
 	{
-		const uint ClusterIdx6 = ClusterIndex * 6;
+		const uint ClusterIdx6 = (ClusterOffset + ClusterIndex) * 6;
 		OutClusterAABBBuffer[ClusterIdx6 + 0] = SharedClusterAABBMin[0].x;
 		OutClusterAABBBuffer[ClusterIdx6 + 1] = SharedClusterAABBMin[0].y;
 		OutClusterAABBBuffer[ClusterIdx6 + 2] = SharedClusterAABBMin[0].z;
diff --git a/Engine/Shaders/Private/HairStrands/HairStrandsDebug.usf b/Engine/Shaders/Private/HairStrands/HairStrandsDebug.usf
--- a/Engine/Shaders/Private/HairStrands/HairStrandsDebug.usf
+++ b/Engine/Shaders/Private/HairStrands/HairStrandsDebug.usf
@@ -2874,6 +2874,7 @@ void MainCS(uint GroupIndex : SV_GroupIndex, uint3 DispatchThreadId : SV_Dispatc
 #include "HairStrandsAABBCommon.ush"
 
 uint InstanceRegisteredIndex;
+uint ClusterOffset;
 uint ClusterCount;
 uint HairGroupId;
 uint PointCount;
@@ -2923,7 +2924,7 @@ void MainDrawDebugAABBCS(uint3 DispatchThreadId : SV_DispatchThreadID)
 		// Cluster
 		if (DispatchThreadId.x < ClusterCount)
 		{
-			const FHairAABB ClusterBound = ReadHairAABB(DispatchThreadId.x, ClusterAABBBuffer, 1);
+			const FHairAABB ClusterBound = ReadHairAABB(ClusterOffset + DispatchThreadId.x, ClusterAABBBuffer, 1);
 			AddAABBTWS(ClusterBound.Min, ClusterBound.Max, float4(0.5f, 1.0f, 0.0f, 0.3f));
 		}
 	}
diff --git a/Engine/Shaders/Private/HairStrands/HairStrandsVoxelRasterCompute.usf b/Engine/Shaders/Private/HairStrands/HairStrandsVoxelRasterCompute.usf
--- a/Engine/Shaders/Private/HairStrands/HairStrandsVoxelRasterCompute.usf
+++ b/Engine/Shaders/Private/HairStrands/HairStrandsVoxelRasterCompute.usf
@@ -160,7 +160,8 @@ void AllocatePageIndex(uint2 DispatchThreadId : SV_DispatchThreadID)
 
 #if SHADER_MARKVALID_PREPARE
 uint 			InstanceRegisteredIndex;
-uint			MaxClusterCount;
+uint			ClusterOffset;
+uint			ClusterCount;
 uint			MacroGroupId;
 uint			bUseMacroGroupBoundCPU;
 float3			MacroGroupBoundCPU_TranslatedWorldMinAABB;
@@ -181,10 +182,12 @@ RWBuffer<uint>	OutValidPageIndexBuffer;
 void MarkValid_PrepareCS(uint2 DispatchThreadId : SV_DispatchThreadID)
 {
 	const uint ClusterIndex = DispatchThreadId.x;
-	if (ClusterIndex >= MaxClusterCount)
+	if (ClusterIndex >= ClusterCount)
+	{
 		return;
+	}
 
-	const uint BaseClusterIndex = 6 * ClusterIndex;
+	const uint BaseClusterIndex = 6 * (ClusterOffset + ClusterIndex);
 
 	FHairAABB ClusterBound;
 	ClusterBound.Min.x = float(ClusterAABBsBuffer[BaseClusterIndex + 0]);
diff --git a/Engine/Source/Runtime/Renderer/Private/HairStrands/HairStrandsRendering.cpp b/Engine/Source/Runtime/Renderer/Private/HairStrands/HairStrandsRendering.cpp
--- a/Engine/Source/Runtime/Renderer/Private/HairStrands/HairStrandsRendering.cpp
+++ b/Engine/Source/Runtime/Renderer/Private/HairStrands/HairStrandsRendering.cpp
@@ -100,12 +100,11 @@ FHairTransientResources* AllocateHairTransientResourcse(FRDGBuilder& GraphBuilde
 		Out->bIsGroupAABBValid.Init(false, InstanceCount);
 
 		// Change this into a structure buffer
-		Out->GroupAABBBuffer = GraphBuilder.CreateBuffer(FRDGBufferDesc::CreateBufferDesc(4, 6 * InstanceCount), TEXT("Hair.TransientResources.GroupAABB"));
-		Out->GroupAABBUAV = GraphBuilder.CreateUAV(Out->GroupAABBBuffer, PF_R32_SINT);
+		Out->GroupAABBBuffer = GraphBuilder.CreateBuffer(FRDGBufferDesc::CreateBufferDesc(4, 6 * InstanceCount), TEXT("Hair.Transient.GroupAABB"));
 		Out->GroupAABBSRV = GraphBuilder.CreateSRV(Out->GroupAABBBuffer, PF_R32_SINT);
 
-		Out->IndirectDispatchArgsBuffer = GraphBuilder.CreateBuffer(FRDGBufferDesc::CreateIndirectDesc<FRHIDispatchIndirectParameters>(InstanceCount), TEXT("Hair.TransientResources.IndirectDispatchArgs"));
-		Out->IndirectDispatchArgsUAV = GraphBuilder.CreateUAV(Out->IndirectDispatchArgsBuffer);
+		// *4u* elements for storing DispatchCount.xyz and .w which contains the original number of items
+		Out->IndirectDispatchArgsBuffer = GraphBuilder.CreateBuffer(FRDGBufferDesc::CreateIndirectDesc(4u * InstanceCount), TEXT("Hair.Transient.IndirectDispatchArgs"));
 		Out->IndirectDispatchArgsSRV = GraphBuilder.CreateSRV(Out->IndirectDispatchArgsBuffer);
 	}
 	return Out;
diff --git a/Engine/Source/Runtime/Renderer/Private/HairStrands/HairStrandsVoxelization.cpp b/Engine/Source/Runtime/Renderer/Private/HairStrands/HairStrandsVoxelization.cpp
--- a/Engine/Source/Runtime/Renderer/Private/HairStrands/HairStrandsVoxelization.cpp
+++ b/Engine/Source/Runtime/Renderer/Private/HairStrands/HairStrandsVoxelization.cpp
@@ -278,7 +278,8 @@ class FVoxelMarkValidPageIndex_PrepareCS : public FGlobalShader
 
 	BEGIN_SHADER_PARAMETER_STRUCT(FParameters, )
 		SHADER_PARAMETER(uint32, InstanceRegisteredIndex)
-		SHADER_PARAMETER(uint32, MaxClusterCount)
+		SHADER_PARAMETER(uint32, ClusterOffset)
+		SHADER_PARAMETER(uint32, ClusterCount)
 		SHADER_PARAMETER(uint32, MacroGroupId)
 		SHADER_PARAMETER(uint32, bUseMacroGroupBoundCPU)
 		SHADER_PARAMETER(FVector3f, MacroGroupBoundCPU_TranslatedWorldMinAABB)
@@ -672,14 +673,15 @@ static void AddAllocateVoxelPagesPass(
 				// * Or use its clusters AABBs (for dynamic groom)
 				// * Or use its CPU AABB (for CLOD)
 				// Even if the bGroupAABBValid is invalid, it has been reset to invalid AABB, so the code below won't mark any page as used.
-				const bool bUseClusterAABB = HairGroupData->GetClusterAABBValid() && bIsGPUDriven;
+				const bool bUseClusterAABB = TransientResources.IsClusterAABBValid(HairGroupData->Instance->RegisteredIndex) && bIsGPUDriven;
 				
 				FVoxelMarkValidPageIndex_PrepareCS::FParameters* Parameters = GraphBuilder.AllocParameters<FVoxelMarkValidPageIndex_PrepareCS::FParameters>();
 				Parameters->InstanceRegisteredIndex						= HairGroupData->Instance->RegisteredIndex;
-				Parameters->MaxClusterCount								= HairGroupData->GetClusterCount();
+				Parameters->ClusterOffset								= TransientResources.GetClusterOffset(HairGroupData->Instance->RegisteredIndex);
+				Parameters->ClusterCount 								= TransientResources.GetClusterCount(HairGroupData->Instance->RegisteredIndex);
 				Parameters->MacroGroupId								= MacroGroup.MacroGroupId;
 				Parameters->GroupAABBsBuffer							= TransientResources.GroupAABBSRV;
-				Parameters->ClusterAABBsBuffer							= RegisterAsSRV(GraphBuilder, HairGroupData->GetClusterAABBBuffer());
+				Parameters->ClusterAABBsBuffer							= TransientResources.ClusterAABBSRV;
 				Parameters->MacroGroupVoxelAlignedAABBBuffer			= GraphBuilder.CreateSRV(MacroGroupResources.MacroGroupVoxelAlignedAABBsBuffer, PF_R32_SINT);
 				Parameters->PageIndexResolutionAndOffsetBuffer			= PageIndexResolutionAndOffsetBufferSRV;
 				Parameters->bUseMacroGroupBoundCPU						= bIsGPUDriven ? 0 : 1;
@@ -689,7 +691,7 @@ static void AddAllocateVoxelPagesPass(
 
 				FIntVector DispatchCount = bUseClusterAABB ?
 					FIntVector(
-						FMath::DivideAndRoundUp(Parameters->MaxClusterCount, FVoxelMarkValidPageIndex_PrepareCS::GetGroupSize(true)), 
+						FMath::DivideAndRoundUp(Parameters->ClusterCount, FVoxelMarkValidPageIndex_PrepareCS::GetGroupSize(true)), 
 						1, 
 						1) :
 					FIntVector(
diff --git a/Engine/Source/Runtime/Renderer/Public/HairStrandsInterface.h b/Engine/Source/Runtime/Renderer/Public/HairStrandsInterface.h
--- a/Engine/Source/Runtime/Renderer/Public/HairStrandsInterface.h
+++ b/Engine/Source/Runtime/Renderer/Public/HairStrandsInterface.h
@@ -220,7 +220,6 @@ public:
 	FRDGExternalBuffer& GetDrawIndirectRasterComputeBuffer() { return Culling->DrawIndirectRasterComputeBuffer; }
 	const FRDGExternalBuffer& GetDrawIndirectRasterComputeBuffer() const { return Culling->DrawIndirectRasterComputeBuffer; }
 	FRDGExternalBuffer& GetDrawIndirectBuffer() { return Culling->DrawIndirectBuffer; }
-	FRDGExternalBuffer& GetClusterAABBBuffer() { return Culling->ClusterAABBBuffer; }
 
 	const FRDGExternalBuffer& GetCulledCurveBuffer() const { return Culling->CulledCurveBuffer; }
 	const FRDGExternalBuffer& GetCulledVertexIdBuffer() const { return Culling->CulledVertexIdBuffer; }
@@ -232,9 +231,6 @@ public:
 
 	bool GetCullingResultAvailable() const { return Culling->bCullingResultAvailable; }
 	void SetCullingResultAvailable(bool b) { Culling->bCullingResultAvailable = b; }
-	
-	void SetClusterAABBValid(bool In) { Culling->bClusterAABBValid = In;  }
-	bool GetClusterAABBValid() const  { return Culling->bClusterAABBValid;  }
 
 	void SupportVoxelization(bool InVoxelize) { bSupportVoxelization = InVoxelize; }
 	bool DoesSupportVoxelization() const { return bSupportVoxelization; }
@@ -291,9 +287,6 @@ public:
 	void SetLODVisibility(bool bVisible) { bLODVisibility = bVisible; }
 	bool GetLODVisibility() const { return bLODVisibility; }
 
-	uint32 GetClusterCount() const { return ClusterCount;  }
-	float GetClusterScale() const { return ClusterScale;  }
-
 	// Return the number of active point/curve for strand geometry
 	RENDERER_API uint32 GetActiveStrandsPointCount() const;
 	RENDERER_API uint32 GetActiveStrandsCurveCount() const;
@@ -359,10 +352,6 @@ public:
 		FRDGExternalBuffer DrawIndirectBuffer;
 		FRDGExternalBuffer DrawIndirectRasterComputeBuffer;
 
-		/* Hair Cluster & Hair Group bounding box buffer */
-		FRDGExternalBuffer ClusterAABBBuffer;
-		bool bClusterAABBValid = false;
-
 		/* Culling & LODing results for a hair group */ // Better to be transient?
 		FRDGExternalBuffer CulledCurveBuffer;
 		FRDGExternalBuffer CulledVertexIdBuffer;
@@ -425,15 +414,23 @@ public:
 // Resources are indexed by hair instance's registered index
 struct FHairTransientResources
 {
+	FORCEINLINE uint32 GetClusterOffset(uint32 InRegisterIndex) const { check(ClusterAABBOffetAndCounts.IsValidIndex(InRegisterIndex)); return ClusterAABBOffetAndCounts[InRegisterIndex].X; }
+	FORCEINLINE uint32 GetClusterCount(uint32 InRegisterIndex)  const { check(ClusterAABBOffetAndCounts.IsValidIndex(InRegisterIndex)); return ClusterAABBOffetAndCounts[InRegisterIndex].Y; }
+	FORCEINLINE bool IsClusterAABBValid(uint32 InRegisterIndex) const { return GetClusterCount(InRegisterIndex) > 0; }
+
+	// Group AABB
+	TBitArray<> bIsGroupAABBValid;
 	FRDGBufferRef    GroupAABBBuffer = nullptr;
-	FRDGBufferUAVRef GroupAABBUAV = nullptr;
 	FRDGBufferSRVRef GroupAABBSRV = nullptr;
 
+	// Cluster AABB
+	TArray<FUintVector2> ClusterAABBOffetAndCounts;
+	FRDGBufferRef    ClusterAABBBuffer = nullptr;
+	FRDGBufferSRVRef ClusterAABBSRV = nullptr;
+
+	// Indirect dispatch args
 	FRDGBufferRef    IndirectDispatchArgsBuffer = nullptr;
-	FRDGBufferUAVRef IndirectDispatchArgsUAV = nullptr;
 	FRDGBufferSRVRef IndirectDispatchArgsSRV = nullptr;
-
-	TBitArray<> bIsGroupAABBValid;
 };
 
 ////////////////////////////////////////////////////////////////////////////////////////////////////
