Fix crash when enabling hair cluster debug view mode,

diff --git a/Engine/Plugins/Runtime/HairStrands/Source/HairStrandsCore/Private/GroomDebug.cpp b/Engine/Plugins/Runtime/HairStrands/Source/HairStrandsCore/Private/GroomDebug.cpp
--- a/Engine/Plugins/Runtime/HairStrands/Source/HairStrandsCore/Private/GroomDebug.cpp
+++ b/Engine/Plugins/Runtime/HairStrands/Source/HairStrandsCore/Private/GroomDebug.cpp
@@ -449,25 +449,41 @@ void AddDrawDebugClusterPass(
 		return;
 	}
 
+	const uint32 ClusterCount = TransientResources.ClusterAABBSRV ? TransientResources.ClusterAABBSRV->Desc.Buffer->Desc.NumElements / 6u : 16000u;
+	const uint32 InstanceCount = TransientResources.bIsGroupAABBValid.Num();
+	const uint32 LineCount = ClusterCount * 16 + InstanceCount * 16u;
+
 	ShaderPrint::SetEnabled(true);
-	ShaderPrint::RequestSpaceForLines(64000u);
-	ShaderPrint::RequestSpaceForCharacters(2000);
+	ShaderPrint::RequestSpaceForLines(LineCount);
+	ShaderPrint::RequestSpaceForCharacters(32 * InstanceCount + 2048);
 	if (!ShaderPrintData) { return; }
 
 	const bool bDebugAABB = ViewMode == EGroomViewMode::ClusterAABB;
 
-	uint32 DataIndex = 0;
+	// Sort cluster group by registered index to get stable listing
+	TArray<const FHairStrandClusterData::FHairGroup*> Groups;
+	Groups.Reserve(HairClusterData.HairGroups.Num());
 	for (const FHairStrandClusterData::FHairGroup& HairGroupClusters : HairClusterData.HairGroups)
+	{
+		Groups.Add(&HairGroupClusters);
+	}
+	Groups.Sort([](const FHairStrandClusterData::FHairGroup& A, const FHairStrandClusterData::FHairGroup& B) 
+	{ 
+		return A.InstanceRegisteredIndex < B.InstanceRegisteredIndex;
+	});
+
+	uint32 DataIndex = 0;
+	for (const FHairStrandClusterData::FHairGroup* Group : Groups)
 	{
 		TShaderMapRef<FDrawDebugClusterAABBCS> ComputeShader(ShaderMap);
 
 		FDrawDebugClusterAABBCS::FParameters* Parameters = GraphBuilder.AllocParameters<FDrawDebugClusterAABBCS::FParameters>();
-		Parameters->InstanceRegisteredIndex = HairGroupClusters.InstanceRegisteredIndex;
+		Parameters->InstanceRegisteredIndex = Group->InstanceRegisteredIndex;
 		Parameters->ViewUniformBuffer = View.ViewUniformBuffer;
-		Parameters->ClusterOffset = TransientResources.GetClusterOffset(HairGroupClusters.InstanceRegisteredIndex);
-		Parameters->ClusterCount  = TransientResources.GetClusterCount(HairGroupClusters.InstanceRegisteredIndex);
-		Parameters->PointCount = HairGroupClusters.HairGroupPublicPtr->GetActiveStrandsPointCount();
-		Parameters->CurveCount = HairGroupClusters.HairGroupPublicPtr->GetActiveStrandsCurveCount();
+		Parameters->ClusterOffset = TransientResources.GetClusterOffset(Group->InstanceRegisteredIndex);
+		Parameters->ClusterCount  = TransientResources.GetClusterCount(Group->InstanceRegisteredIndex);
+		Parameters->PointCount = Group->HairGroupPublicPtr->GetActiveStrandsPointCount();
+		Parameters->CurveCount = Group->HairGroupPublicPtr->GetActiveStrandsCurveCount();
 		Parameters->HairGroupId = DataIndex++;
 		Parameters->bDrawAABB = ViewMode == EGroomViewMode::ClusterAABB ? 1 : 0;
 		Parameters->ClusterAABBBuffer = TransientResources.ClusterAABBSRV;
diff --git a/Engine/Shaders/Private/HairStrands/HairStrandsDebug.usf b/Engine/Shaders/Private/HairStrands/HairStrandsDebug.usf
--- a/Engine/Shaders/Private/HairStrands/HairStrandsDebug.usf
+++ b/Engine/Shaders/Private/HairStrands/HairStrandsDebug.usf
@@ -2901,10 +2901,10 @@ void MainDrawDebugAABBCS(uint3 DispatchThreadId : SV_DispatchThreadID)
 		{
 			if (HairGroupId == 0)
 			{
-				Print(Context, TEXT("Group Clusters Curves Points"), FontYellow);
+				Print(Context, TEXT("Reg.Idx Clusters Curves Points"), FontYellow);
 				Newline(Context);
 			}
-			Print(Context, HairGroupId, FontOrange, 6, 0);
+			Print(Context, InstanceRegisteredIndex, FontOrange, 8, 0);
 			Print(Context, ClusterCount, FontValue, 9, 0);
 			PrintLargeNumber(Context, CurveCount, FontValue);
 			Print(Context, TEXT(" "), FontYellow);
