Fix hair resource transitions pass not support on some platorms.

diff --git a/Engine/Plugins/Runtime/HairStrands/Source/HairStrandsCore/Private/GroomManager.cpp b/Engine/Plugins/Runtime/HairStrands/Source/HairStrandsCore/Private/GroomManager.cpp
--- a/Engine/Plugins/Runtime/HairStrands/Source/HairStrandsCore/Private/GroomManager.cpp
+++ b/Engine/Plugins/Runtime/HairStrands/Source/HairStrandsCore/Private/GroomManager.cpp
@@ -890,7 +890,7 @@ static void RunHairStrandsInterpolation_Strands(
 				ClusterDatas,
 				IndirectDispatchArgsUAVWithSkipBarrier);
 
-			AddTransitionPass(GraphBuilder, ShaderMap, Transitions);
+			AddTransitionPass(GraphBuilder, ShaderMap, View->GetShaderPlatform(), Transitions);
 		}
 	}
 
@@ -912,7 +912,7 @@ static void RunHairStrandsInterpolation_Strands(
 				Transitions.Add(RegisterAsSRV(GraphBuilder, InstanceData.Instance->Strands.DeformedRootResource->LODs[InstanceData.MeshLODIndex].GetDeformedUniqueTrianglePositionBuffer(FHairStrandsDeformedRootResource::FLOD::Current)));
 			}
 		}
-		AddTransitionPass(GraphBuilder, ShaderMap, Transitions);
+		AddTransitionPass(GraphBuilder, ShaderMap, View->GetShaderPlatform(), Transitions);
 	}
 
 	// Update position offset 
@@ -983,7 +983,7 @@ static void RunHairStrandsInterpolation_Strands(
 			}
 		}
 
-		AddTransitionPass(GraphBuilder, ShaderMap, Transitions);
+		AddTransitionPass(GraphBuilder, ShaderMap, View->GetShaderPlatform(), Transitions);
 	}
 
 
@@ -1063,7 +1063,7 @@ static void RunHairStrandsInterpolation_Strands(
 			}
 		}
 
-		AddTransitionPass(GraphBuilder, ShaderMap, Transitions);
+		AddTransitionPass(GraphBuilder, ShaderMap, View->GetShaderPlatform(), Transitions);
 	}
 
 	// Previous Position
diff --git a/Engine/Source/Runtime/Renderer/Private/HairStrands/HairStrandsUtils.cpp b/Engine/Source/Runtime/Renderer/Private/HairStrands/HairStrandsUtils.cpp
--- a/Engine/Source/Runtime/Renderer/Private/HairStrands/HairStrandsUtils.cpp
+++ b/Engine/Source/Runtime/Renderer/Private/HairStrands/HairStrandsUtils.cpp
@@ -367,9 +367,13 @@ private:
 		END_SHADER_PARAMETER_STRUCT()
 
 public:
-		static bool ShouldCompilePermutation(const FGlobalShaderPermutationParameters& Parameters)
+	static bool IsSupported(EShaderPlatform InPlatform)
 	{
-		return IsHairStrandsSupported(EHairStrandsShaderType::All, Parameters.Platform);
+		return IsHairStrandsSupported(EHairStrandsShaderType::Strands, InPlatform);
+	}
+	static bool ShouldCompilePermutation(const FGlobalShaderPermutationParameters& Parameters)
+	{
+		return IsSupported(Parameters.Platform);
 	}
 	static void ModifyCompilationEnvironment(const FGlobalShaderPermutationParameters& Parameters, FShaderCompilerEnvironment& OutEnvironment)
 	{
@@ -383,10 +387,11 @@ IMPLEMENT_GLOBAL_SHADER(FHairResourceTransitionPass, "/Engine/Private/HairStrand
 void AddTransitionPass(
 	FRDGBuilder& GraphBuilder,
 	FGlobalShaderMap* ShaderMap,
+	EShaderPlatform InPlatform,
 	const TArray<FRDGBufferSRVRef>& Transitions)
 {
 	const uint32 ResourceCount = Transitions.Num();
-	if (ResourceCount == 0)
+	if (ResourceCount == 0 || !FHairResourceTransitionPass::IsSupported(InPlatform))
 	{
 		return;
 	}
diff --git a/Engine/Source/Runtime/Renderer/Private/HairStrands/HairStrandsVoxelization.cpp b/Engine/Source/Runtime/Renderer/Private/HairStrands/HairStrandsVoxelization.cpp
--- a/Engine/Source/Runtime/Renderer/Private/HairStrands/HairStrandsVoxelization.cpp
+++ b/Engine/Source/Runtime/Renderer/Private/HairStrands/HairStrandsVoxelization.cpp
@@ -1427,7 +1427,7 @@ void VoxelizeHairStrands(
 
 		}
 	}
-	AddTransitionPass(GraphBuilder, View.ShaderMap, Transitions);
+	AddTransitionPass(GraphBuilder, View.ShaderMap, View.GetShaderPlatform(), Transitions);
 
 	FRDGBufferRef PageToPageIndexBuffer = nullptr;
 	FHairStrandsViewStateData* HairStrandsViewStateData = View.ViewState ? &View.ViewState->HairStrandsViewStateData : nullptr;
diff --git a/Engine/Source/Runtime/Renderer/Public/HairStrandsInterface.h b/Engine/Source/Runtime/Renderer/Public/HairStrandsInterface.h
--- a/Engine/Source/Runtime/Renderer/Public/HairStrandsInterface.h
+++ b/Engine/Source/Runtime/Renderer/Public/HairStrandsInterface.h
@@ -166,6 +166,7 @@ RENDERER_API FRDGBufferUAVRef   RegisterAsUAV(FRDGBuilder& GraphBuilder, const F
 RENDERER_API void AddTransitionPass(
 	FRDGBuilder& GraphBuilder,
 	FGlobalShaderMap* ShaderMap,
+	EShaderPlatform InPlatform,
 	const TArray<FRDGBufferSRVRef>& Transitions);
 
 ////////////////////////////////////////////////////////////////////////////////////////////////////
