Fix missing RBF weight update pass.

diff --git a/Engine/Plugins/Runtime/HairStrands/Source/HairStrandsCore/Private/GroomManager.cpp b/Engine/Plugins/Runtime/HairStrands/Source/HairStrandsCore/Private/GroomManager.cpp
--- a/Engine/Plugins/Runtime/HairStrands/Source/HairStrandsCore/Private/GroomManager.cpp
+++ b/Engine/Plugins/Runtime/HairStrands/Source/HairStrandsCore/Private/GroomManager.cpp
@@ -493,6 +493,7 @@ static void RunHairStrandsInterpolation_Guide(
 		}
 	}
 
+	// RBF sample position update
 	// Guide update need to run only if simulation is enabled, or if RBF is enabled (since RFB are transfer through guides)
 	for (FInstanceData& InstanceData : InstanceDatas)
 	{
@@ -511,6 +512,25 @@ static void RunHairStrandsInterpolation_Guide(
 		}
 	}
 
+	// RBF weights update
+	// Guide update need to run only if simulation is enabled, or if RBF is enabled (since RFB are transfer through guides)
+	for (FInstanceData& InstanceData : InstanceDatas)
+	{
+		if (InstanceData.NeedsMeshUpdate())
+		{
+			if (InstanceData.bGlobalDeformationEnable)
+			{
+				AddHairStrandUpdateMeshSamplesPass(
+					GraphBuilder,
+					ShaderMap,
+					InstanceData.Instance->Debug.MeshLODIndex,
+					InstanceData.MeshDataLOD,
+					InstanceData.Instance->Guides.RestRootResource,
+					InstanceData.Instance->Guides.DeformedRootResource);
+			}
+		}
+	}
+
 	// Update position offset for each instance (GPU)
 	for (FInstanceData& InstanceData : InstanceDatas)
 	{
