Fix RHI validation error with hair.

diff --git a/Engine/Plugins/Runtime/HairStrands/Shaders/Private/NiagaraDataInterfaceHairStrandsTemplate.ush b/Engine/Plugins/Runtime/HairStrands/Shaders/Private/NiagaraDataInterfaceHairStrandsTemplate.ush
--- a/Engine/Plugins/Runtime/HairStrands/Shaders/Private/NiagaraDataInterfaceHairStrandsTemplate.ush
+++ b/Engine/Plugins/Runtime/HairStrands/Shaders/Private/NiagaraDataInterfaceHairStrandsTemplate.ush
@@ -24,7 +24,7 @@ int					{ParameterName}_InterpolationMode;
 float3				{ParameterName}_RestRootOffset;
 float3				{ParameterName}_DeformedRootOffset;
 float3				{ParameterName}_RestPositionOffset;
-Buffer<float4>		{ParameterName}_DeformedPositionOffset;
+StructuredBuffer<float4> {ParameterName}_DeformedPositionOffset;
 RWBuffer<uint>	    {ParameterName}_BoundingBoxBuffer;
 uint				{ParameterName}_ResetSimulation;
 uint				{ParameterName}_RestUpdate;
diff --git a/Engine/Plugins/Runtime/HairStrands/Source/HairStrandsCore/Private/Deformer/DeformerDataInterfaceGroomWrite.cpp b/Engine/Plugins/Runtime/HairStrands/Source/HairStrandsCore/Private/Deformer/DeformerDataInterfaceGroomWrite.cpp
--- a/Engine/Plugins/Runtime/HairStrands/Source/HairStrandsCore/Private/Deformer/DeformerDataInterfaceGroomWrite.cpp
+++ b/Engine/Plugins/Runtime/HairStrands/Source/HairStrandsCore/Private/Deformer/DeformerDataInterfaceGroomWrite.cpp
@@ -14,6 +14,7 @@
 #include "GroomInstance.h"
 #include "DeformerGroomComponentSource.h"
 #include "RenderGraphUtils.h"
+#include "SystemTextures.h"
 
 FString UOptimusGroomWriteDataInterface::GetDisplayName() const
 {
@@ -214,9 +215,11 @@ void FOptimusGroomWriteDataProviderProxy::AllocateResources(FRDGBuilder& GraphBu
 			}
 			else
 			{
-				R.PositionOffsetBufferSRV 	= GraphBuilder.CreateSRV(GraphBuilder.RegisterExternalBuffer(GWhiteVertexBufferWithRDG->Buffer), PF_A32B32G32R32F);
+				FRDGBufferRef DummyUAVBuffer = GraphBuilder.CreateBuffer(FRDGBufferDesc::CreateBufferDesc(8u, 1), TEXT("Hair.Deformer.DummyBuffer"));
+				
+				R.PositionOffsetBufferSRV 	= GraphBuilder.CreateSRV(GSystemTextures.GetDefaultStructuredBuffer(GraphBuilder, 16u), PF_A32B32G32R32F); //GraphBuilder.CreateSRV(GraphBuilder.RegisterExternalBuffer(GWhiteVertexBufferWithRDG->Buffer), PF_A32B32G32R32F);
 				R.PositionBufferSRV 		= GraphBuilder.CreateSRV(GraphBuilder.RegisterExternalBuffer(GWhiteVertexBufferWithRDG->Buffer), PF_R16G16B16A16_UINT);
-				R.PositionBufferUAV 		= GraphBuilder.CreateUAV(GraphBuilder.RegisterExternalBuffer(GWhiteVertexBufferWithRDG->Buffer), PF_R16G16B16A16_UINT);
+				R.PositionBufferUAV 		= GraphBuilder.CreateUAV(DummyUAVBuffer, PF_R16G16B16A16_UINT);
 			}
 			
 			// Curve Attributes
diff --git a/Engine/Plugins/Runtime/HairStrands/Source/HairStrandsCore/Private/GroomResources.cpp b/Engine/Plugins/Runtime/HairStrands/Source/HairStrandsCore/Private/GroomResources.cpp
--- a/Engine/Plugins/Runtime/HairStrands/Source/HairStrandsCore/Private/GroomResources.cpp
+++ b/Engine/Plugins/Runtime/HairStrands/Source/HairStrandsCore/Private/GroomResources.cpp
@@ -1430,7 +1430,7 @@ void FHairStrandsRestRootResource::InternalAllocate(FRDGBuilder& GraphBuilder, u
 			const uint32 WeightCount = FGroomRBFDeformer::GetWeightCount(LODHeader.SampleCount);
 			check(WeightCount == InteroplationWeightCount); // Sanity check
 
-			InternalCreateVertexBufferRDG_FromHairBulkData<FHairStrandsWeightFormat>(GraphBuilder, CPUData.MeshInterpolationWeightsBuffer, InteroplationWeightCount, GPUData.MeshInterpolationWeightsBuffer, ToHairResourceDebugName(HAIRSTRANDS_RESOUCE_NAME(CurveType, Hair.StrandsRestRoot_MeshInterpolationWeightsBuffer), ResourceName), OwnerName, EHairResourceUsageType::Static);
+			InternalCreateStructuredBufferRDG_FromHairBulkData<FHairStrandsWeightFormat>(GraphBuilder, CPUData.MeshInterpolationWeightsBuffer, InteroplationWeightCount, GPUData.MeshInterpolationWeightsBuffer, ToHairResourceDebugName(HAIRSTRANDS_RESOUCE_NAME(CurveType, Hair.StrandsRestRoot_MeshInterpolationWeightsBuffer), ResourceName), OwnerName, EHairResourceUsageType::Static);
 			InternalCreateVertexBufferRDG_FromHairBulkData<FHairStrandsIndexFormat>(GraphBuilder, CPUData.MeshSampleIndicesBuffer, LODHeader.SampleCount, GPUData.MeshSampleIndicesBuffer, ToHairResourceDebugName(HAIRSTRANDS_RESOUCE_NAME(CurveType, Hair.StrandsRestRoot_MeshSampleIndicesBuffer), ResourceName), OwnerName, EHairResourceUsageType::Static);
 			InternalCreateStructuredBufferRDG_FromHairBulkData<FHairStrandsMeshTrianglePositionFormat>(GraphBuilder, CPUData.RestSamplePositionsBuffer, LODHeader.SampleCount, GPUData.RestSamplePositionsBuffer, ToHairResourceDebugName(HAIRSTRANDS_RESOUCE_NAME(CurveType, Hair.StrandsRestRoot_RestSamplePositionsBuffer), ResourceName), OwnerName, EHairResourceUsageType::Static);
 		}
@@ -1516,13 +1516,13 @@ void FHairStrandsDeformedRootResource::InternalAllocate(FRDGBuilder& GraphBuilde
 			{
 				const uint32 RBFEntryCount = FGroomRBFDeformer::GetEntryCount(LOD.SampleCount);
 
-				InternalCreateVertexBufferRDG<FHairStrandsMeshTrianglePositionFormat>(GraphBuilder, LOD.SampleCount, LOD.DeformedSamplePositionsBuffer[0], ToHairResourceDebugName(HAIRSTRANDS_RESOUCE_NAME(CurveType, Hair.StrandsRootDeformed_DeformedSamplePositionsBuffer0), ResourceName), OwnerName, EHairResourceUsageType::Dynamic);
+				InternalCreateStructuredBufferRDG<FHairStrandsMeshTrianglePositionFormat>(GraphBuilder, LOD.SampleCount, LOD.DeformedSamplePositionsBuffer[0], ToHairResourceDebugName(HAIRSTRANDS_RESOUCE_NAME(CurveType, Hair.StrandsRootDeformed_DeformedSamplePositionsBuffer0), ResourceName), OwnerName, EHairResourceUsageType::Dynamic);
 				InternalCreateStructuredBufferRDG<FHairStrandsMeshTrianglePositionFormat>(GraphBuilder, RBFEntryCount, LOD.MeshSampleWeightsBuffer[0], ToHairResourceDebugName(HAIRSTRANDS_RESOUCE_NAME(CurveType, Hair.StrandsRootDeformed_MeshSampleWeightsBuffer0), ResourceName), OwnerName, EHairResourceUsageType::Dynamic);
 
 				// Double buffering is disabled by default unless the read-only cvar r.HairStrands.ContinuousDecimationReordering is set
 				if (IsHairStrandContinuousDecimationReorderingEnabled())
 				{
-					InternalCreateVertexBufferRDG<FHairStrandsMeshTrianglePositionFormat>(GraphBuilder, LOD.SampleCount, LOD.DeformedSamplePositionsBuffer[1], ToHairResourceDebugName(HAIRSTRANDS_RESOUCE_NAME(CurveType, Hair.StrandsRootDeformed_DeformedSamplePositionsBuffer1), ResourceName), OwnerName, EHairResourceUsageType::Dynamic);
+					InternalCreateStructuredBufferRDG<FHairStrandsMeshTrianglePositionFormat>(GraphBuilder, LOD.SampleCount, LOD.DeformedSamplePositionsBuffer[1], ToHairResourceDebugName(HAIRSTRANDS_RESOUCE_NAME(CurveType, Hair.StrandsRootDeformed_DeformedSamplePositionsBuffer1), ResourceName), OwnerName, EHairResourceUsageType::Dynamic);
 					InternalCreateStructuredBufferRDG<FHairStrandsMeshTrianglePositionFormat>(GraphBuilder, RBFEntryCount, LOD.MeshSampleWeightsBuffer[1], ToHairResourceDebugName(HAIRSTRANDS_RESOUCE_NAME(CurveType, Hair.StrandsRootDeformed_MeshSampleWeightsBuffer1), ResourceName), OwnerName, EHairResourceUsageType::Dynamic);
 				}
 			}
diff --git a/Engine/Plugins/Runtime/HairStrands/Source/HairStrandsCore/Private/HairStrandsMeshProjection.cpp b/Engine/Plugins/Runtime/HairStrands/Source/HairStrandsCore/Private/HairStrandsMeshProjection.cpp
--- a/Engine/Plugins/Runtime/HairStrands/Source/HairStrandsCore/Private/HairStrandsMeshProjection.cpp
+++ b/Engine/Plugins/Runtime/HairStrands/Source/HairStrandsCore/Private/HairStrandsMeshProjection.cpp
@@ -739,7 +739,7 @@ private:
 		SHADER_PARAMETER(uint32, MaxSampleCount)
 
 		SHADER_PARAMETER_RDG_BUFFER_SRV(Buffer, SampleIndicesBuffer)
-		SHADER_PARAMETER_RDG_BUFFER_SRV(Buffer, InterpolationWeightsBuffer)
+		SHADER_PARAMETER_RDG_BUFFER_SRV(StructuredBuffer, InterpolationWeightsBuffer)
 		SHADER_PARAMETER_RDG_BUFFER_SRV(StructuredBuffer, SampleRestPositionsBuffer)
 		SHADER_PARAMETER_RDG_BUFFER_SRV(StructuredBuffer, SampleDeformedPositionsBuffer)
 		SHADER_PARAMETER_RDG_BUFFER_UAV(RWStructuredBuffer, OutSampleDeformationsBuffer)
diff --git a/Engine/Shaders/Private/HairStrands/HairStrandsMesh.usf b/Engine/Shaders/Private/HairStrands/HairStrandsMesh.usf
--- a/Engine/Shaders/Private/HairStrands/HairStrandsMesh.usf
+++ b/Engine/Shaders/Private/HairStrands/HairStrandsMesh.usf
@@ -82,7 +82,7 @@ Buffer<float> RDGVertexPositionsBuffer7;
 uint MaxVertexCount;
 uint MaxSampleCount;
 Buffer<uint> SampleIndicesBuffer;
-RWBuffer<float4> OutSamplePositionsBuffer;
+RWStructuredBuffer<float4> OutSamplePositionsBuffer;
 
 [numthreads(128, 1, 1)]
 void MainCS(uint3 DispatchThreadId : SV_DispatchThreadID)
@@ -178,12 +178,12 @@ void MainCS(uint3 DispatchThreadId : SV_DispatchThreadID)
 
 uint MaxSampleCount;
 Buffer<uint> SampleIndicesBuffer;
-Buffer<float> InterpolationWeightsBuffer;
+StructuredBuffer<float> InterpolationWeightsBuffer;
 
-Buffer<float4> SampleRestPositionsBuffer;
-Buffer<float4> SampleDeformedPositionsBuffer;
+StructuredBuffer<float4> SampleRestPositionsBuffer;
+StructuredBuffer<float4> SampleDeformedPositionsBuffer;
 
-RWBuffer<float4> OutSampleDeformationsBuffer; 
+RWStructuredBuffer<float4> OutSampleDeformationsBuffer; 
 
 [numthreads(128, 1, 1)]
 void MainCS(uint3 DispatchThreadId : SV_DispatchThreadID)
