Fix validation error on Vulkan due to mismatchin buffer type.

diff --git a/Engine/Plugins/Runtime/HairStrands/Source/HairStrandsCore/Private/HairStrandsInterpolation.cpp b/Engine/Plugins/Runtime/HairStrands/Source/HairStrandsCore/Private/HairStrandsInterpolation.cpp
--- a/Engine/Plugins/Runtime/HairStrands/Source/HairStrandsCore/Private/HairStrandsInterpolation.cpp
+++ b/Engine/Plugins/Runtime/HairStrands/Source/HairStrandsCore/Private/HairStrandsInterpolation.cpp
@@ -1151,7 +1151,7 @@ class FHairRaytracingGeometryCS : public FGlobalShader
 		SHADER_PARAMETER_RDG_BUFFER_SRV(Buffer, PositionBuffer)
 		SHADER_PARAMETER_RDG_BUFFER_SRV(Buffer, TangentBuffer)
 		SHADER_PARAMETER_RDG_BUFFER_UAV(RWBuffer, OutputPositionBuffer)
-		SHADER_PARAMETER_RDG_BUFFER_UAV(RWBuffer, OutputIndexBuffer)
+		SHADER_PARAMETER_RDG_BUFFER_UAV(RWStructuredBuffer, OutputIndexBuffer)
 		SHADER_PARAMETER(uint32, RaytracingProceduralSplits)
 		SHADER_PARAMETER_STRUCT_INCLUDE(ShaderPrint::FShaderParameters, ShaderPrintUniformBuffer)
 	END_SHADER_PARAMETER_STRUCT()
diff --git a/Engine/Shaders/Private/HairStrands/HairStrandsRaytracingGeometry.usf b/Engine/Shaders/Private/HairStrands/HairStrandsRaytracingGeometry.usf
--- a/Engine/Shaders/Private/HairStrands/HairStrandsRaytracingGeometry.usf
+++ b/Engine/Shaders/Private/HairStrands/HairStrandsRaytracingGeometry.usf
@@ -22,7 +22,7 @@ RWBuffer<float4> OutputPositionBuffer;
 #if PERMUTATION_PROCEDURAL_PRIMITIVE 
 uint             RaytracingProceduralSplits;
 #else
-RWBuffer<uint>   OutputIndexBuffer;
+RWStructuredBuffer<uint> OutputIndexBuffer;
 #endif
 
 uint HairStrandsVF_bCullingEnable;
