Fix assert due to Skin cache not returning the same buffer for all section of a single mesh.

diff --git a/Engine/Plugins/Runtime/HairStrands/Source/HairStrandsCore/Private/HairStrandsMeshProjection.cpp b/Engine/Plugins/Runtime/HairStrands/Source/HairStrandsCore/Private/HairStrandsMeshProjection.cpp
--- a/Engine/Plugins/Runtime/HairStrands/Source/HairStrandsCore/Private/HairStrandsMeshProjection.cpp
+++ b/Engine/Plugins/Runtime/HairStrands/Source/HairStrandsCore/Private/HairStrandsMeshProjection.cpp
@@ -245,7 +245,8 @@ bool AddHairStrandUpdateMeshTrianglesPass(
 		uint32 IndexBaseIndex;
 		uint32 UVsChannelOffset : 8;
 		uint32 UVsChannelCount : 8;
-		uint32 Pad : 16;
+		uint32 bIsSwapped : 8;
+		uint32 Pad : 8;
 	};
 	TArray<FSectionData> SectionDatas;
 	SectionDatas.SetNum(SectionCount);
@@ -257,15 +258,21 @@ bool AddHairStrandUpdateMeshTrianglesPass(
 		SectionDatas[SectionIt].IndexBaseIndex	= MeshSectionData.IndexBaseIndex;
 		SectionDatas[SectionIt].UVsChannelOffset= MeshSectionData.UVsChannelOffset;
 		SectionDatas[SectionIt].UVsChannelCount	= MeshSectionData.UVsChannelCount;
+		SectionDatas[SectionIt].bIsSwapped		= MeshData.Sections[SectionIt].PositionBuffer != CommonParameters.MeshPositionBuffer ? 1u : 0u;
 
 		// Sanity check
 		check(MeshSectionData.SectionIndex == SectionIt);
 		check(MeshSectionData.UVsChannelOffset < 255);
 		check(MeshSectionData.UVsChannelCount < 255);
-		check(CommonParameters.RDGMeshPositionBuffer		== MeshData.Sections[SectionIt].RDGPositionBuffer);
-		check(CommonParameters.RDGMeshPreviousPositionBuffer== MeshData.Sections[SectionIt].RDGPreviousPositionBuffer);
-		check(CommonParameters.MeshPositionBuffer			== MeshData.Sections[SectionIt].PositionBuffer);
-		check(CommonParameters.MeshPreviousPositionBuffer	== MeshData.Sections[SectionIt].PreviousPositionBuffer);
+#if 1 // Relaxed check, with optional buffer swap
+		check(CommonParameters.RDGMeshPositionBuffer == MeshData.Sections[SectionIt].RDGPositionBuffer || CommonParameters.RDGMeshPreviousPositionBuffer == MeshData.Sections[SectionIt].RDGPositionBuffer);
+		check(CommonParameters.MeshPositionBuffer    == MeshData.Sections[SectionIt].PositionBuffer    || CommonParameters.MeshPreviousPositionBuffer    == MeshData.Sections[SectionIt].PositionBuffer);
+#else
+//		check(CommonParameters.RDGMeshPositionBuffer		== MeshData.Sections[SectionIt].RDGPositionBuffer);
+//		check(CommonParameters.RDGMeshPreviousPositionBuffer== MeshData.Sections[SectionIt].RDGPreviousPositionBuffer);
+//		check(CommonParameters.MeshPositionBuffer			== MeshData.Sections[SectionIt].PositionBuffer);
+//		check(CommonParameters.MeshPreviousPositionBuffer	== MeshData.Sections[SectionIt].PreviousPositionBuffer);
+#endif
 		check(CommonParameters.MeshIndexBuffer				== MeshData.Sections[SectionIt].IndexBuffer);
 		check(CommonParameters.MeshUVsBuffer				== MeshData.Sections[SectionIt].UVsBuffer);
 	}
diff --git a/Engine/Shaders/Private/HairStrands/HairStrandsMesh.usf b/Engine/Shaders/Private/HairStrands/HairStrandsMesh.usf
--- a/Engine/Shaders/Private/HairStrands/HairStrandsMesh.usf
+++ b/Engine/Shaders/Private/HairStrands/HairStrandsMesh.usf
@@ -222,6 +222,7 @@ struct FMeshSectionData
 	uint IndexOffset;
 	uint UVsChannelOffset;
 	uint UVsChannelCount;
+	uint bIsSwapped;
 };
 
 uint MaxSectionCount;
@@ -260,21 +261,39 @@ FMeshSectionData ReadMeshSectionData(uint InSectionIndex)
 	Out.IndexOffset = Packed.z;
 	Out.UVsChannelOffset = BitFieldExtractU32(Packed.w, 8, 0);
 	Out.UVsChannelCount = BitFieldExtractU32(Packed.w, 8, 8);
+	Out.bIsSwapped = BitFieldExtractU32(Packed.w, 1, 16);
 	return Out;
 }
 
 FHairMeshTriangleWithUVs GetTriangleData(FHairTriangleIndex TriangleIndex, FMeshSectionData MeshSectionData, bool bPrevious)
 {
-	return GetTriangleData(
-		TriangleIndex.TriangleIndex, 
-		bPrevious ? MeshPreviousPositionBuffer : MeshPositionBuffer, 
-		MeshUVsBuffer, 
-		MeshIndexBuffer, 
-		MeshSectionData.IndexOffset,
-		MeshSectionData.MaxIndexCount,
-		MeshSectionData.MaxVertexCount,
-		MeshSectionData.UVsChannelOffset, 
-		MeshSectionData.UVsChannelCount);
+	const bool bSamplePrev = (MeshSectionData.bIsSwapped > 0 && !bPrevious) || (MeshSectionData.bIsSwapped == 0 && bPrevious);
+	if (bSamplePrev)
+	{
+		return GetTriangleData(
+			TriangleIndex.TriangleIndex, 
+			MeshPreviousPositionBuffer, 
+			MeshUVsBuffer, 
+			MeshIndexBuffer, 
+			MeshSectionData.IndexOffset,
+			MeshSectionData.MaxIndexCount,
+			MeshSectionData.MaxVertexCount,
+			MeshSectionData.UVsChannelOffset, 
+			MeshSectionData.UVsChannelCount);
+	}
+	else
+	{
+		return GetTriangleData(
+			TriangleIndex.TriangleIndex, 
+			MeshPositionBuffer, 
+			MeshUVsBuffer, 
+			MeshIndexBuffer, 
+			MeshSectionData.IndexOffset,
+			MeshSectionData.MaxIndexCount,
+			MeshSectionData.MaxVertexCount,
+			MeshSectionData.UVsChannelOffset, 
+			MeshSectionData.UVsChannelCount);
+	}
 }
 
 [numthreads(GROUP_SIZE, 1, 1)]
