Fix CIS issue.

diff --git a/Engine/Plugins/Runtime/HairStrands/Source/HairStrandsCore/Private/HairStrandsMeshProjection.cpp b/Engine/Plugins/Runtime/HairStrands/Source/HairStrandsCore/Private/HairStrandsMeshProjection.cpp
--- a/Engine/Plugins/Runtime/HairStrands/Source/HairStrandsCore/Private/HairStrandsMeshProjection.cpp
+++ b/Engine/Plugins/Runtime/HairStrands/Source/HairStrandsCore/Private/HairStrandsMeshProjection.cpp
@@ -268,10 +268,10 @@ bool AddHairStrandUpdateMeshTrianglesPass(
 		check(CommonParameters.RDGMeshPositionBuffer == MeshData.Sections[SectionIt].RDGPositionBuffer || CommonParameters.RDGMeshPreviousPositionBuffer == MeshData.Sections[SectionIt].RDGPositionBuffer);
 		check(CommonParameters.MeshPositionBuffer    == MeshData.Sections[SectionIt].PositionBuffer    || CommonParameters.MeshPreviousPositionBuffer    == MeshData.Sections[SectionIt].PositionBuffer);
 #else
-//		check(CommonParameters.RDGMeshPositionBuffer		== MeshData.Sections[SectionIt].RDGPositionBuffer);
-//		check(CommonParameters.RDGMeshPreviousPositionBuffer== MeshData.Sections[SectionIt].RDGPreviousPositionBuffer);
-//		check(CommonParameters.MeshPositionBuffer			== MeshData.Sections[SectionIt].PositionBuffer);
-//		check(CommonParameters.MeshPreviousPositionBuffer	== MeshData.Sections[SectionIt].PreviousPositionBuffer);
+		check(CommonParameters.RDGMeshPositionBuffer		== MeshData.Sections[SectionIt].RDGPositionBuffer);
+		check(CommonParameters.RDGMeshPreviousPositionBuffer== MeshData.Sections[SectionIt].RDGPreviousPositionBuffer);
+		check(CommonParameters.MeshPositionBuffer			== MeshData.Sections[SectionIt].PositionBuffer);
+		check(CommonParameters.MeshPreviousPositionBuffer	== MeshData.Sections[SectionIt].PreviousPositionBuffer);
 #endif
 		check(CommonParameters.MeshIndexBuffer				== MeshData.Sections[SectionIt].IndexBuffer);
 		check(CommonParameters.MeshUVsBuffer				== MeshData.Sections[SectionIt].UVsBuffer);
diff --git a/Engine/Source/Runtime/Renderer/Private/ShadowRendering.cpp b/Engine/Source/Runtime/Renderer/Private/ShadowRendering.cpp
--- a/Engine/Source/Runtime/Renderer/Private/ShadowRendering.cpp
+++ b/Engine/Source/Runtime/Renderer/Private/ShadowRendering.cpp
@@ -2317,22 +2317,9 @@ void FDeferredShadingSceneRenderer::RenderDeferredShadowProjections(
 	{
 		bool bNeedHairShadowMaskPass = false;
 		const bool bVirtualShadowOnePass = VisibleLightInfo.VirtualShadowMapClipmaps.Num() > 0;
-		if (!bVirtualShadowOnePass)
+		if (!bVirtualShadowOnePass && VisibleLightInfo.ShadowsToProject.Num() > 0)
 		{
-			for (int32 ShadowIndex = 0; ShadowIndex < VisibleLightInfo.ShadowsToProject.Num(); ShadowIndex++)
-			{
-				FProjectedShadowInfo* ProjectedShadowInfo = VisibleLightInfo.ShadowsToProject[ShadowIndex];
-				if (ProjectedShadowInfo->HasVirtualShadowMap())
-				{
-					bNeedHairShadowMaskPass = false;
-					break;
-				}
-				else
-				{
-					bNeedHairShadowMaskPass = true;
-					break;
-				}
-			}
+			bNeedHairShadowMaskPass = !VisibleLightInfo.ShadowsToProject[0]->HasVirtualShadowMap();
 		}
 		if (bNeedHairShadowMaskPass)
 		{
