Fix validation error when hair is bound to a skel. mesh without previous position buffer.

diff --git a/Engine/Plugins/Runtime/HairStrands/Source/HairStrandsCore/Private/HairStrandsMeshProjection.cpp b/Engine/Plugins/Runtime/HairStrands/Source/HairStrandsCore/Private/HairStrandsMeshProjection.cpp
--- a/Engine/Plugins/Runtime/HairStrands/Source/HairStrandsCore/Private/HairStrandsMeshProjection.cpp
+++ b/Engine/Plugins/Runtime/HairStrands/Source/HairStrandsCore/Private/HairStrandsMeshProjection.cpp
@@ -265,8 +265,8 @@ bool AddHairStrandUpdateMeshTrianglesPass(
 		check(MeshSectionData.UVsChannelOffset < 255);
 		check(MeshSectionData.UVsChannelCount < 255);
 #if 1 // Relaxed check, with optional buffer swap
-		check(CommonParameters.RDGMeshPositionBuffer == MeshData.Sections[SectionIt].RDGPositionBuffer || CommonParameters.RDGMeshPreviousPositionBuffer == MeshData.Sections[SectionIt].RDGPositionBuffer);
-		check(CommonParameters.MeshPositionBuffer    == MeshData.Sections[SectionIt].PositionBuffer    || CommonParameters.MeshPreviousPositionBuffer    == MeshData.Sections[SectionIt].PositionBuffer);
+		check(CommonParameters.RDGMeshPositionBuffer == MeshData.Sections[SectionIt].RDGPositionBuffer || CommonParameters.RDGMeshPreviousPositionBuffer == MeshData.Sections[SectionIt].RDGPreviousPositionBuffer);
+		check(CommonParameters.MeshPositionBuffer    == MeshData.Sections[SectionIt].PositionBuffer    || CommonParameters.MeshPreviousPositionBuffer    == MeshData.Sections[SectionIt].PreviousPositionBuffer);
 #else
 		check(CommonParameters.RDGMeshPositionBuffer		== MeshData.Sections[SectionIt].RDGPositionBuffer);
 		check(CommonParameters.RDGMeshPreviousPositionBuffer== MeshData.Sections[SectionIt].RDGPreviousPositionBuffer);
@@ -276,6 +276,17 @@ bool AddHairStrandUpdateMeshTrianglesPass(
 		check(CommonParameters.MeshIndexBuffer				== MeshData.Sections[SectionIt].IndexBuffer);
 		check(CommonParameters.MeshUVsBuffer				== MeshData.Sections[SectionIt].UVsBuffer);
 	}
+
+	// If no previous position buffer available, reusing the current position buffers
+	if (CommonParameters.MeshPreviousPositionBuffer == nullptr)
+	{
+		CommonParameters.MeshPreviousPositionBuffer = CommonParameters.MeshPositionBuffer;
+	}
+	if (CommonParameters.RDGMeshPreviousPositionBuffer == nullptr)
+	{
+		CommonParameters.RDGMeshPreviousPositionBuffer = CommonParameters.RDGMeshPositionBuffer;
+	}
+
 	FRDGBufferRef SectionBuffer = CreateStructuredBuffer(GraphBuilder, TEXT("Hair.SkelMeshSectionBuffer"), sizeof(FSectionData),  SectionDatas.Num(), SectionDatas.GetData(), sizeof(FSectionData) * SectionDatas.Num());
 	CommonParameters.MeshSectionBuffer = GraphBuilder.CreateSRV(SectionBuffer);
 
