Fix hair deformation by bone.

diff --git a/Engine/Plugins/Runtime/HairStrands/Source/HairStrandsCore/Private/GroomManager.cpp b/Engine/Plugins/Runtime/HairStrands/Source/HairStrandsCore/Private/GroomManager.cpp
--- a/Engine/Plugins/Runtime/HairStrands/Source/HairStrandsCore/Private/GroomManager.cpp
+++ b/Engine/Plugins/Runtime/HairStrands/Source/HairStrandsCore/Private/GroomManager.cpp
@@ -613,7 +613,7 @@ static void RunHairStrandsInterpolation_Guide(
 	// Apply deformation from skeletal mesh bones
 	for (FInstanceData& InstanceData : InstanceDatas)
 	{
-		if (InstanceData.ActiveGroomCacheType != EGroomCacheType::None && InstanceData.Instance->Guides.bIsDeformationEnable && InstanceData.Instance->Guides.DeformedResource && InstanceData.Instance->DeformedComponent && (InstanceData.Instance->DeformedSection != INDEX_NONE))
+		if (InstanceData.ActiveGroomCacheType == EGroomCacheType::None && InstanceData.Instance->Guides.bIsDeformationEnable && InstanceData.Instance->Guides.DeformedResource && InstanceData.Instance->DeformedComponent && (InstanceData.Instance->DeformedSection != INDEX_NONE))
 		{
 			if (USkeletalMeshComponent* SkeletalMeshComponent = Cast<USkeletalMeshComponent>(InstanceData.Instance->DeformedComponent))
 			{
