Invalidate binding DDC key to force regenerate groom binding.

diff --git a/Engine/Plugins/Runtime/HairStrands/Source/HairStrandsCore/Private/GroomBindingBuilder.cpp b/Engine/Plugins/Runtime/HairStrands/Source/HairStrandsCore/Private/GroomBindingBuilder.cpp
--- a/Engine/Plugins/Runtime/HairStrands/Source/HairStrandsCore/Private/GroomBindingBuilder.cpp
+++ b/Engine/Plugins/Runtime/HairStrands/Source/HairStrandsCore/Private/GroomBindingBuilder.cpp
@@ -58,7 +58,7 @@ static FAutoConsoleVariableRef CVarHairStrandsBindingBuilderWarningEnable(TEXT("
 FString FGroomBindingBuilder::GetVersion()
 {
 	// Important to update the version when groom building changes
-	return TEXT("3i");
+	return TEXT("3p");
 }
 
 ///////////////////////////////////////////////////////////////////////////////////////////////////
diff --git a/Engine/Plugins/Runtime/HairStrands/Source/HairStrandsCore/Private/GroomDebug.cpp b/Engine/Plugins/Runtime/HairStrands/Source/HairStrandsCore/Private/GroomDebug.cpp
--- a/Engine/Plugins/Runtime/HairStrands/Source/HairStrandsCore/Private/GroomDebug.cpp
+++ b/Engine/Plugins/Runtime/HairStrands/Source/HairStrandsCore/Private/GroomDebug.cpp
@@ -43,7 +43,6 @@ static int32 GHairDebugMeshProjection_Render_HairDeformedSamples = 0;
 
 
 static FAutoConsoleVariableRef CVarHairDebugMeshProjection_SkinCacheMeshInUVsSpace(TEXT("r.HairStrands.MeshProjection.DebugInUVsSpace"), GHairDebugMeshProjection_SkinCacheMeshInUVsSpace, TEXT("Render debug mes projection in UVs space"));
-static FAutoConsoleVariableRef CVarHairDebugMeshProjection_SkinCacheMesh(TEXT("r.HairStrands.MeshProjection.DebugSkinCache"), GHairDebugMeshProjection_SkinCacheMesh, TEXT("Render the skin cache used in projection"));
 
 static FAutoConsoleVariableRef CVarHairDebugMeshProjection_Render_HairRestTriangles(TEXT("r.HairStrands.MeshProjection.Render.Rest.Triangles"), GHairDebugMeshProjection_Render_HairRestTriangles, TEXT("Render strands rest triangles"));
 static FAutoConsoleVariableRef CVarHairDebugMeshProjection_Render_HairRestFrames(TEXT("r.HairStrands.MeshProjection.Render.Rest.Frames"), GHairDebugMeshProjection_Render_HairRestFrames, TEXT("Render strands rest frames"));
diff --git a/Engine/Plugins/Runtime/HairStrands/Source/HairStrandsCore/Private/HairStrandsMeshProjection.cpp b/Engine/Plugins/Runtime/HairStrands/Source/HairStrandsCore/Private/HairStrandsMeshProjection.cpp
--- a/Engine/Plugins/Runtime/HairStrands/Source/HairStrandsCore/Private/HairStrandsMeshProjection.cpp
+++ b/Engine/Plugins/Runtime/HairStrands/Source/HairStrandsCore/Private/HairStrandsMeshProjection.cpp
@@ -259,6 +259,7 @@ bool AddHairStrandUpdateMeshTrianglesPass(
 		SectionDatas[SectionIt].UVsChannelOffset= MeshSectionData.UVsChannelOffset;
 		SectionDatas[SectionIt].UVsChannelCount	= MeshSectionData.UVsChannelCount;
 		SectionDatas[SectionIt].bIsSwapped		= MeshData.Sections[SectionIt].PositionBuffer != CommonParameters.MeshPositionBuffer ? 1u : 0u;
+		SectionDatas[SectionIt].Pad				= 0u;
 
 		// Sanity check
 		check(MeshSectionData.SectionIndex == SectionIt);
