Add support for additional texture layout for groom cards/meshes.

diff --git a/Engine/Plugins/Runtime/HairStrands/Source/HairStrandsCore/Private/GroomAsset.cpp b/Engine/Plugins/Runtime/HairStrands/Source/HairStrandsCore/Private/GroomAsset.cpp
--- a/Engine/Plugins/Runtime/HairStrands/Source/HairStrandsCore/Private/GroomAsset.cpp
+++ b/Engine/Plugins/Runtime/HairStrands/Source/HairStrandsCore/Private/GroomAsset.cpp
@@ -409,53 +409,48 @@ void UpdateHairStrandsVerbosity(IConsoleVariable* InCVarVerbosity)
 /////////////////////////////////////////////////////////////////////////////////////////
 
 template<typename ResourceType>
-static void InitAtlasTexture(ResourceType* InResource, UTexture2D* InTexture, EHairAtlasTextureType InType)
+static void InitAtlasTexture(ResourceType* InResource, const FHairGroupCardsTextures& In)
 {
-	if (InTexture == nullptr || InResource == nullptr)
+	if (InResource == nullptr)
 	{
 		return;
 	}
 
-	InTexture->ConditionalPostLoad();
+	const TArray<UTexture2D*>& InTextures = In.Textures;
+	const uint32 InLayoutIndex = uint32(In.Layout);
 
-	ENQUEUE_RENDER_COMMAND(HairStrandsCardsTextureCommand)(
-	[InResource, InTexture, InType](FRHICommandListImmediate& RHICmdList)
+	bool bHasValidTextures = false;
+	for (UTexture2D* T : InTextures)
 	{
-		FSamplerStateRHIRef DefaultSampler = TStaticSamplerState<SF_AnisotropicLinear, AM_Clamp, AM_Clamp, AM_Clamp, 0, 8/*MaxAnisotropy*/>::GetRHI();
-		switch (InType)
+		if (T)
 		{
-		case EHairAtlasTextureType::Depth:
-		{
-			InResource->DepthTexture = InTexture->TextureReference.TextureReferenceRHI;
-			InResource->DepthSampler = DefaultSampler;
-		} break;
-		case EHairAtlasTextureType::Tangent:
-		{
-			InResource->TangentTexture = InTexture->TextureReference.TextureReferenceRHI;
-			InResource->TangentSampler = DefaultSampler;
-		} break;
-		case EHairAtlasTextureType::Attribute:
-		{
-			InResource->AttributeTexture = InTexture->TextureReference.TextureReferenceRHI;
-			InResource->AttributeSampler = DefaultSampler;
-		} break;
-		case EHairAtlasTextureType::Coverage:
-		{
-			InResource->CoverageTexture = InTexture->TextureReference.TextureReferenceRHI;
-			InResource->CoverageSampler = DefaultSampler;
-		} break;
-		case EHairAtlasTextureType::AuxilaryData:
-		{
-			InResource->AuxilaryDataTexture = InTexture->TextureReference.TextureReferenceRHI;
-			InResource->AuxilaryDataSampler = DefaultSampler;
-		} break;
-		case EHairAtlasTextureType::Material:
-		{
-			InResource->MaterialTexture = InTexture->TextureReference.TextureReferenceRHI;
-			InResource->MaterialSampler = DefaultSampler;
-		} break;
+			T->ConditionalPostLoad();
+			bHasValidTextures = true;
 		}
-	});
+	}
+
+	if (bHasValidTextures)
+	{
+		ENQUEUE_RENDER_COMMAND(HairStrandsCardsTextureCommand)(
+			[InResource, InTextures, InLayoutIndex](FRHICommandListImmediate& RHICmdList)
+		{
+			FSamplerStateRHIRef DefaultSampler = TStaticSamplerState<SF_AnisotropicLinear, AM_Clamp, AM_Clamp, AM_Clamp, 0, 8/*MaxAnisotropy*/>::GetRHI();
+
+			const uint32 TextureCount = InTextures.Num();
+			InResource->Textures.SetNum(TextureCount);
+			InResource->Samplers.SetNum(TextureCount);
+			InResource->LayoutIndex = InLayoutIndex;
+
+			for (uint32 TextureIt=0;TextureIt<TextureCount;++TextureIt)
+			{
+				if (UTexture2D* Texture = InTextures[TextureIt])
+				{
+					InResource->Textures[TextureIt] = Texture->TextureReference.TextureReferenceRHI;
+				}
+				InResource->Samplers[TextureIt] = DefaultSampler;
+			}
+		});
+	}
 }
 
 template<typename T>
@@ -1296,9 +1291,10 @@ void UGroomAsset::PostLoad()
 		}
 	}
 
-	// Convert old procedural cards to import cards
+	// Convert cards
 	for (FHairGroupsCardsSourceDescription& Group : GetHairGroupsCards())
 	{
+		// Convert old procedural cards to import cards
 		if (Group.SourceType_DEPRECATED == EHairCardsSourceType::Procedural)
 		{
 			Group.bInvertUV = true;
@@ -1307,6 +1303,36 @@ void UGroomAsset::PostLoad()
 			Group.ProceduralMesh_DEPRECATED = nullptr;
 			Group.ImportedMeshKey = FString();
 		}
+
+		// Convert old textures
+		if (Group.Textures.Textures.Num() == 0)
+		{
+			Group.Textures.Layout = EHairTextureLayout::Layout0;
+			Group.Textures.Textures.Reserve(6);
+			Group.Textures.Textures.Add(Group.Textures.DepthTexture_DEPRECATED);
+			Group.Textures.Textures.Add(Group.Textures.CoverageTexture_DEPRECATED);
+			Group.Textures.Textures.Add(Group.Textures.TangentTexture_DEPRECATED);
+			Group.Textures.Textures.Add(Group.Textures.AttributeTexture_DEPRECATED);
+			Group.Textures.Textures.Add(Group.Textures.MaterialTexture_DEPRECATED);
+			Group.Textures.Textures.Add(Group.Textures.AuxilaryDataTexture_DEPRECATED);
+		}
+	}
+
+	// Convert meshes
+	for (FHairGroupsMeshesSourceDescription& Group : GetHairGroupsMeshes())
+	{
+		// Convert old textures
+		if (Group.Textures.Textures.Num() == 0)
+		{
+			Group.Textures.Layout = EHairTextureLayout::Layout1;
+			Group.Textures.Textures.Reserve(6);
+			Group.Textures.Textures.Add(Group.Textures.DepthTexture_DEPRECATED);
+			Group.Textures.Textures.Add(Group.Textures.CoverageTexture_DEPRECATED);
+			Group.Textures.Textures.Add(Group.Textures.TangentTexture_DEPRECATED);
+			Group.Textures.Textures.Add(Group.Textures.AttributeTexture_DEPRECATED);
+			Group.Textures.Textures.Add(Group.Textures.MaterialTexture_DEPRECATED);
+			Group.Textures.Textures.Add(Group.Textures.AuxilaryDataTexture_DEPRECATED);
+		}
 	}
 
 #if WITH_EDITORONLY_DATA
@@ -1559,12 +1585,8 @@ void UGroomAsset::BeginDestroy()
 static bool IsCardsTextureResources(const FName PropertyName)
 {
 	return
-		   PropertyName == GET_MEMBER_NAME_CHECKED(FHairGroupCardsTextures, DepthTexture)
-		|| PropertyName == GET_MEMBER_NAME_CHECKED(FHairGroupCardsTextures, CoverageTexture)
-		|| PropertyName == GET_MEMBER_NAME_CHECKED(FHairGroupCardsTextures, TangentTexture)
-		|| PropertyName == GET_MEMBER_NAME_CHECKED(FHairGroupCardsTextures, AttributeTexture)
-		|| PropertyName == GET_MEMBER_NAME_CHECKED(FHairGroupCardsTextures, AuxilaryDataTexture)
-		|| PropertyName == GET_MEMBER_NAME_CHECKED(FHairGroupCardsTextures, MaterialTexture);
+		   PropertyName == GET_MEMBER_NAME_CHECKED(FHairGroupCardsTextures, Layout)
+		|| PropertyName == GET_MEMBER_NAME_CHECKED(FHairGroupCardsTextures, Textures);
 }
 static void InitCardsTextureResources(UGroomAsset* GroomAsset);
 
@@ -1679,6 +1701,19 @@ void UGroomAsset::PostEditChangeProperty(FPropertyChangedEvent& PropertyChangedE
 		CacheDerivedDatas();
 	}
 
+	// Update cards/meshes texture array according to the layout prior to reload the UI
+	if (PropertyName == GET_MEMBER_NAME_CHECKED(FHairGroupCardsTextures, Layout))
+	{
+		for (auto& Group : GetHairGroupsCards())
+		{
+			Group.Textures.Textures.SetNum(GetHairTextureLayoutTextureCount(Group.Textures.Layout));
+		}
+		for (auto& Group : GetHairGroupsMeshes())
+		{
+			Group.Textures.Textures.SetNum(GetHairTextureLayoutTextureCount(Group.Textures.Layout));
+		}
+	}
+
 	const bool bHairStrandsRaytracingRadiusChanged = PropertyName == GET_MEMBER_NAME_CHECKED(FHairShadowSettings, HairRaytracingRadiusScale);
 	
 	// By pass update if bStrandsInterpolationChanged has the resources have already been recreated
@@ -2969,13 +3004,13 @@ bool UGroomAsset::BuildCardsData(uint32 GroupIndex)
 				// content is up to date with what has been saved.
 				if (IsInGameThread())
 				{
-					FHairGroupCardsTextures& Textures = GetHairGroupsCards()[SourceIt].Textures;
-					if (Textures.DepthTexture != nullptr)		Textures.DepthTexture->UpdateResource();
-					if (Textures.TangentTexture != nullptr)		Textures.TangentTexture->UpdateResource();
-					if (Textures.AttributeTexture != nullptr)	Textures.AttributeTexture->UpdateResource();
-					if (Textures.CoverageTexture != nullptr)	Textures.CoverageTexture->UpdateResource();
-					if (Textures.AuxilaryDataTexture != nullptr)Textures.AuxilaryDataTexture->UpdateResource();
-					if (Textures.MaterialTexture != nullptr)	Textures.MaterialTexture->UpdateResource();
+					for (UTexture2D* Tex : GetHairGroupsCards()[SourceIt].Textures.Textures)
+					{
+						if (Tex != nullptr)
+						{
+							Tex->UpdateResource();
+						}
+					}
 				}
 			}
 		}
@@ -3051,12 +3086,7 @@ bool UGroomAsset::BuildCardsData(uint32 GroupIndex)
 				BeginInitResource(LOD.RestResource); // Immediate allocation, as needed for the vertex factory, input stream building
 
 				// 2.1 Load atlas textures
-				InitAtlasTexture(LOD.RestResource, Desc->Textures.DepthTexture, EHairAtlasTextureType::Depth);
-				InitAtlasTexture(LOD.RestResource, Desc->Textures.TangentTexture, EHairAtlasTextureType::Tangent);
-				InitAtlasTexture(LOD.RestResource, Desc->Textures.AttributeTexture, EHairAtlasTextureType::Attribute);
-				InitAtlasTexture(LOD.RestResource, Desc->Textures.CoverageTexture, EHairAtlasTextureType::Coverage);
-				InitAtlasTexture(LOD.RestResource, Desc->Textures.AuxilaryDataTexture, EHairAtlasTextureType::AuxilaryData);
-				InitAtlasTexture(LOD.RestResource, Desc->Textures.MaterialTexture, EHairAtlasTextureType::Material);
+				InitAtlasTexture(LOD.RestResource, Desc->Textures);
 				LOD.RestResource->bInvertUV = Desc->bInvertUV;
 				
 				// 2.2 Load interoplatino resources
@@ -3212,13 +3242,13 @@ bool UGroomAsset::BuildMeshesData(uint32 GroupIndex)
 				// content is up to date with what has been saved.
 				if (IsInGameThread())
 				{
-					FHairGroupCardsTextures& Textures = GetHairGroupsMeshes()[SourceIt].Textures;
-					if (Textures.DepthTexture != nullptr)		Textures.DepthTexture->UpdateResource();
-					if (Textures.TangentTexture != nullptr)		Textures.TangentTexture->UpdateResource();
-					if (Textures.AttributeTexture != nullptr)	Textures.AttributeTexture->UpdateResource();
-					if (Textures.CoverageTexture != nullptr)	Textures.CoverageTexture->UpdateResource();
-					if (Textures.AuxilaryDataTexture != nullptr)Textures.AuxilaryDataTexture->UpdateResource();
-					if (Textures.MaterialTexture != nullptr)	Textures.MaterialTexture->UpdateResource();
+					for (UTexture2D* Texture : GetHairGroupsMeshes()[SourceIt].Textures.Textures)
+					{
+						if (Texture != nullptr)
+						{
+							Texture->UpdateResource();
+						}
+					}
 				}
 			}
 		}
@@ -3458,12 +3488,7 @@ static void InitCardsTextureResources(UGroomAsset* GroomAsset)
 			{
 				if (Desc)
 				{
-					InitAtlasTexture(LOD.RestResource, Desc->Textures.DepthTexture, EHairAtlasTextureType::Depth);
-					InitAtlasTexture(LOD.RestResource, Desc->Textures.TangentTexture, EHairAtlasTextureType::Tangent);
-					InitAtlasTexture(LOD.RestResource, Desc->Textures.AttributeTexture, EHairAtlasTextureType::Attribute);
-					InitAtlasTexture(LOD.RestResource, Desc->Textures.CoverageTexture, EHairAtlasTextureType::Coverage);
-					InitAtlasTexture(LOD.RestResource, Desc->Textures.AuxilaryDataTexture, EHairAtlasTextureType::AuxilaryData);
-					InitAtlasTexture(LOD.RestResource, Desc->Textures.MaterialTexture, EHairAtlasTextureType::Material);
+					InitAtlasTexture(LOD.RestResource, Desc->Textures);
 					if (LOD.RestResource)
 					{
 						LOD.RestResource->bInvertUV = Desc->bInvertUV; // Should fix procedural texture so that this does not happen
@@ -3518,12 +3543,7 @@ void UGroomAsset::InitCardsResources()
 
 				if (Desc)
 				{
-					InitAtlasTexture(LOD.RestResource, Desc->Textures.DepthTexture, EHairAtlasTextureType::Depth);
-					InitAtlasTexture(LOD.RestResource, Desc->Textures.TangentTexture, EHairAtlasTextureType::Tangent);
-					InitAtlasTexture(LOD.RestResource, Desc->Textures.AttributeTexture, EHairAtlasTextureType::Attribute);
-					InitAtlasTexture(LOD.RestResource, Desc->Textures.CoverageTexture, EHairAtlasTextureType::Coverage);
-					InitAtlasTexture(LOD.RestResource, Desc->Textures.AuxilaryDataTexture, EHairAtlasTextureType::AuxilaryData);
-					InitAtlasTexture(LOD.RestResource, Desc->Textures.MaterialTexture, EHairAtlasTextureType::Material);
+					InitAtlasTexture(LOD.RestResource, Desc->Textures);
 					LOD.RestResource->bInvertUV = Desc->bInvertUV; // Should fix procedural texture so that this does not happen
 				}
 			}
@@ -3568,15 +3588,9 @@ void UGroomAsset::InitMeshesResources()
 				LOD.RestResource = new FHairMeshesRestResource(LOD.BulkData, FHairResourceName(GetFName(), GroupIndex, LODIt), GetAssetPathName(LODIt));
 				BeginInitResource(LOD.RestResource); // Immediate allocation, as needed for the vertex factory, input stream building
 
-
 				if (Desc)
 				{
-					InitAtlasTexture(LOD.RestResource, Desc->Textures.DepthTexture, EHairAtlasTextureType::Depth);
-					InitAtlasTexture(LOD.RestResource, Desc->Textures.TangentTexture, EHairAtlasTextureType::Tangent);
-					InitAtlasTexture(LOD.RestResource, Desc->Textures.AttributeTexture, EHairAtlasTextureType::Attribute);
-					InitAtlasTexture(LOD.RestResource, Desc->Textures.CoverageTexture, EHairAtlasTextureType::Coverage);
-					InitAtlasTexture(LOD.RestResource, Desc->Textures.AuxilaryDataTexture, EHairAtlasTextureType::AuxilaryData);
-					InitAtlasTexture(LOD.RestResource, Desc->Textures.MaterialTexture, EHairAtlasTextureType::Material);
+					InitAtlasTexture(LOD.RestResource, Desc->Textures);
 				}
 			}
 		}
diff --git a/Engine/Plugins/Runtime/HairStrands/Source/HairStrandsCore/Private/GroomAssetCards.cpp b/Engine/Plugins/Runtime/HairStrands/Source/HairStrandsCore/Private/GroomAssetCards.cpp
--- a/Engine/Plugins/Runtime/HairStrands/Source/HairStrandsCore/Private/GroomAssetCards.cpp
+++ b/Engine/Plugins/Runtime/HairStrands/Source/HairStrandsCore/Private/GroomAssetCards.cpp
@@ -18,27 +18,16 @@ FHairGroupsCardsSourceDescription::FHairGroupsCardsSourceDescription()
 
 void FHairGroupCardsTextures::SetTexture(EHairAtlasTextureType SlotID, UTexture2D* Texture)
 {
+	Layout = EHairTextureLayout::Layout0; // Default hair card layout for now
+	Textures.SetNum(6);
 	switch (SlotID)
 	{
-	case EHairAtlasTextureType::Depth:
-		DepthTexture = Texture;
-		break;
-
-	case EHairAtlasTextureType::Coverage:
-		CoverageTexture = Texture;
-		break;
-
-	case EHairAtlasTextureType::Tangent:
-		TangentTexture = Texture;
-		break;
-
-	case EHairAtlasTextureType::Attribute:
-		AttributeTexture = Texture;
-		break;
-
-	case EHairAtlasTextureType::AuxilaryData:
-		AuxilaryDataTexture = Texture;
-		break;
+	case EHairAtlasTextureType::Depth: 			Textures[0] = Texture; break;
+	case EHairAtlasTextureType::Coverage:		Textures[1] = Texture; break;
+	case EHairAtlasTextureType::Tangent:		Textures[2] = Texture; break;
+	case EHairAtlasTextureType::Attribute:		Textures[3] = Texture; break;
+	case EHairAtlasTextureType::AuxilaryData:	Textures[4] = Texture; break;
+	case EHairAtlasTextureType::Material:		Textures[5] = Texture; break;
 	};
 }
 
@@ -97,4 +86,14 @@ UStaticMesh* FHairGroupsCardsSourceDescription::GetMesh() const
 #endif
 }
 
-
+uint32 GetHairTextureLayoutTextureCount(EHairTextureLayout In)
+{
+	switch (In)
+	{
+		case EHairTextureLayout::Layout0: return 6u;
+		case EHairTextureLayout::Layout1: return 6u;
+		case EHairTextureLayout::Layout2: return 4u;
+		case EHairTextureLayout::Layout3: return 4u;
+	}
+	return 0;
+}
diff --git a/Engine/Plugins/Runtime/HairStrands/Source/HairStrandsCore/Private/GroomDebug.cpp b/Engine/Plugins/Runtime/HairStrands/Source/HairStrandsCore/Private/GroomDebug.cpp
--- a/Engine/Plugins/Runtime/HairStrands/Source/HairStrandsCore/Private/GroomDebug.cpp
+++ b/Engine/Plugins/Runtime/HairStrands/Source/HairStrandsCore/Private/GroomDebug.cpp
@@ -551,18 +551,8 @@ static void AddDrawDebugCardsAtlasPass(
 		return;
 	}
 
-	FTextureReferenceRHIRef AtlasTexture = nullptr;
-
-	const int32 DebugMode = FMath::Clamp(GHairCardsAtlasDebug, 1, 6);
-	switch (DebugMode)
-	{
-	case 1: AtlasTexture = Instance->Cards.LODs[LODIndex].RestResource->DepthTexture; break;
-	case 2: AtlasTexture = Instance->Cards.LODs[LODIndex].RestResource->CoverageTexture; break;
-	case 3: AtlasTexture = Instance->Cards.LODs[LODIndex].RestResource->TangentTexture; break;
-	case 4:
-	case 5:
-	case 6: AtlasTexture = Instance->Cards.LODs[LODIndex].RestResource->AttributeTexture; break;
-	}
+	const int32 DebugMode = FMath::Clamp(GHairCardsAtlasDebug, 1, Instance->Cards.LODs[LODIndex].RestResource->Textures.Num()-1);
+	FTextureReferenceRHIRef AtlasTexture = Instance->Cards.LODs[LODIndex].RestResource->Textures.IsValidIndex(DebugMode) ? Instance->Cards.LODs[LODIndex].RestResource->Textures[DebugMode] : nullptr;
 
 	if (AtlasTexture != nullptr)
 	{
diff --git a/Engine/Plugins/Runtime/HairStrands/Source/HairStrandsCore/Private/GroomResources.cpp b/Engine/Plugins/Runtime/HairStrands/Source/HairStrandsCore/Private/GroomResources.cpp
--- a/Engine/Plugins/Runtime/HairStrands/Source/HairStrandsCore/Private/GroomResources.cpp
+++ b/Engine/Plugins/Runtime/HairStrands/Source/HairStrandsCore/Private/GroomResources.cpp
@@ -861,13 +861,6 @@ void FHairCardsRestResource::InternalAllocate()
 	CreateBuffer<FHairCardsNormalFormat>(RHICmdList, BulkData.Normals, NormalsBuffer, ToHairResourceDebugName(TEXT("Hair.CardsRest_NormalBuffer"), ResourceName), OwnerName);
 	CreateBuffer<FHairCardsUVFormat>(RHICmdList, BulkData.UVs, UVsBuffer, ToHairResourceDebugName(TEXT("Hair.CardsRest_UVBuffer"), ResourceName), OwnerName);
 	CreateBuffer<FHairCardsMaterialFormat>(RHICmdList, BulkData.Materials, MaterialsBuffer, ToHairResourceDebugName(TEXT("Hair.CardsRest_MaterialBuffer"), ResourceName), OwnerName);
-
-	FSamplerStateRHIRef DefaultSampler = TStaticSamplerState<SF_Bilinear, AM_Clamp, AM_Clamp, AM_Clamp>::GetRHI();
-	DepthSampler = DefaultSampler;
-	TangentSampler = DefaultSampler;
-	CoverageSampler = DefaultSampler;
-	AttributeSampler = DefaultSampler;
-	MaterialSampler = DefaultSampler;
 }
 
 void FHairCardsRestResource::InternalRelease()
diff --git a/Engine/Plugins/Runtime/HairStrands/Source/HairStrandsCore/Private/HairCardsVertexFactory.cpp b/Engine/Plugins/Runtime/HairStrands/Source/HairStrandsCore/Private/HairCardsVertexFactory.cpp
--- a/Engine/Plugins/Runtime/HairStrands/Source/HairStrandsCore/Private/HairCardsVertexFactory.cpp
+++ b/Engine/Plugins/Runtime/HairStrands/Source/HairStrandsCore/Private/HairCardsVertexFactory.cpp
@@ -27,18 +27,65 @@ template<typename T> inline void VFC_BindParam(FMeshDrawSingleShaderBindings& Sh
 
 enum class EHairCardsFactoryFlags : uint32
 {
-	InvertedUV        = 0x1,
-	TextureRootUV     = 0x2,
-	TextureGroupIndex = 0x4,
-	TextureBaseColor  = 0x8,
-	TextureRoughness  = 0x10,
-	TextureAttribute  = 0x20,
-	TextureCoverage   = 0x40,
-	TextureDepth      = 0x80,
-	TextureTangent    = 0x100,
-	TextureAuxilary   = 0x200
+	InvertedUV = 0x1,
 };
 
+static void SetTextures(
+	const TArray<FTextureReferenceRHIRef>& InTextures,
+	const TArray<FSamplerStateRHIRef>& InSamplers,
+	const uint32 InLayoutIndex,
+	FHairCardsVertexFactoryUniformShaderParameters& Out)
+{
+	Out.LayoutIndex = InLayoutIndex;
+	Out.TextureCount = InTextures.Num();
+	for (uint32 TextureIt = 0, TextureCount = InTextures.Num(); TextureIt < TextureCount; ++TextureIt)
+	{
+		if (InTextures[TextureIt] != nullptr)
+		{
+			Out.Flags |= 1u<<(TextureIt+1u);		
+			switch (TextureIt)
+			{
+				case 0:
+				{
+					Out.Texture0Texture = InTextures[TextureIt];
+					Out.Texture0Sampler = InSamplers[TextureIt];
+				} break;
+				case 1:
+				{
+					Out.Texture1Texture = InTextures[TextureIt];
+					Out.Texture1Sampler = InSamplers[TextureIt];
+				} break;
+				case 2:
+				{
+					Out.Texture2Texture = InTextures[TextureIt];
+					Out.Texture2Sampler = InSamplers[TextureIt];
+				} break;
+				case 3:
+				{
+					Out.Texture3Texture = InTextures[TextureIt];
+					Out.Texture3Sampler = InSamplers[TextureIt];
+				} break;
+				case 4:
+				{
+					Out.Texture4Texture = InTextures[TextureIt];
+					Out.Texture4Sampler = InSamplers[TextureIt];
+				} break;
+				case 5:
+				{
+					Out.Texture5Texture = InTextures[TextureIt];
+					Out.Texture5Sampler = InSamplers[TextureIt];
+				} break;
+				default:
+				{
+					// Only support up to 6 textures for now
+					check(false);
+				}
+
+			}
+		}
+	}
+}
+
 ////////////////////////////////////////////////////////////////////////////////////////////////////////////////////////
 // Cards based vertex factory
 IMPLEMENT_GLOBAL_SHADER_PARAMETER_STRUCT(FHairCardsVertexFactoryUniformShaderParameters, "HairCardsVF");
@@ -56,18 +103,7 @@ FHairCardsUniformBuffer CreateHairCardsVFUniformBuffer(
 		const FHairGroupInstance::FCards::FLOD& LOD = Instance->Cards.LODs[LODIndex];
 
 		// Cards atlas UV are inverted so fetching needs to be inverted on the y-axis
-		UniformParameters.Flags = 0;
-		UniformParameters.Flags |= LOD.RestResource->bInvertUV 						? uint32(EHairCardsFactoryFlags::InvertedUV) : 0u;
-		UniformParameters.Flags |= LOD.RestResource->CoverageTexture != nullptr 	? uint32(EHairCardsFactoryFlags::TextureCoverage) : 0;
-		UniformParameters.Flags |= LOD.RestResource->DepthTexture != nullptr 		? uint32(EHairCardsFactoryFlags::TextureDepth) : 0;
-		UniformParameters.Flags |= LOD.RestResource->TangentTexture != nullptr 		? uint32(EHairCardsFactoryFlags::TextureTangent) : 0;
-		UniformParameters.Flags |= LOD.RestResource->AuxilaryDataTexture != nullptr ? uint32(EHairCardsFactoryFlags::TextureAuxilary) : 0;
-		// These attributes are stored on vertex data rather than on textures
-		// EHairCardsFactoryFlags::TextureRootUV
-		// EHairCardsFactoryFlags::TextureGroupIndex
-		// EHairCardsFactoryFlags::TextureBaseColor
-		// EHairCardsFactoryFlags::TextureRoughness
-
+		UniformParameters.Flags = LOD.RestResource->bInvertUV ? uint32(EHairCardsFactoryFlags::InvertedUV) : 0u;
 		UniformParameters.AttributeTextureIndex = 0; // TODO pack attribute texture/channel description. Currently use hardcoded layout in shader attribute code.
 		UniformParameters.AttributeChannelIndex = 0; // TODO pack attribute texture/channel description. Currently use hardcoded layout in shader attribute code.
 		UniformParameters.MaxVertexCount = LOD.RestResource->GetVertexCount();
@@ -92,19 +128,7 @@ FHairCardsUniformBuffer CreateHairCardsVFUniformBuffer(
 		UniformParameters.UVsBuffer = LOD.RestResource->UVsBuffer.ShaderResourceViewRHI.GetReference();
 		UniformParameters.MaterialsBuffer = LOD.RestResource->MaterialsBuffer.ShaderResourceViewRHI.GetReference();
 
-
-		UniformParameters.Texture0Texture = LOD.RestResource->DepthTexture;
-		UniformParameters.Texture0Sampler = LOD.RestResource->DepthSampler;
-		UniformParameters.Texture1Texture = LOD.RestResource->TangentTexture;
-		UniformParameters.Texture1Sampler = LOD.RestResource->TangentSampler;
-		UniformParameters.Texture2Texture = LOD.RestResource->CoverageTexture;
-		UniformParameters.Texture2Sampler = LOD.RestResource->CoverageSampler;
-		UniformParameters.Texture3Texture = LOD.RestResource->AttributeTexture;
-		UniformParameters.Texture3Sampler = LOD.RestResource->AttributeSampler;
-		UniformParameters.Texture4Texture = nullptr; // Material properties & Group index are stored on vertices for cards
-		UniformParameters.Texture4Sampler = nullptr; // Material properties & Group index are stored on vertices for cards
-		UniformParameters.Texture5Texture = LOD.RestResource->AuxilaryDataTexture;
-		UniformParameters.Texture5Sampler = LOD.RestResource->AuxilaryDataSampler;
+		SetTextures(LOD.RestResource->Textures, LOD.RestResource->Samplers, LOD.RestResource->LayoutIndex, UniformParameters);
 	}
 	else if (GeometryType == EHairGeometryType::Meshes)
 	{
@@ -126,16 +150,6 @@ FHairCardsUniformBuffer CreateHairCardsVFUniformBuffer(
 
 		// Meshes UV are not inverted so no need to invert the y-axis
 		UniformParameters.Flags = 0;
-		UniformParameters.Flags |= uint32(EHairCardsFactoryFlags::TextureRootUV);
-		UniformParameters.Flags |= uint32(EHairCardsFactoryFlags::TextureGroupIndex);
-		UniformParameters.Flags |= uint32(EHairCardsFactoryFlags::TextureBaseColor);
-		UniformParameters.Flags |= uint32(EHairCardsFactoryFlags::TextureRoughness);
-		UniformParameters.Flags |= uint32(EHairCardsFactoryFlags::TextureAttribute);
-		UniformParameters.Flags |= LOD.RestResource->CoverageTexture != nullptr ? uint32(EHairCardsFactoryFlags::TextureCoverage) : 0;
-		UniformParameters.Flags |= LOD.RestResource->DepthTexture != nullptr ? uint32(EHairCardsFactoryFlags::TextureDepth) : 0;
-		UniformParameters.Flags |= LOD.RestResource->TangentTexture != nullptr ? uint32(EHairCardsFactoryFlags::TextureTangent) : 0;
-		UniformParameters.Flags |= LOD.RestResource->AuxilaryDataTexture != nullptr ? uint32(EHairCardsFactoryFlags::TextureAuxilary) : 0;
-
 		UniformParameters.AttributeTextureIndex = 0; // TODO pack attribute texture/channel description. Currently use hardcoded layout in shader attribute code.
 		UniformParameters.AttributeChannelIndex = 0; // TODO pack attribute texture/channel description. Currently use hardcoded layout in shader attribute code.
 		UniformParameters.MaxVertexCount = LOD.RestResource->GetVertexCount();
@@ -145,18 +159,7 @@ FHairCardsUniformBuffer CreateHairCardsVFUniformBuffer(
 		UniformParameters.UVsBuffer = LOD.RestResource->UVsBuffer.ShaderResourceViewRHI.GetReference();
 		UniformParameters.MaterialsBuffer = UniformParameters.NormalsBuffer; // Reuse normal buffer as a dummy input for material buffer, since it not used for meshes (material data is fetch through textures)
 
-		UniformParameters.Texture0Texture = LOD.RestResource->DepthTexture;
-		UniformParameters.Texture0Sampler = LOD.RestResource->DepthSampler;
-		UniformParameters.Texture1Texture = LOD.RestResource->TangentTexture;
-		UniformParameters.Texture1Sampler = LOD.RestResource->TangentSampler;
-		UniformParameters.Texture2Texture = LOD.RestResource->CoverageTexture;
-		UniformParameters.Texture2Sampler = LOD.RestResource->CoverageSampler;
-		UniformParameters.Texture3Texture = LOD.RestResource->AttributeTexture;
-		UniformParameters.Texture3Sampler = LOD.RestResource->AttributeSampler;
-		UniformParameters.Texture5Texture = LOD.RestResource->AuxilaryDataTexture;
-		UniformParameters.Texture4Texture = LOD.RestResource->MaterialTexture;
-		UniformParameters.Texture4Sampler = LOD.RestResource->MaterialSampler;
-		UniformParameters.Texture5Sampler = LOD.RestResource->AuxilaryDataSampler;
+		SetTextures(LOD.RestResource->Textures, LOD.RestResource->Samplers, LOD.RestResource->LayoutIndex, UniformParameters);
 	}
 
 	if (!bSupportsManualVertexFetch)
diff --git a/Engine/Plugins/Runtime/HairStrands/Source/HairStrandsCore/Public/GroomAssetCards.h b/Engine/Plugins/Runtime/HairStrands/Source/HairStrandsCore/Public/GroomAssetCards.h
--- a/Engine/Plugins/Runtime/HairStrands/Source/HairStrandsCore/Public/GroomAssetCards.h
+++ b/Engine/Plugins/Runtime/HairStrands/Source/HairStrandsCore/Public/GroomAssetCards.h
@@ -28,6 +28,18 @@ enum class EHairCardsGuideType : uint8
 	GuideBased UMETA(DisplayName = "Guide-Based"),
 };
 
+UENUM(BlueprintType)
+enum class EHairTextureLayout : uint8
+{
+	Layout0 UMETA(DisplayName = "Card Default"),
+	Layout1 UMETA(DisplayName = "Mesh Default"),
+	Layout2 UMETA(DisplayName = "Card Compact"),
+	Layout3 UMETA(DisplayName = "Mesh Compact"),
+};
+
+// Returns the number of textures used for a particular layout
+HAIRSTRANDSCORE_API uint32 GetHairTextureLayoutTextureCount(EHairTextureLayout In);
+
 USTRUCT(BlueprintType)
 struct HAIRSTRANDSCORE_API FHairGroupCardsInfo
 {
@@ -46,23 +58,28 @@ struct HAIRSTRANDSCORE_API FHairGroupCardsTextures
 	GENERATED_BODY()
 
 	UPROPERTY(EditAnywhere, Category = "CardsTextures")
-	TObjectPtr<UTexture2D> DepthTexture = nullptr;
+	EHairTextureLayout Layout = EHairTextureLayout::Layout0;
 
 	UPROPERTY(EditAnywhere, Category = "CardsTextures")
-	TObjectPtr<UTexture2D> CoverageTexture = nullptr;
+	TArray<TObjectPtr<UTexture2D>> Textures;
 
-	UPROPERTY(EditAnywhere, Category = "CardsTextures")
-	TObjectPtr<UTexture2D> TangentTexture = nullptr;
+	UPROPERTY()
+	TObjectPtr<UTexture2D> DepthTexture_DEPRECATED = nullptr;
 
-	UPROPERTY(EditAnywhere, Category = "CardsAttributes")
-	TObjectPtr<UTexture2D> AttributeTexture = nullptr;
+	UPROPERTY()
+	TObjectPtr<UTexture2D> CoverageTexture_DEPRECATED = nullptr;
+
+	UPROPERTY()
+	TObjectPtr<UTexture2D> TangentTexture_DEPRECATED = nullptr;
 
-	UPROPERTY(EditAnywhere, Category = "CardsAttributes")
-	TObjectPtr<UTexture2D> AuxilaryDataTexture = nullptr;
+	UPROPERTY()
+	TObjectPtr<UTexture2D> AttributeTexture_DEPRECATED = nullptr;
 
-	/* This texture is only used by hair meshes and can be generated through the hair-textures option on Groom Asset. */
-	UPROPERTY(EditAnywhere, Category = "CardsAttributes")
-	TObjectPtr<UTexture2D> MaterialTexture = nullptr;
+	UPROPERTY()
+	TObjectPtr<UTexture2D> AuxilaryDataTexture_DEPRECATED = nullptr;
+
+	UPROPERTY()
+	TObjectPtr<UTexture2D> MaterialTexture_DEPRECATED = nullptr;
 
 	void SetTexture(EHairAtlasTextureType SlotID, UTexture2D* Texture);
 
diff --git a/Engine/Plugins/Runtime/HairStrands/Source/HairStrandsCore/Public/GroomInstance.h b/Engine/Plugins/Runtime/HairStrands/Source/HairStrandsCore/Public/GroomInstance.h
--- a/Engine/Plugins/Runtime/HairStrands/Source/HairStrandsCore/Public/GroomInstance.h
+++ b/Engine/Plugins/Runtime/HairStrands/Source/HairStrandsCore/Public/GroomInstance.h
@@ -24,6 +24,8 @@ enum class EGroomViewMode : uint8;
 //  * 12 bits for cards count -> 4000 cards for a hair group
 BEGIN_GLOBAL_SHADER_PARAMETER_STRUCT(FHairCardsVertexFactoryUniformShaderParameters, HAIRSTRANDSCORE_API)
 	SHADER_PARAMETER(uint32, Flags)
+	SHADER_PARAMETER(uint32, LayoutIndex)
+	SHADER_PARAMETER(uint32, TextureCount)
 	SHADER_PARAMETER(uint32, AttributeTextureIndex)
 	SHADER_PARAMETER(uint32, AttributeChannelIndex)
 	SHADER_PARAMETER(uint32, MaxVertexCount)
diff --git a/Engine/Plugins/Runtime/HairStrands/Source/HairStrandsCore/Public/GroomResources.h b/Engine/Plugins/Runtime/HairStrands/Source/HairStrandsCore/Public/GroomResources.h
--- a/Engine/Plugins/Runtime/HairStrands/Source/HairStrandsCore/Public/GroomResources.h
+++ b/Engine/Plugins/Runtime/HairStrands/Source/HairStrandsCore/Public/GroomResources.h
@@ -647,19 +647,9 @@ struct FHairCardsRestResource : public FHairCommonResource
 	FHairCardsVertexBuffer UVsBuffer;
 	FHairCardsVertexBuffer MaterialsBuffer;
 
-	FSamplerStateRHIRef DepthSampler = nullptr;
-	FSamplerStateRHIRef TangentSampler = nullptr;
-	FSamplerStateRHIRef CoverageSampler = nullptr;
-	FSamplerStateRHIRef AttributeSampler = nullptr;
-	FSamplerStateRHIRef AuxilaryDataSampler = nullptr;
-	FSamplerStateRHIRef MaterialSampler = nullptr;
-
-	FTextureReferenceRHIRef	DepthTexture = nullptr;
-	FTextureReferenceRHIRef	CoverageTexture = nullptr;
-	FTextureReferenceRHIRef	TangentTexture = nullptr;
-	FTextureReferenceRHIRef	AttributeTexture = nullptr;
-	FTextureReferenceRHIRef	AuxilaryDataTexture = nullptr;
-	FTextureReferenceRHIRef	MaterialTexture = nullptr;
+	TArray<FSamplerStateRHIRef> Samplers;
+	TArray<FTextureReferenceRHIRef> Textures;
+	uint32 LayoutIndex = 0;
 
 	/* Reference to the hair strands render data */
 	const FHairCardsBulkData& BulkData;
@@ -773,19 +763,9 @@ struct FHairMeshesRestResource : public FHairCommonResource
 	FHairCardsVertexBuffer NormalsBuffer;
 	FHairCardsVertexBuffer UVsBuffer;
 
-	FSamplerStateRHIRef DepthSampler = nullptr;
-	FSamplerStateRHIRef TangentSampler = nullptr;
-	FSamplerStateRHIRef CoverageSampler = nullptr;
-	FSamplerStateRHIRef AttributeSampler = nullptr;
-	FSamplerStateRHIRef AuxilaryDataSampler = nullptr;
-	FSamplerStateRHIRef MaterialSampler = nullptr;
-
-	FTextureReferenceRHIRef	DepthTexture = nullptr;
-	FTextureReferenceRHIRef	CoverageTexture = nullptr;
-	FTextureReferenceRHIRef	TangentTexture = nullptr;
-	FTextureReferenceRHIRef	AttributeTexture = nullptr;
-	FTextureReferenceRHIRef	AuxilaryDataTexture = nullptr;
-	FTextureReferenceRHIRef	MaterialTexture = nullptr;
+	TArray<FSamplerStateRHIRef> Samplers;
+	TArray<FTextureReferenceRHIRef>	Textures;
+	uint32 LayoutIndex = 0;
 
 	/* Reference to the hair strands render data */
 	const FHairMeshesBulkData& BulkData;
diff --git a/Engine/Plugins/Runtime/HairStrands/Source/HairStrandsEditor/Private/GroomAssetDetails.cpp b/Engine/Plugins/Runtime/HairStrands/Source/HairStrandsEditor/Private/GroomAssetDetails.cpp
--- a/Engine/Plugins/Runtime/HairStrands/Source/HairStrandsEditor/Private/GroomAssetDetails.cpp
+++ b/Engine/Plugins/Runtime/HairStrands/Source/HairStrandsEditor/Private/GroomAssetDetails.cpp
@@ -893,12 +893,8 @@ bool FGroomRenderingDetails::CommonResetToDefault(TSharedPtr<IPropertyHandle> Ch
 	{
 		{
 			FHairGroupCardsTextures Default;
-			HAIR_RESET1(GetHairGroupsCards(), FHairGroupCardsTextures, Textures, DepthTexture);
-			HAIR_RESET1(GetHairGroupsCards(), FHairGroupCardsTextures, Textures, CoverageTexture);
-			HAIR_RESET1(GetHairGroupsCards(), FHairGroupCardsTextures, Textures, TangentTexture);
-			HAIR_RESET1(GetHairGroupsCards(), FHairGroupCardsTextures, Textures, AttributeTexture);
-			HAIR_RESET1(GetHairGroupsCards(), FHairGroupCardsTextures, Textures, MaterialTexture);
-			HAIR_RESET1(GetHairGroupsCards(), FHairGroupCardsTextures, Textures, AuxilaryDataTexture);
+			HAIR_RESET1(GetHairGroupsCards(), FHairGroupCardsTextures, Textures, Layout);
+			HAIR_RESET1(GetHairGroupsCards(), FHairGroupCardsTextures, Textures, Textures);
 		}
 	}
 
@@ -914,12 +910,8 @@ bool FGroomRenderingDetails::CommonResetToDefault(TSharedPtr<IPropertyHandle> Ch
 
 		{
 			FHairGroupCardsTextures Default;
-			HAIR_RESET1(GetHairGroupsMeshes(), FHairGroupCardsTextures, Textures, DepthTexture);
-			HAIR_RESET1(GetHairGroupsMeshes(), FHairGroupCardsTextures, Textures, CoverageTexture);
-			HAIR_RESET1(GetHairGroupsMeshes(), FHairGroupCardsTextures, Textures, TangentTexture);
-			HAIR_RESET1(GetHairGroupsMeshes(), FHairGroupCardsTextures, Textures, AttributeTexture);
-			HAIR_RESET1(GetHairGroupsMeshes(), FHairGroupCardsTextures, Textures, MaterialTexture);
-			HAIR_RESET1(GetHairGroupsMeshes(), FHairGroupCardsTextures, Textures, AuxilaryDataTexture);
+			HAIR_RESET1(GetHairGroupsMeshes(), FHairGroupCardsTextures, Textures, Layout);
+			HAIR_RESET1(GetHairGroupsMeshes(), FHairGroupCardsTextures, Textures, Textures);
 		}
 	}
 
@@ -1508,6 +1500,62 @@ void FGroomRenderingDetails::OnGenerateElementForHairGroup(TSharedRef<IPropertyH
 			{
 				// Not node display
 			}
+			else if (PropertyName == GET_MEMBER_NAME_CHECKED(FHairGroupsCardsSourceDescription, Textures))
+			{
+				EHairTextureLayout LayoutType = EHairTextureLayout::Layout0;
+				if (GroomAsset != nullptr && GroomAsset->GetHairGroupsCards().IsValidIndex(GroupIndex))
+				{
+					LayoutType = GroomAsset->GetHairGroupsCards()[GroupIndex].Textures.Layout;
+				}
+
+				IDetailGroup& TextureGroup = ChildrenBuilder.AddGroup(TEXT("HairCardsTextures"), LOCTEXT("HairCardsTextures", "Textures"));
+				{
+					// Layout type
+					TextureGroup.AddPropertyRow(ChildHandle->GetChildHandle(0).ToSharedRef());
+
+					// Textures
+					const uint32 TextureCount = GroomAsset->GetHairGroupsCards()[GroupIndex].Textures.Textures.Num();
+					TSharedPtr<IPropertyHandle> TextureArrayHandle = ChildHandle->GetChildHandle(1);
+					switch (LayoutType)
+					{
+						case EHairTextureLayout::Layout0:
+						{
+							check(TextureCount == GetHairTextureLayoutTextureCount(EHairTextureLayout::Layout0));
+							TextureGroup.AddPropertyRow(TextureArrayHandle->GetChildHandle(0).ToSharedRef()).DisplayName(FText(LOCTEXT("Layout0_0", "Depth\n R8")));
+							TextureGroup.AddPropertyRow(TextureArrayHandle->GetChildHandle(1).ToSharedRef()).DisplayName(FText(LOCTEXT("Layout0_1", "Coverage\n R8")));
+							TextureGroup.AddPropertyRow(TextureArrayHandle->GetChildHandle(2).ToSharedRef()).DisplayName(FText(LOCTEXT("Layout0_2", "Tangent\n RGB8")));
+							TextureGroup.AddPropertyRow(TextureArrayHandle->GetChildHandle(3).ToSharedRef()).DisplayName(FText(LOCTEXT("Layout0_3", "Attributes\n RootUV | CoordU | Seed\n RGB8")));
+							TextureGroup.AddPropertyRow(TextureArrayHandle->GetChildHandle(4).ToSharedRef()).DisplayName(FText(LOCTEXT("Layout0_4", "Material\n Color | Roughess\n RBGA8 ")));
+							TextureGroup.AddPropertyRow(TextureArrayHandle->GetChildHandle(5).ToSharedRef()).DisplayName(FText(LOCTEXT("Layout0_5", "Auxiliary\n RGBA8")));
+						} break;
+						case EHairTextureLayout::Layout1:
+						{
+							check(TextureCount == GetHairTextureLayoutTextureCount(EHairTextureLayout::Layout1));
+							TextureGroup.AddPropertyRow(TextureArrayHandle->GetChildHandle(0).ToSharedRef()).DisplayName(FText(LOCTEXT("Layout1_0", "Depth\n R8")));
+							TextureGroup.AddPropertyRow(TextureArrayHandle->GetChildHandle(1).ToSharedRef()).DisplayName(FText(LOCTEXT("Layout1_1", "Coverage\n R8")));
+							TextureGroup.AddPropertyRow(TextureArrayHandle->GetChildHandle(2).ToSharedRef()).DisplayName(FText(LOCTEXT("Layout1_2", "Tangent\n RGB8")));
+							TextureGroup.AddPropertyRow(TextureArrayHandle->GetChildHandle(3).ToSharedRef()).DisplayName(FText(LOCTEXT("Layout1_3", "Attributes\n RootUV | CoordU | Seed\n RGB8")));
+							TextureGroup.AddPropertyRow(TextureArrayHandle->GetChildHandle(4).ToSharedRef()).DisplayName(FText(LOCTEXT("Layout1_4", "Material\n Color | GroupID\n RBGA8")));
+							TextureGroup.AddPropertyRow(TextureArrayHandle->GetChildHandle(5).ToSharedRef()).DisplayName(FText(LOCTEXT("Layout1_5", "Auxiliary\n RGBA8")));
+						} break;
+						case EHairTextureLayout::Layout2:
+						{
+							check(TextureCount == GetHairTextureLayoutTextureCount(EHairTextureLayout::Layout2));
+							TextureGroup.AddPropertyRow(TextureArrayHandle->GetChildHandle(0).ToSharedRef()).DisplayName(FText(LOCTEXT("Layout2_0", "Tangent | CoordU\n RGBA8")));
+							TextureGroup.AddPropertyRow(TextureArrayHandle->GetChildHandle(2).ToSharedRef()).DisplayName(FText(LOCTEXT("Layout2_1", "Coverage | Depth | Seed\n RGB8")));
+							TextureGroup.AddPropertyRow(TextureArrayHandle->GetChildHandle(3).ToSharedRef()).DisplayName(FText(LOCTEXT("Layout2_2", "Color | Roughness\n RGBA8")));
+						} break;
+						case EHairTextureLayout::Layout3:
+						{
+							check(TextureCount == GetHairTextureLayoutTextureCount(EHairTextureLayout::Layout3));
+							TextureGroup.AddPropertyRow(TextureArrayHandle->GetChildHandle(0).ToSharedRef()).DisplayName(FText(LOCTEXT("Layout3_0", "Tangent\n RGB8")));
+							TextureGroup.AddPropertyRow(TextureArrayHandle->GetChildHandle(1).ToSharedRef()).DisplayName(FText(LOCTEXT("Layout3_1", "RooUV | CoordU | GroupID\n RGBA8")));
+							TextureGroup.AddPropertyRow(TextureArrayHandle->GetChildHandle(2).ToSharedRef()).DisplayName(FText(LOCTEXT("Layout3_2", "Coverage | Depth | Seed\n RGB8")));
+							TextureGroup.AddPropertyRow(TextureArrayHandle->GetChildHandle(3).ToSharedRef()).DisplayName(FText(LOCTEXT("Layout3_3", "Color | Roughness\n RGBA8")));
+						} break;
+					}
+				}
+			}
 			else
 			{
 				IDetailPropertyRow& PropertyRow = AddPropertyWithCustomReset(ChildHandle, ChildrenBuilder, GroupIndex, -1);
@@ -1636,7 +1684,7 @@ void FGroomRenderingDetails::OnGenerateElementForHairGroup(TSharedRef<IPropertyH
 				ChildHandle->GetNumChildren(SubChildrenCount);
 				for (uint32 SubChildIt = 0; SubChildIt < SubChildrenCount; ++SubChildIt)
 				{
-					TSharedPtr<IPropertyHandle> SubChildHandle = ChildHandle->GetChildHandle(SubChildIt);					
+					TSharedPtr<IPropertyHandle> SubChildHandle = ChildHandle->GetChildHandle(SubChildIt);
 					AddPropertySeparator(SubChildHandle->GetProperty()->GetFName(), ChildrenBuilder);
 					ExpandStruct(SubChildHandle, ChildrenBuilder, GroupIndex, -1, true);
 				}
@@ -1667,6 +1715,11 @@ void FGroomRenderingDetails::OnGroomLODTypeChanged()
 	GroomDetailLayout->ForceRefreshDetails();
 }
 
+void FGroomRenderingDetails::OnGroomTextureLayoutChanged()
+{
+	GroomDetailLayout->ForceRefreshDetails();
+}
+
 // Hair binding display
 void FGroomRenderingDetails::OnGenerateElementForBindingAsset(TSharedRef<IPropertyHandle> StructProperty, int32 BindingIndex, IDetailChildrenBuilder& ChildrenBuilder, IDetailLayoutBuilder* DetailLayout)
 {
diff --git a/Engine/Plugins/Runtime/HairStrands/Source/HairStrandsEditor/Private/GroomAssetDetails.h b/Engine/Plugins/Runtime/HairStrands/Source/HairStrandsEditor/Private/GroomAssetDetails.h
--- a/Engine/Plugins/Runtime/HairStrands/Source/HairStrandsEditor/Private/GroomAssetDetails.h
+++ b/Engine/Plugins/Runtime/HairStrands/Source/HairStrandsEditor/Private/GroomAssetDetails.h
@@ -88,6 +88,7 @@ private:
 	FReply OnGenerateCardDataUsingPlugin(int32 GroupIndex);
 	FReply OnSelectBinding(int32 BindingIndex, FProperty* Property);
 	void OnGroomLODTypeChanged();
+	void OnGroomTextureLayoutChanged();
 
 	void AddPropertySeparator(FName PropertyName, IDetailChildrenBuilder& ChildrenBuilder);
 	void ExpandStructForLOD(TSharedRef<IPropertyHandle>& PropertyHandle, IDetailChildrenBuilder& ChildrenBuilder, int32 GroupIndex, int32 LODIndex, bool bOverrideReset);
diff --git a/Engine/Plugins/Runtime/HairStrands/Source/HairStrandsEditor/Private/GroomCustomAssetEditorToolkit.cpp b/Engine/Plugins/Runtime/HairStrands/Source/HairStrandsEditor/Private/GroomCustomAssetEditorToolkit.cpp
--- a/Engine/Plugins/Runtime/HairStrands/Source/HairStrandsEditor/Private/GroomCustomAssetEditorToolkit.cpp
+++ b/Engine/Plugins/Runtime/HairStrands/Source/HairStrandsEditor/Private/GroomCustomAssetEditorToolkit.cpp
@@ -302,10 +302,17 @@ void FGroomCustomAssetEditorToolkit::InitPreviewComponents()
 
 void FGroomCustomAssetEditorToolkit::OnClose() 
 {
-	//FCoreUObjectDelegates::OnObjectPropertyChanged.Remove(PropertyListenDelegate);
-	if (GroomAsset.IsValid() && PropertyListenDelegate.IsValid())
+	// Remove all delegates
+	if (GroomAsset.IsValid() && PropertyListenDelegates.Num() > 0)
 	{
-		GroomAsset->GetOnGroomAssetResourcesChanged().Remove(PropertyListenDelegate);
+		for (FDelegateHandle Handle : PropertyListenDelegates)
+		{
+			if (Handle.IsValid())
+			{
+				GroomAsset->GetOnGroomAssetResourcesChanged().Remove(Handle);
+			}
+		}
+		PropertyListenDelegates.Empty();
 	}
 
 	PropertiesTab.Reset();
@@ -656,7 +663,9 @@ void FGroomCustomAssetEditorToolkit::InitCustomAssetEditor(const EToolkitMode::T
 			LocalToolKit->DetailView_MaterialProperties->ForceRefresh();
 			LocalToolKit->DetailView_BindingProperties->ForceRefresh();
 		};
-		PropertyListenDelegate = GroomAsset->GetOnGroomAssetResourcesChanged().AddLambda(InvalidateDetailViews);
+
+		PropertyListenDelegates.Add(GroomAsset->GetOnGroomAssetResourcesChanged().AddLambda(InvalidateDetailViews));
+		PropertyListenDelegates.Add(GroomAsset->GetOnGroomAssetChanged().AddLambda(InvalidateDetailViews));
 	}
 }
 
diff --git a/Engine/Plugins/Runtime/HairStrands/Source/HairStrandsEditor/Public/GroomCustomAssetEditorToolkit.h b/Engine/Plugins/Runtime/HairStrands/Source/HairStrandsEditor/Public/GroomCustomAssetEditorToolkit.h
--- a/Engine/Plugins/Runtime/HairStrands/Source/HairStrandsEditor/Public/GroomCustomAssetEditorToolkit.h
+++ b/Engine/Plugins/Runtime/HairStrands/Source/HairStrandsEditor/Public/GroomCustomAssetEditorToolkit.h
@@ -179,7 +179,7 @@ private:
 
 	bool bIsTabManagerInitialized = false;
 	int32 ActiveGroomBindingIndex = -1;
-	FDelegateHandle PropertyListenDelegate;
+	TArray<FDelegateHandle> PropertyListenDelegates;
 	TWeakObjectPtr<UGroomAsset> GroomAsset;
 	TWeakObjectPtr<UGroomBindingAsset> GroomBindingAsset;
 	TWeakObjectPtr<UGroomBindingAssetList> GroomBindingAssetList;
diff --git a/Engine/Shaders/Private/HairStrands/HairCardsAttributeCommon.ush b/Engine/Shaders/Private/HairStrands/HairCardsAttributeCommon.ush
--- a/Engine/Shaders/Private/HairStrands/HairCardsAttributeCommon.ush
+++ b/Engine/Shaders/Private/HairStrands/HairCardsAttributeCommon.ush
@@ -9,16 +9,13 @@
 ///////////////////////////////////////////////////////////////////////////////////////////////////
 
 // Vertex factory Flags
-#define EHairCardsVFlags_InvertedUV        0x1
-#define EHairCardsVFlags_TextureRootUV     0x2
-#define EHairCardsVFlags_TextureGroupIndex 0x4
-#define EHairCardsVFlags_TextureColor      0x8
-#define EHairCardsVFlags_TextureRoughness  0x10
-#define EHairCardsVFlags_TextureAttribute  0x20
-#define EHairCardsVFlags_TextureCoverage   0x40
-#define EHairCardsVFlags_TextureDepth      0x80
-#define EHairCardsVFlags_TextureTangent    0x100
-#define EHairCardsVFlags_TextureAuxilary   0x200
+#define EHairCardsVFlags_InvertedUV  0x1
+#define EHairCardsVFlags_Texture0    0x2
+#define EHairCardsVFlags_Texture1    0x4
+#define EHairCardsVFlags_Texture2    0x8
+#define EHairCardsVFlags_Texture3    0x10
+#define EHairCardsVFlags_Texture4    0x20
+#define EHairCardsVFlags_Texture5    0x40
 
 // Attribute index
 #define EHairCardsVFAttribute_Depth        0
@@ -35,44 +32,58 @@
 ///////////////////////////////////////////////////////////////////////////////////////////////////
 // Helper functions for fetching attributes
 
+#define LAYOUT_VALUES(L0, L1, L2, L3) \
+	(HairCardsVF.LayoutIndex == 0 ? L0 :  \
+	(HairCardsVF.LayoutIndex == 1 ? L1 :  \
+	(HairCardsVF.LayoutIndex == 2 ? L2 : L3)))
+
 #define HAIRCARDSVF_USE_LEGACY_LAYOUT 1
 #if HAIRCARDSVF_USE_LEGACY_LAYOUT
 uint GetHairCardTextureIndex(uint In) 
 { 
+	// See GroomAssetDetails.cpp, EHairTextureLayout::Layout0/1/2/3 for mapping details. (9 = not present/not available)
 	switch (In)
 	{
-		case EHairCardsVFAttribute_Depth      : return 0;
-		case EHairCardsVFAttribute_Seed       : return 3;
-		case EHairCardsVFAttribute_CoordU     : return 3;
-		case EHairCardsVFAttribute_GroupIndex : return 4; // Only used for mesh, not for cards
-		case EHairCardsVFAttribute_Roughness  : return 4;
-		case EHairCardsVFAttribute_RootUV     : return 3;
-		case EHairCardsVFAttribute_Color      : return 4;
-		case EHairCardsVFAttribute_Tangent    : return 1;
-		case EHairCardsVFAttribute_Coverage   : return 2;
-		case EHairCardsVFAttribute_Auxilary   : return 5; 
+		case EHairCardsVFAttribute_Depth      : return LAYOUT_VALUES(0, 0, 1, 2);
+		case EHairCardsVFAttribute_Seed       : return LAYOUT_VALUES(3, 3, 1, 2);
+		case EHairCardsVFAttribute_CoordU     : return LAYOUT_VALUES(3, 3, 0, 1);
+		case EHairCardsVFAttribute_GroupIndex : return LAYOUT_VALUES(9, 4, 9, 1); // Only used for mesh, not for cards
+		case EHairCardsVFAttribute_Roughness  : return LAYOUT_VALUES(9, 9, 9, 3);
+		case EHairCardsVFAttribute_RootUV     : return LAYOUT_VALUES(9, 3, 9, 1);
+		case EHairCardsVFAttribute_Color      : return LAYOUT_VALUES(9, 4, 9, 3);
+		case EHairCardsVFAttribute_Tangent    : return LAYOUT_VALUES(2, 2, 0, 0);
+		case EHairCardsVFAttribute_Coverage   : return LAYOUT_VALUES(1, 1, 1, 2);
+		case EHairCardsVFAttribute_Auxilary   : return LAYOUT_VALUES(5, 5, 9, 9); 
 	}
 	return 0;
 }
 
+bool HasTextureBound(uint In)
+{
+	// HairCardsVF.Flags stores InverttUV on the first bit, texture bound after
+	const uint TextureIndex = GetHairCardTextureIndex(In);
+	return (TextureIndex < HairCardsVF.TextureCount) && (HairCardsVF.Flags & (1u << (TextureIndex + 1))) != 0;
+}
+
 uint GetHairCardChannelIndex(uint In) 
 { 
+	// See GroomAssetDetails.cpp, EHairTextureLayout::Layout0/1/2/3 for mapping details
 	switch (In)
 	{
-		case EHairCardsVFAttribute_Depth      : return 0;
-		case EHairCardsVFAttribute_Seed       : return 3;
-		case EHairCardsVFAttribute_CoordU     : return 2;
-		case EHairCardsVFAttribute_GroupIndex : return 3;
-		case EHairCardsVFAttribute_Roughness  : return 3;
-		case EHairCardsVFAttribute_RootUV     : return 0;
-		case EHairCardsVFAttribute_Color      : return 0;
-		case EHairCardsVFAttribute_Tangent    : return 0;
-		case EHairCardsVFAttribute_Coverage   : return 0;
-		case EHairCardsVFAttribute_Auxilary   : return 0; 
+		case EHairCardsVFAttribute_Depth      : return LAYOUT_VALUES(0, 0, 1, 1);
+		case EHairCardsVFAttribute_Seed       : return LAYOUT_VALUES(3, 3, 2, 2);
+		case EHairCardsVFAttribute_CoordU     : return LAYOUT_VALUES(2, 2, 3, 1);
+		case EHairCardsVFAttribute_GroupIndex : return LAYOUT_VALUES(3, 3, 0, 2);
+		case EHairCardsVFAttribute_Roughness  : return LAYOUT_VALUES(3, 0, 3, 3);
+		case EHairCardsVFAttribute_RootUV     : return LAYOUT_VALUES(0, 0, 0, 0);
+		case EHairCardsVFAttribute_Color      : return LAYOUT_VALUES(0, 0, 0, 0);
+		case EHairCardsVFAttribute_Tangent    : return LAYOUT_VALUES(0, 0, 0, 0);
+		case EHairCardsVFAttribute_Coverage   : return LAYOUT_VALUES(0, 0, 0, 0);
+		case EHairCardsVFAttribute_Auxilary   : return LAYOUT_VALUES(0, 0, 0, 0); 
 	}
 	return 0;
 }
-#else
+#else	
 uint GetHairCardTextureIndex(uint In)	{ return ((HairCards.AttributeTextureIndex) >> (In*3u) & 0x7u); }
 uint GetHairCardChannelIndex(uint In)	{ return ((HairCards.AttributeChannelIndex) >> (In*2u) & 0x3u); }
 #endif
@@ -111,16 +122,16 @@ float3 FetchAttribute3(uint In, float2 InAtlasUV)               { const float4 V
 float4 FetchAttribute4(uint In, float2 InAtlasUV)               { const float4 Value = FetchAttributeTexture(GetHairCardTextureIndex(In), InAtlasUV);         return Value;                                                  }
 
 // Special attribute fetch functions
-float  GetDepthOffset (float2 InAtlasUV)                        { float  Out = 0;         if (HairCardsVF.Flags & EHairCardsVFlags_TextureDepth)     { Out = FetchAttribute1(EHairCardsVFAttribute_Depth,     InAtlasUV);                          } return Out; }
-float  GetSeed        (float2 InAtlasUV)                        { float  Out = 0;         if (HairCardsVF.Flags & EHairCardsVFlags_TextureDepth)     { Out = FetchAttribute1(EHairCardsVFAttribute_Seed,      InAtlasUV);                          } return Out; }
-float  GetCoordU      (float2 InAtlasUV)                        { float  Out = 0;         if (HairCardsVF.Flags & EHairCardsVFlags_TextureDepth)     { Out = FetchAttribute1(EHairCardsVFAttribute_CoordU,    InAtlasUV);                          } return Out; }
-float  GetGroupIndex  (float2 InAtlasUV, float  InDefault)      { float  Out = InDefault; if (HairCardsVF.Flags & EHairCardsVFlags_TextureGroupIndex){ Out = FetchAttribute1(EHairCardsVFAttribute_GroupIndex,InAtlasUV); Out *= 255.f;            } return Out; }
-float  GetRoughness   (float2 InAtlasUV, float  InDefault)      { float  Out = InDefault; if (HairCardsVF.Flags & EHairCardsVFlags_TextureRoughness) { Out = FetchAttribute1(EHairCardsVFAttribute_Roughness, InAtlasUV);                          } return Out; }
-float  GetCoverage    (float2 InAtlasUV, float  InDefault)      { float  Out = InDefault; if (HairCardsVF.Flags & EHairCardsVFlags_TextureCoverage)  { Out = FetchAttribute1(EHairCardsVFAttribute_Coverage,  InAtlasUV, HairCardsVF.CoverageBias);} return Out; }
-float2 GetRootUV      (float2 InAtlasUV, float2 InDefault)      { float2 Out = InDefault; if (HairCardsVF.Flags & EHairCardsVFlags_TextureRootUV)    { Out = FetchAttribute2(EHairCardsVFAttribute_RootUV,    InAtlasUV);                          } return Out; }
-float3 GetColor       (float2 InAtlasUV, float3 InDefault)      { float3 Out = InDefault; if (HairCardsVF.Flags & EHairCardsVFlags_TextureColor)     { Out = FetchAttribute3(EHairCardsVFAttribute_Color,     InAtlasUV); Out = Pow2(Out);         } return Out; } // Cheap sRGB -> Linear encoding
-float3 GetTangent     (float2 InAtlasUV, float3 InDefault)      { float3 Out = InDefault; if (HairCardsVF.Flags & EHairCardsVFlags_TextureTangent)   { Out = FetchAttribute3(EHairCardsVFAttribute_Tangent,   InAtlasUV); Out = Out*2-1.f;         } return Out; }
-float4 GetAuxilary    (float2 InAtlasUV, float4 InDefault)      { float4 Out = InDefault; if (HairCardsVF.Flags & EHairCardsVFlags_TextureAuxilary)  { Out = FetchAttribute4(EHairCardsVFAttribute_Auxilary,  InAtlasUV);                          } return Out; }
+float  GetDepthOffset (float2 InAtlasUV)                        { float  Out = 0;         if (HasTextureBound(EHairCardsVFAttribute_Depth))		 { Out = FetchAttribute1(EHairCardsVFAttribute_Depth,     InAtlasUV);                          } return Out; }
+float  GetSeed        (float2 InAtlasUV)                        { float  Out = 0;         if (HasTextureBound(EHairCardsVFAttribute_Seed))		 { Out = FetchAttribute1(EHairCardsVFAttribute_Seed,      InAtlasUV);                          } return Out; }
+float  GetCoordU      (float2 InAtlasUV)                        { float  Out = 0;         if (HasTextureBound(EHairCardsVFAttribute_CoordU))	 { Out = FetchAttribute1(EHairCardsVFAttribute_CoordU,    InAtlasUV);                          } return Out; }
+float  GetGroupIndex  (float2 InAtlasUV, float  InDefault)      { float  Out = InDefault; if (HasTextureBound(EHairCardsVFAttribute_GroupIndex)) { Out = FetchAttribute1(EHairCardsVFAttribute_GroupIndex,InAtlasUV); Out *= 255.f;            } return Out; }
+float  GetRoughness   (float2 InAtlasUV, float  InDefault)      { float  Out = InDefault; if (HasTextureBound(EHairCardsVFAttribute_Roughness))	 { Out = FetchAttribute1(EHairCardsVFAttribute_Roughness, InAtlasUV);                          } return Out; }
+float  GetCoverage    (float2 InAtlasUV, float  InDefault)      { float  Out = InDefault; if (HasTextureBound(EHairCardsVFAttribute_Coverage))	 { Out = FetchAttribute1(EHairCardsVFAttribute_Coverage,  InAtlasUV, HairCardsVF.CoverageBias);} return Out; }
+float2 GetRootUV      (float2 InAtlasUV, float2 InDefault)      { float2 Out = InDefault; if (HasTextureBound(EHairCardsVFAttribute_RootUV))	 { Out = FetchAttribute2(EHairCardsVFAttribute_RootUV,    InAtlasUV);                          } return Out; }
+float3 GetColor       (float2 InAtlasUV, float3 InDefault)      { float3 Out = InDefault; if (HasTextureBound(EHairCardsVFAttribute_Color))		 { Out = FetchAttribute3(EHairCardsVFAttribute_Color,     InAtlasUV); Out = Pow2(Out);         } return Out; } // Cheap sRGB -> Linear encoding
+float3 GetTangent     (float2 InAtlasUV, float3 InDefault)      { float3 Out = InDefault; if (HasTextureBound(EHairCardsVFAttribute_Tangent))	 { Out = FetchAttribute3(EHairCardsVFAttribute_Tangent,   InAtlasUV); Out = Out*2-1.f;         } return Out; }
+float4 GetAuxilary    (float2 InAtlasUV, float4 InDefault)      { float4 Out = InDefault; if (HasTextureBound(EHairCardsVFAttribute_Auxilary))	 { Out = FetchAttribute4(EHairCardsVFAttribute_Auxilary,  InAtlasUV);                          } return Out; }
 
 ///////////////////////////////////////////////////////////////////////////////////////////////////
 // MaterialTemplate interface
@@ -184,7 +195,7 @@ float3 GetHairStrandsColor(float2 InAtlasUV, float3 InColor)
 float GetHairStrandsRoughness(float2 InAtlasUV, float InCardRoughness)
 {
 	//return GetRoughness(InAtlasUV, InCardRoughness);
-	return (HairCardsVF.Flags & EHairCardsVFlags_TextureRoughness) ? 0 : InCardRoughness;
+	return (HairCardsVF.Flags & EHairCardsVFAttribute_Roughness) ? 0 : InCardRoughness;
 }
 
 float GetHairStrandsAO(float2 InAtlasUV, float InAO)
