Fix crash when adding new cards/meshes entry to a groom.

#prefligh 656466ab77403070f3f0e21a

diff --git a/Engine/Plugins/Runtime/HairStrands/Source/HairStrandsCore/Private/GroomAsset.cpp b/Engine/Plugins/Runtime/HairStrands/Source/HairStrandsCore/Private/GroomAsset.cpp
--- a/Engine/Plugins/Runtime/HairStrands/Source/HairStrandsCore/Private/GroomAsset.cpp
+++ b/Engine/Plugins/Runtime/HairStrands/Source/HairStrandsCore/Private/GroomAsset.cpp
@@ -1702,7 +1702,11 @@ void UGroomAsset::PostEditChangeProperty(FPropertyChangedEvent& PropertyChangedE
 	}
 
 	// Update cards/meshes texture array according to the layout prior to reload the UI
-	if (PropertyName == GET_MEMBER_NAME_CHECKED(FHairGroupCardsTextures, Layout))
+	const bool bCardsOrMeshesGroupChanged = 
+		PropertyName == GetHairGroupsCardsMemberName() ||
+		PropertyName == GetHairGroupsMeshesMemberName() ||
+		PropertyName == GET_MEMBER_NAME_CHECKED(FHairGroupCardsTextures, Layout);
+	if (bCardsOrMeshesGroupChanged)
 	{
 		for (auto& Group : GetHairGroupsCards())
 		{
