Fix cards guide generation based on strands guide.

* DDC key tracking was incorrect
* Unbuilt 'strands' guides when using generated guides.

diff --git a/Engine/Plugins/Runtime/HairStrands/Source/HairStrandsCore/Private/GroomAsset.cpp b/Engine/Plugins/Runtime/HairStrands/Source/HairStrandsCore/Private/GroomAsset.cpp
--- a/Engine/Plugins/Runtime/HairStrands/Source/HairStrandsCore/Private/GroomAsset.cpp
+++ b/Engine/Plugins/Runtime/HairStrands/Source/HairStrandsCore/Private/GroomAsset.cpp
@@ -2189,41 +2189,6 @@ namespace GroomDerivedDataCacheUtils
 		LODSettings.BuildDDCKey(Ar);
 	}
 
-	// Geneate DDC key for an given hair group and a given LOD
-	FString BuildCardsDerivedDataKeySuffix(const UGroomAsset* GroomAsset, FHairGroupsCardsSourceDescription& Desc)
-	{
-		if (Desc.GroupIndex < 0 || Desc.LODIndex < 0)
-		{
-			return FString();
-		}
-
-		if (Desc.GroupIndex >= GroomAsset->GetHairGroupsPlatformData().Num() || Desc.LODIndex >= GroomAsset->GetHairGroupsPlatformData()[Desc.GroupIndex].Cards.LODs.Num())
-		{
-			return FString();
-		}
-
-		// Serialize the FHairGroupsCardsSourceDescription into a temporary array
-		// The archive is flagged as persistent so that machines of different endianness produce identical binary results.
-		TArray<uint8> TempBytes;
-		TempBytes.Reserve(512);
-		FMemoryWriter Ar(TempBytes, true);
-
-		Ar << Desc.GroupIndex;
-		Ar << Desc.LODIndex;
-		if (Desc.ImportedMesh)
-		{
-			FString Key = Desc.GetMeshKey();
-			Ar << Key;
-		}
-		Ar << Desc.GuideType;
-
-		FSHAHash Hash;
-		FSHA1::HashBuffer(TempBytes.GetData(), TempBytes.Num(), Hash.Hash);
-
-		static FString CardPrefixString(TEXT("CARDS_V") + FHairCardsBuilder::GetVersion() + TEXT("_"));
-		return CardPrefixString + Hash.ToString();
-	}
-
 	// Geneate DDC key for an entire hair group (i.e. with all its LODs) for cards data
 	FString BuildCardsDerivedDataKeySuffix(uint32 GroupIndex, const TArray<FHairLODSettings>& LODs, TArray<FHairGroupsCardsSourceDescription>& SourceDescriptions)
 	{
@@ -2255,6 +2220,7 @@ namespace GroomDerivedDataCacheUtils
 				FString Key = Desc.GetMeshKey();
 				Ar << Key;
 			}
+			Ar << Desc.GuideType;
 			// Material is not included as it doesn't affect the data building
 		}
 
@@ -2905,7 +2871,13 @@ inline FString GetLODName(const UGroomAsset* Asset, uint32 LODIndex)
 	return Asset->GetOutermost()->GetName() + FString::Printf(TEXT("_LOD%d"), LODIndex);
 }
 
-static bool InternalImportCardGeometry(FHairGroupsCardsSourceDescription* InDesc, const FHairDescriptionGroups& InHairDescriptionGroups, FHairCardsBulkData& OutBulkData, FHairStrandsDatas& OutGuidesData, FHairCardsInterpolationBulkData& OutInterpolationBulkData)
+static bool InternalImportCardGeometry(
+	FHairGroupsCardsSourceDescription* InDesc, 
+	const FHairDescriptionGroups& InHairDescriptionGroups, 
+	const TArray<FHairGroupsInterpolation>& InHairInterpolationGroups,
+	FHairCardsBulkData& OutBulkData, 
+	FHairStrandsDatas& OutGuidesData, 
+	FHairCardsInterpolationBulkData& OutInterpolationBulkData)
 {
 	UStaticMesh* CardsMesh = InDesc->ImportedMesh;
 	check(CardsMesh);
@@ -2916,14 +2888,15 @@ static bool InternalImportCardGeometry(FHairGroupsCardsSourceDescription* InDesc
 	// * Voxelize hair strands data (all group), to transfer group index from strands to cards
 	//FHairGroupInfo DummyInfo;
 
-	// Guides
 	FHairStrandsDatas TempHairGuidesData = InHairDescriptionGroups.HairGroups[InDesc->GroupIndex].Guides;
-	FGroomBuilder::BuildData(TempHairGuidesData);
+	FHairStrandsDatas TempHairStrandsData = InHairDescriptionGroups.HairGroups[InDesc->GroupIndex].Strands;
+	FGroomBuilder::BuildData(
+		InHairDescriptionGroups.HairGroups[InDesc->GroupIndex], 
+		InHairInterpolationGroups[InDesc->GroupIndex],
+		TempHairStrandsData,
+		TempHairGuidesData);
 
-	// Strands
 	FHairStrandsVoxelData TempHairStrandsVoxelData;
-	FHairStrandsDatas TempHairStrandsData = InHairDescriptionGroups.HairGroups[InDesc->GroupIndex].Strands;
-	FGroomBuilder::BuildData(TempHairStrandsData);
 	FGroomBuilder::VoxelizeGroupIndex(InHairDescriptionGroups, TempHairStrandsVoxelData);
 
 	// Transfer RootUV from strands to guides
@@ -3066,7 +3039,8 @@ bool UGroomAsset::BuildCardsData(uint32 GroupIndex)
 				CardsMesh->ConditionalPostLoad();
 
 				const FHairDescriptionGroups& LocalHairDescriptionGroups = GetHairDescriptionGroups();
-				bInitResources = InternalImportCardGeometry(Desc, LocalHairDescriptionGroups, LOD.BulkData, LODGuidesData, LOD.InterpolationBulkData);
+				const TArray<FHairGroupsInterpolation>& LocalInterpolationGroups= GetHairGroupsInterpolation();
+				bInitResources = InternalImportCardGeometry(Desc, LocalHairDescriptionGroups, LocalInterpolationGroups, LOD.BulkData, LODGuidesData, LOD.InterpolationBulkData);
 				if (!bInitResources)
 				{
 					UE_LOG(LogHairStrands, Warning, TEXT("Failed to import cards from %s for Group %d LOD %d."), *CardsMesh->GetName(), GroupIndex, LODIt);
@@ -3121,11 +3095,10 @@ bool UGroomAsset::BuildCardsData(uint32 GroupIndex)
 
 				// 2. (Re)Build data for the sim guides (since there are transient/no-cached)
 				const FHairDescriptionGroups& LocalHairDescriptionGroups = GetHairDescriptionGroups();
-				FHairGroupInfo DummyInfo;
 				check(LocalHairDescriptionGroups.IsValid());
 				FHairStrandsDatas SimGuidesData;
 				FHairStrandsDatas RenStrandsData; // Dummy, not used
-				FGroomBuilder::BuildData(LocalHairDescriptionGroups.HairGroups[GroupIndex], GetHairGroupsInterpolation()[GroupIndex], DummyInfo, RenStrandsData, SimGuidesData);
+				FGroomBuilder::BuildData(LocalHairDescriptionGroups.HairGroups[GroupIndex], GetHairGroupsInterpolation()[GroupIndex], RenStrandsData, SimGuidesData);
 
 				// 3. Compute interpolation data & bulk data between the sim guides and the cards guides
 				FHairStrandsInterpolationDatas LODGuideInterpolationData;
@@ -3845,10 +3818,9 @@ void UGroomAsset::CreateDebugData()
 	{
 		// 2.1 Generate transient strands data
 		const FHairDescriptionGroups& LocalHairDescriptionGroups = GetHairDescriptionGroups();
-		FHairGroupInfo DummyInfo;
 		FHairStrandsDatas StrandsData;
 		FHairStrandsDatas GuidesData;
-		FGroomBuilder::BuildData(LocalHairDescriptionGroups.HairGroups[GroupIndex], GetHairGroupsInterpolation()[GroupIndex], DummyInfo, StrandsData, GuidesData);
+		FGroomBuilder::BuildData(LocalHairDescriptionGroups.HairGroups[GroupIndex], GetHairGroupsInterpolation()[GroupIndex], StrandsData, GuidesData);
 
 		FHairGroupPlatformData& GroupData = GetHairGroupsPlatformData()[GroupIndex];
 		CreateHairStrandsDebugDatas(StrandsData, GroupData.Debug.Data);
@@ -3946,7 +3918,7 @@ bool UGroomAsset::GetHairStrandsDatas(
 		check(GroupIndex == HairGroup.Info.GroupID);
 
 		FHairGroupInfo DummyInfo;
-		FGroomBuilder::BuildData(HairGroup, GetHairGroupsInterpolation()[GroupIndex], DummyInfo, OutStrandsData, OutGuidesData);
+		FGroomBuilder::BuildData(HairGroup, GetHairGroupsInterpolation()[GroupIndex], OutStrandsData, OutGuidesData);
 	}
 	return bIsValid;
 }
@@ -3975,7 +3947,8 @@ bool UGroomAsset::GetHairCardsGuidesDatas(
 			FHairCardsInterpolationBulkData	OutDummyInterpolationBulkData;
 			FHairCardsBulkData OutDummyLODBulkData;
 			const FHairDescriptionGroups& LocalHairDescriptionGroups = GetHairDescriptionGroups();
-			InternalImportCardGeometry(Desc, LocalHairDescriptionGroups, OutDummyLODBulkData, OutCardsGuidesData, OutDummyInterpolationBulkData);
+			const TArray<FHairGroupsInterpolation>& LocalInterpolationGroups = GetHairGroupsInterpolation();
+			InternalImportCardGeometry(Desc, LocalHairDescriptionGroups, LocalInterpolationGroups, OutDummyLODBulkData, OutCardsGuidesData, OutDummyInterpolationBulkData);
 			return true;
 		}
 	}
diff --git a/Engine/Plugins/Runtime/HairStrands/Source/HairStrandsCore/Private/GroomBuilder.cpp b/Engine/Plugins/Runtime/HairStrands/Source/HairStrandsCore/Private/GroomBuilder.cpp
--- a/Engine/Plugins/Runtime/HairStrands/Source/HairStrandsCore/Private/GroomBuilder.cpp
+++ b/Engine/Plugins/Runtime/HairStrands/Source/HairStrandsCore/Private/GroomBuilder.cpp
@@ -2078,6 +2078,23 @@ void FGroomBuilder::BuildData(FHairStrandsDatas& OutStrands)
 	HairStrandsBuilder::BuildInternalData(OutStrands);
 }
 
+void FGroomBuilder::BuildData(
+	const FHairDescriptionGroup& InHairDescriptionGroup,
+	const FHairGroupsInterpolation& InSettings,
+	FHairStrandsDatas& OutRen,
+	FHairStrandsDatas& OutSim,
+	bool bAllowCurveReordering)
+{
+	FHairGroupInfo DummyGroupInfo;
+	BuildData(
+		InHairDescriptionGroup, 
+		InSettings, 
+		DummyGroupInfo, 
+		OutRen,
+		OutSim,
+		bAllowCurveReordering);
+}
+
 void FGroomBuilder::BuildData(
 	const FHairDescriptionGroup& InHairDescriptionGroup, 
 	const FHairGroupsInterpolation& InSettings, 
diff --git a/Engine/Plugins/Runtime/HairStrands/Source/HairStrandsCore/Private/HairCardsBuilder.cpp b/Engine/Plugins/Runtime/HairStrands/Source/HairStrandsCore/Private/HairCardsBuilder.cpp
--- a/Engine/Plugins/Runtime/HairStrands/Source/HairStrandsCore/Private/HairCardsBuilder.cpp
+++ b/Engine/Plugins/Runtime/HairStrands/Source/HairStrandsCore/Private/HairCardsBuilder.cpp
@@ -66,7 +66,7 @@ namespace FHairCardsBuilder
 FString GetVersion()
 {
 	// Important to update the version when cards building or importing changes
-	return TEXT("9m");
+	return TEXT("10");
 }
 
 bool InternalCreateCardsGuides(
diff --git a/Engine/Plugins/Runtime/HairStrands/Source/HairStrandsCore/Public/GroomBuilder.h b/Engine/Plugins/Runtime/HairStrands/Source/HairStrandsCore/Public/GroomBuilder.h
--- a/Engine/Plugins/Runtime/HairStrands/Source/HairStrandsCore/Public/GroomBuilder.h
+++ b/Engine/Plugins/Runtime/HairStrands/Source/HairStrandsCore/Public/GroomBuilder.h
@@ -70,6 +70,13 @@ struct HAIRSTRANDSCORE_API FGroomBuilder
 		FHairStrandsDatas& OutGuides,
 		bool bAllowCurveReordering=true);
 
+	static void BuildData(
+		const FHairDescriptionGroup& InHairDescriptionGroup,
+		const FHairGroupsInterpolation& InSettings,
+		FHairStrandsDatas& OutStrands,
+		FHairStrandsDatas& OutGuides,
+		bool bAllowCurveReordering=true);
+
 	// 2.b Build FHairStrandsDatas for Strands or Guides. 
 	// This version:
 	// * Suppose OutStrands already contains curves & points data. 
