Add texture layout support for groom strands texture generation.

diff --git a/Engine/Plugins/Runtime/HairStrands/Source/HairStrandsCore/Private/GroomAssetCards.cpp b/Engine/Plugins/Runtime/HairStrands/Source/HairStrandsCore/Private/GroomAssetCards.cpp
--- a/Engine/Plugins/Runtime/HairStrands/Source/HairStrandsCore/Private/GroomAssetCards.cpp
+++ b/Engine/Plugins/Runtime/HairStrands/Source/HairStrandsCore/Private/GroomAssetCards.cpp
@@ -88,12 +88,71 @@ UStaticMesh* FHairGroupsCardsSourceDescription::GetMesh() const
 
 uint32 GetHairTextureLayoutTextureCount(EHairTextureLayout In)
 {
+	uint32 OutCount = 0;
 	switch (In)
 	{
-		case EHairTextureLayout::Layout0: return 6u;
-		case EHairTextureLayout::Layout1: return 6u;
-		case EHairTextureLayout::Layout2: return 4u;
-		case EHairTextureLayout::Layout3: return 4u;
+		case EHairTextureLayout::Layout0: OutCount = 6u; break;
+		case EHairTextureLayout::Layout1: OutCount = 6u; break;
+		case EHairTextureLayout::Layout2: OutCount = 4u; break;
+		case EHairTextureLayout::Layout3: OutCount = 4u; break;
 	}
-	return 0;
+	check(OutCount <= HAIR_CARDS_MAX_TEXTURE_COUNT);
+	return OutCount;
 }
+
+const TCHAR* GetHairTextureLayoutTextureName(EHairTextureLayout InLayout, uint32 InIndex, bool bDetail)
+{
+	switch (InLayout)
+	{
+		case EHairTextureLayout::Layout0:
+		{
+			check(6 == GetHairTextureLayoutTextureCount(EHairTextureLayout::Layout0));
+			switch (InIndex)
+			{
+				case 0: return bDetail ? TEXT("Depth\n R8")									: TEXT("Depth");
+				case 1: return bDetail ? TEXT("Coverage\n R8")								: TEXT("Coverage");
+				case 2: return bDetail ? TEXT("Tangent\n RGB8")								: TEXT("Tangent");
+				case 3: return bDetail ? TEXT("Attributes\n RootUV | CoordU | Seed\n RGB8")	: TEXT("Attributes");
+				case 4: return bDetail ? TEXT("Material\n Color | Roughess\n RBGA8 ")		: TEXT("Material");
+				case 5: return bDetail ? TEXT("Auxiliary\n RGBA8")							: TEXT("Auxiliary");
+			}
+		} break;
+		case EHairTextureLayout::Layout1:
+		{
+			check(6 == GetHairTextureLayoutTextureCount(EHairTextureLayout::Layout1));
+			switch (InIndex)
+			{
+				case 0: return bDetail ? TEXT("Depth\n R8") 								: TEXT("Depth");
+				case 1: return bDetail ? TEXT("Coverage\n R8") 								: TEXT("Coverage");
+				case 2: return bDetail ? TEXT("Tangent\n RGB8") 							: TEXT("Tangent");
+				case 3: return bDetail ? TEXT("Attributes\n RootUV | CoordU | Seed\n RGB8") : TEXT("Attributes");
+				case 4: return bDetail ? TEXT("Material\n Color | GroupID\n RBGA8") 		: TEXT("Material");
+				case 5: return bDetail ? TEXT("Auxiliary\n RGBA8") 							: TEXT("Auxiliary");
+			}
+		} break;
+		case EHairTextureLayout::Layout2:
+		{
+			check(3 == GetHairTextureLayoutTextureCount(EHairTextureLayout::Layout2));
+			switch (InIndex)
+			{
+				case 0: return bDetail ? TEXT("Tangent | CoordU\n RGBA8")					: TEXT("TangentCoordU");
+				case 1: return bDetail ? TEXT("Coverage | Depth | Seed\n RGB8")				: TEXT("CoverageDepthSeed");
+				case 2: return bDetail ? TEXT("Color | Roughness\n RGBA8")					: TEXT("ColorRoughness");
+			}
+		} break;
+		case EHairTextureLayout::Layout3:
+		{
+			check(4 == GetHairTextureLayoutTextureCount(EHairTextureLayout::Layout3));
+			switch (InIndex)
+			{
+				case 0: return bDetail ? TEXT("Tangent\n RGB8")								: TEXT("Tangent");
+				case 1: return bDetail ? TEXT("RooUV | CoordU | GroupID\n RGBA8")			: TEXT("RooUVCoordUGroupID");
+				case 2: return bDetail ? TEXT("Coverage | Depth | Seed\n RGB8")				: TEXT("CoverageDepthSeed");
+				case 3: return bDetail ? TEXT("Color | Roughness\n RGBA8")					: TEXT("ColorRoughness");
+			}
+		} break;
+	}
+
+	check(false);
+	return TEXT("Unknown");
+}
\ No newline at end of file
diff --git a/Engine/Plugins/Runtime/HairStrands/Source/HairStrandsCore/Private/GroomTextureBuilder.cpp b/Engine/Plugins/Runtime/HairStrands/Source/HairStrandsCore/Private/GroomTextureBuilder.cpp
--- a/Engine/Plugins/Runtime/HairStrands/Source/HairStrandsCore/Private/GroomTextureBuilder.cpp
+++ b/Engine/Plugins/Runtime/HairStrands/Source/HairStrandsCore/Private/GroomTextureBuilder.cpp
@@ -19,6 +19,7 @@
 #include "TextureResource.h"
 #include "UObject/Package.h"
 #include "HairStrandsDefinitions.h"
+#include "RenderCaptureInterface.h"
 
 #if WITH_EDITORONLY_DATA
 #include "DerivedDataRequestOwner.h"
@@ -457,6 +458,107 @@ void FGroomTextureBuilder::BuildFollicleTexture(const TArray<FFollicleInfo>& InC
 ///////////////////////////////////////////////////////////////////////////////////////////////////
 // Strands texture generation
 
+///////////////////////////////////////////////////////////////////////////////////////////////////
+class FHairStrandTextureClearCS : public FGlobalShader
+{
+	DECLARE_GLOBAL_SHADER(FHairStrandTextureClearCS);
+	SHADER_USE_PARAMETER_STRUCT(FHairStrandTextureClearCS, FGlobalShader);
+
+	BEGIN_SHADER_PARAMETER_STRUCT(FParameters, )
+		SHADER_PARAMETER(uint32, LayoutIndex)
+		SHADER_PARAMETER(FIntPoint, Resolution)
+		SHADER_PARAMETER_RDG_TEXTURE_UAV(RWTexture2D, OutTexture)
+	END_SHADER_PARAMETER_STRUCT()
+
+public:
+	static bool ShouldCompilePermutation(const FGlobalShaderPermutationParameters& Parameters) { return IsHairStrandsSupported(EHairStrandsShaderType::Tool, Parameters.Platform); }
+	static void ModifyCompilationEnvironment(const FGlobalShaderPermutationParameters& Parameters, FShaderCompilerEnvironment& OutEnvironment)
+	{
+		FGlobalShader::ModifyCompilationEnvironment(Parameters, OutEnvironment);
+		OutEnvironment.SetDefine(TEXT("SHADER_TEXTURE_CLEAR"), 1);
+	}
+};
+
+IMPLEMENT_GLOBAL_SHADER(FHairStrandTextureClearCS, "/Engine/Private/HairStrands/HairStrandsTexturesGeneration.usf", "MainCS", SF_Compute);
+
+static void AddTextureClearPass(
+	FRDGBuilder& GraphBuilder,
+	FGlobalShaderMap* ShaderMap,
+	const FIntPoint& Resolution,
+	const EHairTextureLayout Layout,
+	FRDGTextureRef& OutTexture)
+{
+	FHairStrandTextureClearCS::FParameters* Parameters = GraphBuilder.AllocParameters<FHairStrandTextureClearCS::FParameters>();
+	Parameters->LayoutIndex = uint32(Layout);
+	Parameters->Resolution = Resolution;
+	Parameters->OutTexture = GraphBuilder.CreateUAV(OutTexture);
+
+	TShaderMapRef<FHairStrandTextureClearCS> ComputeShader(ShaderMap);
+	FIntVector DispatchCount = FComputeShaderUtils::GetGroupCount(Resolution, FIntPoint(8, 8));
+
+	FComputeShaderUtils::AddPass(
+		GraphBuilder,
+		RDG_EVENT_NAME("HairStrandsTexture::Clear"),
+		ComputeShader,
+		Parameters,
+		DispatchCount);
+}
+
+
+///////////////////////////////////////////////////////////////////////////////////////////////////
+class FHairStrandTextureCopyCS : public FGlobalShader
+{
+	DECLARE_GLOBAL_SHADER(FHairStrandTextureCopyCS);
+	SHADER_USE_PARAMETER_STRUCT(FHairStrandTextureCopyCS, FGlobalShader);
+
+	BEGIN_SHADER_PARAMETER_STRUCT(FParameters, )
+		SHADER_PARAMETER(uint32, TextureIndex)
+		SHADER_PARAMETER(uint32, TextureCount)
+		SHADER_PARAMETER(FIntPoint, Resolution)
+		SHADER_PARAMETER_RDG_TEXTURE(Texture2DArray, InTexture)
+		SHADER_PARAMETER_RDG_TEXTURE_UAV(RWTexture2D, OutTexture)
+	END_SHADER_PARAMETER_STRUCT()
+
+public:
+		static bool ShouldCompilePermutation(const FGlobalShaderPermutationParameters& Parameters) { return IsHairStrandsSupported(EHairStrandsShaderType::Tool, Parameters.Platform); }
+	static void ModifyCompilationEnvironment(const FGlobalShaderPermutationParameters& Parameters, FShaderCompilerEnvironment& OutEnvironment)
+	{
+		FGlobalShader::ModifyCompilationEnvironment(Parameters, OutEnvironment);
+		OutEnvironment.SetDefine(TEXT("SHADER_TEXTURE_COPY"), 1);
+	}
+};
+
+IMPLEMENT_GLOBAL_SHADER(FHairStrandTextureCopyCS, "/Engine/Private/HairStrands/HairStrandsTexturesGeneration.usf", "MainCS", SF_Compute);
+
+static FRDGTextureRef AddTextureCopyPass(
+	FRDGBuilder& GraphBuilder,
+	FGlobalShaderMap* ShaderMap,
+	uint32 InTextureIndex,
+	const FRDGTextureRef& InTexture)
+{
+	FRDGTextureRef OutTexture = GraphBuilder.CreateTexture(FRDGTextureDesc::Create2D(InTexture->Desc.Extent, InTexture->Desc.Format, FClearValueBinding::Black, TexCreate_UAV | TexCreate_ShaderResource | TexCreate_External), TEXT("Hair.Readback.TransientCopy"));
+
+	FHairStrandTextureCopyCS::FParameters* Parameters = GraphBuilder.AllocParameters<FHairStrandTextureCopyCS::FParameters>();
+	Parameters->TextureIndex = InTextureIndex;
+	Parameters->TextureCount = InTexture->Desc.Depth;
+	Parameters->Resolution = InTexture->Desc.Extent;
+	Parameters->InTexture = InTexture;
+	Parameters->OutTexture = GraphBuilder.CreateUAV(OutTexture);
+
+	TShaderMapRef<FHairStrandTextureCopyCS> ComputeShader(ShaderMap);
+	FIntVector DispatchCount = FComputeShaderUtils::GetGroupCount(Parameters->Resolution, FIntPoint(8, 8));
+
+	FComputeShaderUtils::AddPass(
+		GraphBuilder,
+		RDG_EVENT_NAME("HairStrandsTexture::Copy"),
+		ComputeShader,
+		Parameters,
+		DispatchCount);
+
+	return OutTexture;
+}
+
+
 
 ///////////////////////////////////////////////////////////////////////////////////////////////////
 class FHairStrandTextureDilationCS : public FGlobalShader
@@ -465,20 +567,15 @@ class FHairStrandTextureDilationCS : public FGlobalShader
 	SHADER_USE_PARAMETER_STRUCT(FHairStrandTextureDilationCS, FGlobalShader);
 
 	BEGIN_SHADER_PARAMETER_STRUCT(FParameters, )
+		SHADER_PARAMETER(uint32, LayoutIndex)
+		SHADER_PARAMETER(uint32, TextureCount)
+		SHADER_PARAMETER(uint32, bClearPass)
 		SHADER_PARAMETER(FIntPoint, Resolution)
-
-		SHADER_PARAMETER_RDG_TEXTURE(Texture2D, Source_DepthTexture)
-		SHADER_PARAMETER_RDG_TEXTURE(Texture2D, Source_CoverageTexture)
-		SHADER_PARAMETER_RDG_TEXTURE(Texture2D, Source_TangentTexture)
-		SHADER_PARAMETER_RDG_TEXTURE(Texture2D, Source_AttributeTexture)
-		SHADER_PARAMETER_RDG_TEXTURE(Texture2D, Source_MaterialTexture)
-
-		SHADER_PARAMETER_RDG_TEXTURE_UAV(RWTexture2D, TriangleMaskTexture)
-		SHADER_PARAMETER_RDG_TEXTURE_UAV(RWTexture2D, Target_DepthTexture)
-		SHADER_PARAMETER_RDG_TEXTURE_UAV(RWTexture2D, Target_CoverageTexture)
-		SHADER_PARAMETER_RDG_TEXTURE_UAV(RWTexture2D, Target_TangentTexture)
-		SHADER_PARAMETER_RDG_TEXTURE_UAV(RWTexture2D, Target_AttributeTexture)
-		SHADER_PARAMETER_RDG_TEXTURE_UAV(RWTexture2D, Target_MaterialTexture)
+		SHADER_PARAMETER_RDG_TEXTURE(Texture2D, DepthTestTexture)
+		SHADER_PARAMETER_RDG_TEXTURE(Texture2D, SourceTexture)
+		SHADER_PARAMETER_RDG_TEXTURE_UAV(RWTexture2D, TargetTexture)
+		SHADER_PARAMETER_RDG_TEXTURE(Texture2D, SourceTriangleMaskTexture)
+		SHADER_PARAMETER_RDG_TEXTURE_UAV(RWTexture2D, TargetTriangleMaskTexture)
 	END_SHADER_PARAMETER_STRUCT()
 
 public:
@@ -496,46 +593,36 @@ static void AddTextureDilationPass(
 	FRDGBuilder& GraphBuilder,
 	FGlobalShaderMap* ShaderMap,
 	const FIntPoint& Resolution,
-	FRDGTextureRef TriangleMaskTexture,
-
-	FRDGTextureRef Source_DepthTexture,
-	FRDGTextureRef Source_CoverageTexture,
-	FRDGTextureRef Source_TangentTexture,
-	FRDGTextureRef Source_AttributeTexture,
-	FRDGTextureRef Source_MaterialTexture,
-
-	FRDGTextureRef Target_DepthTexture,
-	FRDGTextureRef Target_CoverageTexture,
-	FRDGTextureRef Target_TangentTexture,
-	FRDGTextureRef Target_AttributeTexture,
-	FRDGTextureRef Target_MaterialTexture)
+	const EHairTextureLayout Layout,
+	const bool bClearPass,
+	FRDGTextureRef& DepthTestTexture,
+	FRDGTextureRef& SourceTriangleMaskTexture,
+	FRDGTextureRef& TargetTriangleMaskTexture,
+	FRDGTextureRef& SourceTexture,
+	FRDGTextureRef& TargetTexture)
 {
 	FHairStrandTextureDilationCS::FParameters* Parameters = GraphBuilder.AllocParameters<FHairStrandTextureDilationCS::FParameters>();
+	Parameters->TextureCount 			= HAIR_CARDS_MAX_TEXTURE_COUNT;
+	Parameters->LayoutIndex				= uint32(Layout);
+	Parameters->bClearPass				= bClearPass ? 1u : 0u;
 	Parameters->Resolution				= Resolution;
-	Parameters->TriangleMaskTexture		= GraphBuilder.CreateUAV(TriangleMaskTexture);
-	Parameters->Source_DepthTexture		= Source_DepthTexture;
-	Parameters->Source_CoverageTexture	= Source_CoverageTexture;
-	Parameters->Source_TangentTexture	= Source_TangentTexture;
-	Parameters->Source_AttributeTexture = Source_AttributeTexture;
-	Parameters->Source_MaterialTexture  = Source_MaterialTexture;
-	Parameters->Target_DepthTexture		= GraphBuilder.CreateUAV(Target_DepthTexture);
-	Parameters->Target_CoverageTexture	= GraphBuilder.CreateUAV(Target_CoverageTexture);
-	Parameters->Target_TangentTexture	= GraphBuilder.CreateUAV(Target_TangentTexture);
-	Parameters->Target_AttributeTexture = GraphBuilder.CreateUAV(Target_AttributeTexture);
-	Parameters->Target_MaterialTexture  = GraphBuilder.CreateUAV(Target_MaterialTexture);
+	Parameters->DepthTestTexture		= DepthTestTexture;
+	Parameters->SourceTriangleMaskTexture = SourceTriangleMaskTexture;
+	Parameters->TargetTriangleMaskTexture = GraphBuilder.CreateUAV(TargetTriangleMaskTexture);
+	Parameters->SourceTexture 			= SourceTexture;
+	Parameters->TargetTexture 			= GraphBuilder.CreateUAV(TargetTexture);
 
 	TShaderMapRef<FHairStrandTextureDilationCS> ComputeShader(ShaderMap);
-	FIntVector DispatchCount = FComputeShaderUtils::GetGroupCount(Resolution, FIntPoint(8, 4));
+	FIntVector DispatchCount = FComputeShaderUtils::GetGroupCount(Resolution, FIntPoint(8, 8));
 
 	FComputeShaderUtils::AddPass(
 		GraphBuilder,
-		RDG_EVENT_NAME("HairStrandsTetureDilation"),
+		RDG_EVENT_NAME("HairStrandsTexture::Dilation"),
 		ComputeShader,
 		Parameters,
 		DispatchCount);
 }
 
-
 static void InternalAllocateStrandsTexture(UTexture2D* Out, const FIntPoint& Resolution, uint32 MipCount, EPixelFormat Format, ETextureSourceFormat SourceFormat)
 {
 	FTextureFormatSettings FormatSettings;
@@ -555,33 +642,13 @@ static void InternalAllocateStrandsTexture(UTexture2D* Out, const FIntPoint& Res
 	//Out->UpdateResource();
 }
 
-static void InternalAllocateStrandsTexture_Depth(UTexture2D* Out, const FIntPoint& Resolution, uint32 MipCount)
-{
-	InternalAllocateStrandsTexture(Out, Resolution, 1, PF_G16, ETextureSourceFormat::TSF_G16);
-}
-
-static void InternalAllocateStrandsTexture_Coverage(UTexture2D* Out, const FIntPoint& Resolution, uint32 MipCount)
-{
-	InternalAllocateStrandsTexture(Out, Resolution, 1, PF_B8G8R8A8, ETextureSourceFormat::TSF_BGRA8);
-}
-
-static void InternalAllocateStrandsTexture_Tangent(UTexture2D* Out, const FIntPoint& Resolution, uint32 MipCount)
-{
-	InternalAllocateStrandsTexture(Out, Resolution, 1, PF_B8G8R8A8, ETextureSourceFormat::TSF_BGRA8);
-}
-
-static void InternalAllocateStrandsTexture_Attribute(UTexture2D* Out, const FIntPoint& Resolution, uint32 MipCount)
-{
-	InternalAllocateStrandsTexture(Out, Resolution, 1, PF_B8G8R8A8, ETextureSourceFormat::TSF_BGRA8);
-}
-
-static void InternalAllocateStrandsTexture_Material(UTexture2D* Out, const FIntPoint& Resolution, uint32 MipCount)
+static void InternalAllocateStrandsTexture_Generic(UTexture2D* Out, const FIntPoint& Resolution, uint32 MipCount)
 {
 	InternalAllocateStrandsTexture(Out, Resolution, 1, PF_B8G8R8A8, ETextureSourceFormat::TSF_BGRA8);
 }
 
 #if WITH_EDITOR
-FStrandsTexturesOutput FGroomTextureBuilder::CreateGroomStrandsTexturesTexture(const UGroomAsset* GroomAsset, uint32 InResolution)
+FStrandsTexturesOutput FGroomTextureBuilder::CreateGroomStrandsTexturesTexture(const UGroomAsset* GroomAsset, uint32 InResolution, EHairTextureLayout InLayout)
 {
 	FStrandsTexturesOutput Output;
 
@@ -592,11 +659,12 @@ FStrandsTexturesOutput FGroomTextureBuilder::CreateGroomStrandsTexturesTexture(c
 
 	const FString PackageName = GroomAsset->GetOutermost()->GetName();
 	const FIntPoint Resolution(InResolution, InResolution);
-	Output.Depth = FHairStrandsCore::CreateTexture(PackageName, Resolution, TEXT("_Depth"), InternalAllocateStrandsTexture_Depth);
-	Output.Coverage = FHairStrandsCore::CreateTexture(PackageName, Resolution, TEXT("_Opacity"), InternalAllocateStrandsTexture_Coverage);
-	Output.Tangent = FHairStrandsCore::CreateTexture(PackageName, Resolution, TEXT("_Tangent"), InternalAllocateStrandsTexture_Tangent);
-	Output.Attribute = FHairStrandsCore::CreateTexture(PackageName, Resolution, TEXT("_Attribute"), InternalAllocateStrandsTexture_Attribute);
-	Output.Material = FHairStrandsCore::CreateTexture(PackageName, Resolution, TEXT("_Material"), InternalAllocateStrandsTexture_Material);
+	const uint32 TextureCount = GetHairTextureLayoutTextureCount(InLayout);
+	Output.Textures.SetNum(TextureCount);
+	for (uint32 TextureIt=0;TextureIt<TextureCount;++TextureIt)
+	{
+		Output.Textures[TextureIt] = FHairStrandsCore::CreateTexture(PackageName, Resolution, FString(TEXT("_")) + GetHairTextureLayoutTextureName(InLayout, TextureIt, false/*bDetail*/), InternalAllocateStrandsTexture_Generic);
+	}
 	return Output;
 }
 #endif
@@ -639,6 +707,7 @@ class FHairStrandsTexturePS : public FGlobalShader
 
 		BEGIN_SHADER_PARAMETER_STRUCT(FParameters, )
 		SHADER_PARAMETER_STRUCT_INCLUDE(ShaderPrint::FShaderParameters, ShaderPrintParameters)
+		SHADER_PARAMETER(uint32, LayoutIndex)
 		SHADER_PARAMETER(FIntPoint, OutputResolution)
 		SHADER_PARAMETER(uint32, VertexCount)
 		SHADER_PARAMETER(float, MaxDistance)
@@ -663,6 +732,7 @@ class FHairStrandsTexturePS : public FGlobalShader
 		SHADER_PARAMETER_SRV(Buffer, VertexBuffer)
 		SHADER_PARAMETER_SRV(Buffer, UVsBuffer)
 		SHADER_PARAMETER_SRV(Buffer, NormalsBuffer)
+		SHADER_PARAMETER_RDG_TEXTURE_UAV(RWTexture2DArray, OutTexture)
 		RENDER_TARGET_BINDING_SLOTS()
 	END_SHADER_PARAMETER_STRUCT()
 
@@ -690,14 +760,11 @@ static FIntPoint GetTileResolution(FIntPoint In)
 
 struct FHairStrandsRDGTextures
 {
+	EHairTextureLayout Layout = EHairTextureLayout::Layout0;
 	FIntPoint TileCount = FIntPoint::ZeroValue; // Total tile count
 	FIntPoint TileCoord = FIntPoint::ZeroValue; // Current tile coord
 
-	FRDGTextureRef DepthTexture = nullptr;
-	FRDGTextureRef CoverageTexture = nullptr;
-	FRDGTextureRef TangentTexture = nullptr;
-	FRDGTextureRef AttributeTexture = nullptr;
-	FRDGTextureRef MaterialTexture = nullptr;
+	FRDGTextureRef Texture;
 	FRDGTextureRef TriangleMaskTexture = nullptr;
 	FRDGTextureRef DepthTestTexture = nullptr;
 
@@ -706,14 +773,11 @@ struct FHairStrandsRDGTextures
 
 struct FHairStrandsRTs
 {
+	EHairTextureLayout Layout = EHairTextureLayout::Layout0;
 	FIntPoint TileCount; // Total tile count
 	FIntPoint TileCoord; // Current tile coord
 
-	TRefCountPtr<IPooledRenderTarget> DepthTexture;
-	TRefCountPtr<IPooledRenderTarget> CoverageTexture;
-	TRefCountPtr<IPooledRenderTarget> TangentTexture;
-	TRefCountPtr<IPooledRenderTarget> AttributeTexture;
-	TRefCountPtr<IPooledRenderTarget> MaterialTexture;
+	TRefCountPtr<IPooledRenderTarget> Texture;
 	TRefCountPtr<IPooledRenderTarget> TriangleMaskTexture;
 	TRefCountPtr<IPooledRenderTarget> DepthTestTexture;
 	
@@ -722,21 +786,19 @@ struct FHairStrandsRTs
 
 struct FStrandsTexturesReadback
 {
-	TUniquePtr<FRHIGPUTextureReadback> Depth;
-	TUniquePtr<FRHIGPUTextureReadback> Tangent;
-	TUniquePtr<FRHIGPUTextureReadback> Coverage;
-	TUniquePtr<FRHIGPUTextureReadback> Attribute;
-	TUniquePtr<FRHIGPUTextureReadback> Material;
+	TArray<TUniquePtr<FRHIGPUTextureReadback>> Textures;
 	uint32 NotReadyFrameCount = 0;
 
 	bool IsReady() const
 	{
-		return
-			Depth->IsReady() &&
-			Tangent->IsReady() &&
-			Coverage->IsReady() &&
-			Attribute->IsReady() &&
-			Material->IsReady();
+		for (const TUniquePtr<FRHIGPUTextureReadback>& Texture : Textures)
+		{
+			if (!Texture->IsReady())
+			{
+				return false;
+			}
+		}
+		return true;
 	}
 };
 
@@ -774,10 +836,11 @@ static void InternalGenerateHairStrandsTextures(
 	FHairStrandsInstanceRawParameters& Instance,
 	FHairStrandsRDGTextures& Out)
 {
-	const FIntPoint OutputResolution = Out.DepthTexture->Desc.Extent;
+	const FIntPoint OutputResolution = Out.Texture->Desc.Extent;
 	const FIntPoint OutTileCoord = Out.TileCoord;
 
 	FHairStrandsTexturePS::FParameters* ParametersPS = GraphBuilder.AllocParameters<FHairStrandsTexturePS::FParameters>();
+	ParametersPS->LayoutIndex = uint32(Out.Layout);
 	ParametersPS->OutputResolution = OutputResolution;
 	ParametersPS->VertexCount = VertexCount;
 	ParametersPS->VertexBuffer = InMeshVertexBuffer;
@@ -807,12 +870,8 @@ static void InternalGenerateHairStrandsTextures(
 		ShaderPrint::SetParameters(GraphBuilder, *ShaderPrintData, ParametersPS->ShaderPrintParameters);
 	}
 
-	ParametersPS->RenderTargets[0] = FRenderTargetBinding(Out.DepthTexture, bClear ? ERenderTargetLoadAction::EClear : ERenderTargetLoadAction::ELoad);
-	ParametersPS->RenderTargets[1] = FRenderTargetBinding(Out.TangentTexture, bClear ? ERenderTargetLoadAction::EClear : ERenderTargetLoadAction::ELoad);
-	ParametersPS->RenderTargets[2] = FRenderTargetBinding(Out.CoverageTexture, bClear ? ERenderTargetLoadAction::EClear : ERenderTargetLoadAction::ELoad);
-	ParametersPS->RenderTargets[3] = FRenderTargetBinding(Out.AttributeTexture, bClear ? ERenderTargetLoadAction::EClear : ERenderTargetLoadAction::ELoad);
-	ParametersPS->RenderTargets[4] = FRenderTargetBinding(Out.MaterialTexture, bClear ? ERenderTargetLoadAction::EClear : ERenderTargetLoadAction::ELoad);
-	ParametersPS->RenderTargets[5] = FRenderTargetBinding(Out.TriangleMaskTexture, bClear ? ERenderTargetLoadAction::EClear : ERenderTargetLoadAction::ELoad);
+	ParametersPS->OutTexture = GraphBuilder.CreateUAV(Out.Texture);
+	ParametersPS->RenderTargets[0] = FRenderTargetBinding(Out.TriangleMaskTexture, bClear ? ERenderTargetLoadAction::EClear : ERenderTargetLoadAction::ELoad);
 	ParametersPS->RenderTargets.DepthStencil = FDepthStencilBinding(
 		Out.DepthTestTexture,
 		bClear ? ERenderTargetLoadAction::EClear : ERenderTargetLoadAction::ELoad,
@@ -875,51 +934,30 @@ static void InternalGenerateHairStrandsTextures(
 			}
 			else
 			{
-				//for (int32 TileY = 0; TileY < TileCountY; ++TileY)
-				//for (int32 TileX = 0; TileX < TileCountX; ++TileX)
-				{
-					const uint32 OffsetX = OutTileCoord.X * TileSize;
-					const uint32 OffsetY = OutTileCoord.Y * TileSize;
-					RHICmdList.SetScissorRect(true, OffsetX, OffsetY, OffsetX + TileSize, OffsetY + TileSize);
-					RHICmdList.DrawIndexedPrimitive(InMeshIndexBuffer, VertexBaseIndex, 0, VertexCount, IndexBaseIndex, PrimitiveCount, 1);
-
-					// Flush, to ensure that all texture generation is done (TDR)
-					#if 0
-					GDynamicRHI->RHISubmitCommandsAndFlushGPU();
-					GDynamicRHI->RHIBlockUntilGPUIdle();
-					#endif
-				}
+				const uint32 OffsetX = OutTileCoord.X * TileSize;
+				const uint32 OffsetY = OutTileCoord.Y * TileSize;
+				RHICmdList.SetScissorRect(true, OffsetX, OffsetY, OffsetX + TileSize, OffsetY + TileSize);
+				RHICmdList.DrawIndexedPrimitive(InMeshIndexBuffer, VertexBaseIndex, 0, VertexCount, IndexBaseIndex, PrimitiveCount, 1);
+
+				// Flush, to ensure that all texture generation is done (TDR)
+				#if 0
+				GDynamicRHI->RHISubmitCommandsAndFlushGPU();
+				GDynamicRHI->RHIBlockUntilGPUIdle();
+				#endif
 			}
 		});
 }
 
 
-static FHairStrandsRDGTextures CreateTextures(FRDGBuilder& GraphBuilder, FIntPoint InResolution)
+static FHairStrandsRDGTextures CreateTextures(FRDGBuilder& GraphBuilder, FIntPoint InResolution, EHairTextureLayout Layout)
 {
 	FHairStrandsRDGTextures Out;
+	Out.Layout = Layout;
 
-	FRDGTextureDesc Desc = FRDGTextureDesc::Create2D(InResolution, PF_A8R8G8B8, FClearValueBinding::Black, TexCreate_RenderTargetable | TexCreate_UAV | TexCreate_ShaderResource);
-
-	Desc.Format = PF_G16;
-	Desc.ClearValue = FClearValueBinding::White;
-	Out.DepthTexture = GraphBuilder.CreateTexture(Desc, TEXT("Hair.DepthTexture"));
-
-	Desc.Format = PF_B8G8R8A8;
-	Desc.ClearValue = FClearValueBinding::Transparent;
-	Out.CoverageTexture = GraphBuilder.CreateTexture(Desc, TEXT("Hair.CoverageTexture"));
-
-	Desc.Format = PF_B8G8R8A8;
-	Desc.ClearValue = FClearValueBinding::Transparent;
-	Out.TangentTexture = GraphBuilder.CreateTexture(Desc, TEXT("Hair.TangentTexture"));
-
-	Desc.Format = PF_B8G8R8A8;
-	Desc.ClearValue = FClearValueBinding::Transparent;
-	Out.AttributeTexture = GraphBuilder.CreateTexture(Desc, TEXT("Hair.AttributeTexture"));
-
-	Desc.Format = PF_B8G8R8A8;
-	Desc.ClearValue = FClearValueBinding::Transparent;
-	Out.MaterialTexture = GraphBuilder.CreateTexture(Desc, TEXT("Hair.MaterialTexture"));
+	FRDGTextureDesc Desc = FRDGTextureDesc::Create2DArray(InResolution, PF_B8G8R8A8, FClearValueBinding::Black, TexCreate_RenderTargetable | TexCreate_UAV | TexCreate_ShaderResource | TexCreate_External, HAIR_CARDS_MAX_TEXTURE_COUNT);
+	Out.Texture = GraphBuilder.CreateTexture(Desc, TEXT("Hair.Texture"));
 
+	// Textures for depth testing
 	Desc.Format = PF_R32_UINT;
 	Desc.ClearValue = FClearValueBinding::Black;
 	Out.TriangleMaskTexture = GraphBuilder.CreateTexture(Desc, TEXT("Hair.TriangleMaskTexture"));
@@ -938,27 +976,21 @@ static FHairStrandsRDGTextures CreateTextures(FRDGBuilder& GraphBuilder, FIntPoi
 static FHairStrandsRDGTextures RegisterTextures(FRDGBuilder& GraphBuilder, const FHairStrandsRTs& In)
 {
 	FHairStrandsRDGTextures Out;
-	Out.DepthTexture		= GraphBuilder.RegisterExternalTexture(In.DepthTexture);
-	Out.CoverageTexture		= GraphBuilder.RegisterExternalTexture(In.CoverageTexture);
-	Out.TangentTexture		= GraphBuilder.RegisterExternalTexture(In.TangentTexture);
-	Out.AttributeTexture	= GraphBuilder.RegisterExternalTexture(In.AttributeTexture);
-	Out.MaterialTexture		= GraphBuilder.RegisterExternalTexture(In.MaterialTexture);
+	Out.Layout = In.Layout;
+	Out.Texture = GraphBuilder.RegisterExternalTexture(In.Texture);
+	Out.DepthTestTexture = GraphBuilder.RegisterExternalTexture(In.DepthTestTexture);
 	Out.TriangleMaskTexture = GraphBuilder.RegisterExternalTexture(In.TriangleMaskTexture);
-	Out.DepthTestTexture	= GraphBuilder.RegisterExternalTexture(In.DepthTestTexture);
-	Out.TileCoord			= In.TileCoord;
-	Out.TileCount			= In.TileCount;
+	Out.TileCoord = In.TileCoord;
+	Out.TileCount = In.TileCount;
 	return Out;
 }
 
 static void ExportTextures(FRDGBuilder& GraphBuilder, const FHairStrandsRDGTextures& In, FHairStrandsRTs& Out)
 {
-	GraphBuilder.QueueTextureExtraction(In.DepthTexture,		&Out.DepthTexture);
-	GraphBuilder.QueueTextureExtraction(In.CoverageTexture,		&Out.CoverageTexture);
-	GraphBuilder.QueueTextureExtraction(In.TangentTexture,		&Out.TangentTexture);
-	GraphBuilder.QueueTextureExtraction(In.AttributeTexture,	&Out.AttributeTexture);
-	GraphBuilder.QueueTextureExtraction(In.MaterialTexture,		&Out.MaterialTexture);
+	Out.Layout = In.Layout;
+	GraphBuilder.QueueTextureExtraction(In.Texture, &Out.Texture);
 	GraphBuilder.QueueTextureExtraction(In.TriangleMaskTexture, &Out.TriangleMaskTexture);
-	GraphBuilder.QueueTextureExtraction(In.DepthTestTexture,	&Out.DepthTestTexture);
+	GraphBuilder.QueueTextureExtraction(In.DepthTestTexture, &Out.DepthTestTexture);
 	Out.TileCoord = In.TileCoord;
 	Out.TileCount = In.TileCount;
 }
@@ -1136,11 +1168,15 @@ static bool TraceTextures(
 
 static FHairStrandsRDGTextures DilateTextures(FRDGBuilder& GraphBuilder, FGlobalShaderMap* ShaderMap, FHairStrandsRDGTextures& In)
 {
-	const FIntPoint OutputResolution = In.CoverageTexture->Desc.Extent;
+	const FIntPoint OutputResolution = In.Texture->Desc.Extent;
 
 	FHairStrandsRDGTextures Out[2];
+	// Overwrite the TriangleMaskTexture which will be reinitialized from the depth test texture over the iteration
 	Out[0] = In;
-	Out[1] = CreateTextures(GraphBuilder, OutputResolution);
+	Out[0].TriangleMaskTexture = CreateTextures(GraphBuilder, OutputResolution, In.Layout).TriangleMaskTexture;
+	Out[1] = CreateTextures(GraphBuilder, OutputResolution, In.Layout);
+
+	AddClearUAVPass(GraphBuilder, GraphBuilder.CreateUAV(Out[0].TriangleMaskTexture), 0u);
 
 	uint32 SourceIndex = 0;
 	uint32 TargetIndex = 0;
@@ -1153,34 +1189,30 @@ static FHairStrandsRDGTextures DilateTextures(FRDGBuilder& GraphBuilder, FGlobal
 			GraphBuilder,
 			ShaderMap,
 			OutputResolution,
-			In.TriangleMaskTexture,
-
-			Out[SourceIndex].DepthTexture,
-			Out[SourceIndex].CoverageTexture,
-			Out[SourceIndex].TangentTexture,
-			Out[SourceIndex].AttributeTexture,
-			Out[SourceIndex].MaterialTexture,
-
-			Out[TargetIndex].DepthTexture,
-			Out[TargetIndex].CoverageTexture,
-			Out[TargetIndex].TangentTexture,
-			Out[TargetIndex].AttributeTexture,
-			Out[TargetIndex].MaterialTexture);
+			In.Layout,
+			DilationIt == 0,
+			In.DepthTestTexture,
+			Out[SourceIndex].TriangleMaskTexture,
+			Out[TargetIndex].TriangleMaskTexture,
+			Out[SourceIndex].Texture,
+			Out[TargetIndex].Texture);
 
 		SourceIndex = TargetIndex;
 	}
-
-	Out[SourceIndex].TriangleMaskTexture = In.TriangleMaskTexture;
 	return Out[SourceIndex];
 }
 
-static void CopyTextures(FRDGBuilder& GraphBuilder, const FHairStrandsRDGTextures& In, FStrandsTexturesReadback& Out)
+static void CopyTextures(FRDGBuilder& GraphBuilder, FGlobalShaderMap* ShaderMap, const EHairTextureLayout InLayout, const FHairStrandsRDGTextures& In, FStrandsTexturesReadback& Out)
 {
-	AddEnqueueCopyPass(GraphBuilder, Out.Depth.Get(),		In.DepthTexture);
-	AddEnqueueCopyPass(GraphBuilder, Out.Coverage.Get(),	In.CoverageTexture);
-	AddEnqueueCopyPass(GraphBuilder, Out.Tangent.Get(),		In.TangentTexture);
-	AddEnqueueCopyPass(GraphBuilder, Out.Attribute.Get(),	In.AttributeTexture);
-	AddEnqueueCopyPass(GraphBuilder, Out.Material.Get(),	In.MaterialTexture);
+	const uint32 TextureCount = GetHairTextureLayoutTextureCount(InLayout);
+	for (uint32 TextureIt=0; TextureIt<TextureCount;++TextureIt)
+	{
+		// Split texture array into seperate texture
+		FRDGTextureRef Texture = AddTextureCopyPass(GraphBuilder, ShaderMap, TextureIt, In.Texture);
+
+		// Enqueue copy to readback texture
+		AddEnqueueCopyPass(GraphBuilder, Out.Textures[TextureIt].Get(), Texture);
+	}
 }
 
 static void CopyReadbackToTexture(TUniquePtr<FRHIGPUTextureReadback>& In, UTexture2D* Out)
@@ -1252,18 +1284,11 @@ void RunHairStrandsTexturesQueries(FRDGBuilder& GraphBuilder, FGlobalShaderMap*
 			{
 				if (R->Readback.IsReady() || ++R->Readback.NotReadyFrameCount > 32)
 				{
-					CopyReadbackToTexture(R->Readback.Depth,    R->Output.Depth);
-					CopyReadbackToTexture(R->Readback.Tangent,  R->Output.Tangent);
-					CopyReadbackToTexture(R->Readback.Coverage, R->Output.Coverage);
-					CopyReadbackToTexture(R->Readback.Attribute,R->Output.Attribute);
-					CopyReadbackToTexture(R->Readback.Material, R->Output.Material);
-	
-					R->Readback.Depth    = nullptr;
-					R->Readback.Tangent  = nullptr;
-					R->Readback.Coverage = nullptr;
-					R->Readback.Attribute= nullptr;
-					R->Readback.Material = nullptr;
-	
+					for (uint32 TextureIt = 0, TextureCount = R->Readback.Textures.Num(); TextureIt < TextureCount; ++TextureIt)
+					{
+						CopyReadbackToTexture(R->Readback.Textures[TextureIt], R->Output.Textures[TextureIt]);
+					}
+					R->Readback.Textures.Empty();
 					delete R;
 				}
 				else
@@ -1287,6 +1312,9 @@ void RunHairStrandsTexturesQueries(FRDGBuilder& GraphBuilder, FGlobalShaderMap*
 		{
 			if (R)
 			{
+
+				RenderCaptureInterface::FScopedCapture RenderCapture(true, GraphBuilder, TEXT("PassX"));
+
 				FHairStrandsRDGTextures Textures = RegisterTextures(GraphBuilder, R->InProgress);
 
 				TraceTextures(GraphBuilder, ShaderMap, R->Info, DebugShaderData, Textures);
@@ -1294,7 +1322,7 @@ void RunHairStrandsTexturesQueries(FRDGBuilder& GraphBuilder, FGlobalShaderMap*
 				if (Textures.IsCompleted())
 				{
 					Textures = DilateTextures(GraphBuilder, ShaderMap, Textures);
-					CopyTextures(GraphBuilder, Textures, R->Readback);
+					CopyTextures(GraphBuilder, ShaderMap, R->Info.Layout, Textures, R->Readback);
 					GStrandsTexturesReadbacks.Enqueue(R);
 				}
 				else
@@ -1317,22 +1345,41 @@ void RunHairStrandsTexturesQueries(FRDGBuilder& GraphBuilder, FGlobalShaderMap*
 		while (GStrandsTexturesQueries.Dequeue(Q))
 		{
 			FStrandsTexturesProcess* R = new FStrandsTexturesProcess();
-			R->Readback.Depth     = MakeUnique<FRHIGPUTextureReadback>(TEXT("Readback.Depth"));
-			R->Readback.Tangent   = MakeUnique<FRHIGPUTextureReadback>(TEXT("Readback.Tangent"));
-			R->Readback.Coverage  = MakeUnique<FRHIGPUTextureReadback>(TEXT("Readback.Coverage"));
-			R->Readback.Attribute = MakeUnique<FRHIGPUTextureReadback>(TEXT("Readback.Attribute"));
-			R->Readback.Material  = MakeUnique<FRHIGPUTextureReadback>(TEXT("Readback.Material"));		
-			R->Output			  = Q.Output;
-			R->Info				  = Q.Info;
 
-			FHairStrandsRDGTextures Textures = CreateTextures(GraphBuilder, FIntPoint(FMath::Clamp(Q.Info.Resolution, 512u, 16384u), FMath::Clamp(Q.Info.Resolution, 512u, 16384u)));
+			auto GetTextureName = [](uint32 In)
+			{
+				switch (In)
+				{	
+					case 0: return TEXT("Hair.Readback.Texture0");
+					case 1: return TEXT("Hair.Readback.Texture1");
+					case 2: return TEXT("Hair.Readback.Texture2");
+					case 3: return TEXT("Hair.Readback.Texture3");
+					case 4: return TEXT("Hair.Readback.Texture4");
+					case 5: return TEXT("Hair.Readback.Texture5");
+				}
+				return  TEXT("Hair.Readback.Unknown");
+			};
+
+			const uint32 TextureCount = GetHairTextureLayoutTextureCount(Q.Info.Layout);
+			R->Readback.Textures.SetNum(TextureCount);
+			for (uint32 TextureIt = 0; TextureIt < TextureCount; ++TextureIt)
+			{
+				R->Readback.Textures[TextureIt] = MakeUnique<FRHIGPUTextureReadback>(GetTextureName(TextureIt));
+			}	
+			R->Output = Q.Output;
+			R->Info = Q.Info;
+
+			RenderCaptureInterface::FScopedCapture RenderCapture(true, GraphBuilder, TEXT("Pass0"));
+
+			FHairStrandsRDGTextures Textures = CreateTextures(GraphBuilder, FIntPoint(FMath::Clamp(Q.Info.Resolution, 512u, 16384u), FMath::Clamp(Q.Info.Resolution, 512u, 16384u)), Q.Info.Layout);
+			AddTextureClearPass(GraphBuilder, ShaderMap, Q.Info.Resolution, Q.Info.Layout, Textures.Texture);
 
 			TraceTextures(GraphBuilder, ShaderMap, Q.Info, DebugShaderData, Textures);
 
 			if (Textures.IsCompleted())
 			{
 				Textures = DilateTextures(GraphBuilder, ShaderMap, Textures);
-				CopyTextures(GraphBuilder, Textures, R->Readback);
+				CopyTextures(GraphBuilder, ShaderMap, R->Info.Layout, Textures, R->Readback);
 				GStrandsTexturesReadbacks.Enqueue(R);
 			}
 			else
diff --git a/Engine/Plugins/Runtime/HairStrands/Source/HairStrandsCore/Public/GroomAssetCards.h b/Engine/Plugins/Runtime/HairStrands/Source/HairStrandsCore/Public/GroomAssetCards.h
--- a/Engine/Plugins/Runtime/HairStrands/Source/HairStrandsCore/Public/GroomAssetCards.h
+++ b/Engine/Plugins/Runtime/HairStrands/Source/HairStrandsCore/Public/GroomAssetCards.h
@@ -39,6 +39,7 @@ enum class EHairTextureLayout : uint8
 
 // Returns the number of textures used for a particular layout
 HAIRSTRANDSCORE_API uint32 GetHairTextureLayoutTextureCount(EHairTextureLayout In);
+HAIRSTRANDSCORE_API const TCHAR* GetHairTextureLayoutTextureName(EHairTextureLayout InLayout, uint32 InIndex, bool bDetail);
 
 USTRUCT(BlueprintType)
 struct HAIRSTRANDSCORE_API FHairGroupCardsInfo
diff --git a/Engine/Plugins/Runtime/HairStrands/Source/HairStrandsCore/Public/GroomCreateStrandsTexturesOptions.h b/Engine/Plugins/Runtime/HairStrands/Source/HairStrandsCore/Public/GroomCreateStrandsTexturesOptions.h
--- a/Engine/Plugins/Runtime/HairStrands/Source/HairStrandsCore/Public/GroomCreateStrandsTexturesOptions.h
+++ b/Engine/Plugins/Runtime/HairStrands/Source/HairStrandsCore/Public/GroomCreateStrandsTexturesOptions.h
@@ -33,6 +33,10 @@ class HAIRSTRANDSCORE_API UGroomCreateStrandsTexturesOptions : public UObject
 
 public:
 
+	/** Resolution of the output texture maps (tangent, coverage, ...) */
+	UPROPERTY(EditAnywhere, BlueprintReadWrite, meta = (ShowOnlyInnerProperties), Category = Options)
+	EHairTextureLayout Layout = EHairTextureLayout::Layout1;
+
 	/** Resolution of the output texture maps (tangent, coverage, ...) */
 	UPROPERTY(EditAnywhere, BlueprintReadWrite, meta = (ShowOnlyInnerProperties), Category = Options)
 	int32 Resolution = 2048;
diff --git a/Engine/Plugins/Runtime/HairStrands/Source/HairStrandsCore/Public/GroomTextureBuilder.h b/Engine/Plugins/Runtime/HairStrands/Source/HairStrandsCore/Public/GroomTextureBuilder.h
--- a/Engine/Plugins/Runtime/HairStrands/Source/HairStrandsCore/Public/GroomTextureBuilder.h
+++ b/Engine/Plugins/Runtime/HairStrands/Source/HairStrandsCore/Public/GroomTextureBuilder.h
@@ -7,6 +7,7 @@
 #include "UObject/ObjectMacros.h"
 #include "GroomSettings.h"
 #include "HairStrandsInterface.h"
+#include "GroomAssetCards.h"
 
 class UTexture2D;
 class UGroomAsset;
@@ -33,6 +34,7 @@ struct FFollicleInfo
 
 struct FStrandsTexturesInfo
 {
+	EHairTextureLayout Layout = EHairTextureLayout::Layout1;
 	UGroomAsset* GroomAsset = nullptr;
 	USkeletalMesh* SkeletalMesh = nullptr;
 	UStaticMesh* StaticMesh = nullptr;
@@ -47,12 +49,8 @@ struct FStrandsTexturesInfo
 
 struct FStrandsTexturesOutput
 {
-	class UTexture2D* Depth = nullptr;
-	class UTexture2D* Tangent = nullptr;
-	class UTexture2D* Coverage = nullptr;
-	class UTexture2D* Attribute = nullptr;
-	class UTexture2D* Material = nullptr;
-	bool IsValid() const { return Depth && Tangent && Coverage && Attribute && Material; }
+	TArray<UTexture2D*> Textures;
+	bool IsValid() const { return Textures.Num() > 0; }
 };
 
 ////////////////////////////////////////////////////////////////////////////////////////////////////////////////////////
@@ -69,7 +67,7 @@ struct HAIRSTRANDSCORE_API FGroomTextureBuilder
 
 	// Strands textures
 #if WITH_EDITOR
-	static FStrandsTexturesOutput CreateGroomStrandsTexturesTexture(const UGroomAsset* GroomAsset, uint32 Resolution);
+	static FStrandsTexturesOutput CreateGroomStrandsTexturesTexture(const UGroomAsset* GroomAsset, uint32 Resolution, EHairTextureLayout InLayout);
 	static void BuildStrandsTextures(const FStrandsTexturesInfo& InInfo, const FStrandsTexturesOutput& Output);
 #endif
 };
diff --git a/Engine/Plugins/Runtime/HairStrands/Source/HairStrandsEditor/Private/GroomActions.cpp b/Engine/Plugins/Runtime/HairStrands/Source/HairStrandsEditor/Private/GroomActions.cpp
--- a/Engine/Plugins/Runtime/HairStrands/Source/HairStrandsEditor/Private/GroomActions.cpp
+++ b/Engine/Plugins/Runtime/HairStrands/Source/HairStrandsEditor/Private/GroomActions.cpp
@@ -397,6 +397,7 @@ void ExecuteCreateStrandsTextures(const FToolMenuContext& InContext)
 				}
 				
 				FStrandsTexturesInfo Info;
+				Info.Layout = CurrentOptions->Layout;
 				Info.GroomAsset   = GroomAsset;
 				Info.TracingDirection = SignDirection;
 				Info.MaxTracingDistance = MaxDistance;
@@ -417,7 +418,7 @@ void ExecuteCreateStrandsTextures(const FToolMenuContext& InContext)
 						Info.GroupIndices.Add(GroupIndex);
 					}
 				}
-				FStrandsTexturesOutput Output = FGroomTextureBuilder::CreateGroomStrandsTexturesTexture(GroomAsset, Info.Resolution);
+				FStrandsTexturesOutput Output = FGroomTextureBuilder::CreateGroomStrandsTexturesTexture(GroomAsset, Info.Resolution, Info.Layout);
 				if (Output.IsValid())
 				{
 					FGroomTextureBuilder::BuildStrandsTextures(Info, Output);
diff --git a/Engine/Plugins/Runtime/HairStrands/Source/HairStrandsEditor/Private/GroomAssetDetails.cpp b/Engine/Plugins/Runtime/HairStrands/Source/HairStrandsEditor/Private/GroomAssetDetails.cpp
--- a/Engine/Plugins/Runtime/HairStrands/Source/HairStrandsEditor/Private/GroomAssetDetails.cpp
+++ b/Engine/Plugins/Runtime/HairStrands/Source/HairStrandsEditor/Private/GroomAssetDetails.cpp
@@ -1516,43 +1516,9 @@ void FGroomRenderingDetails::OnGenerateElementForHairGroup(TSharedRef<IPropertyH
 					// Textures
 					const uint32 TextureCount = GroomAsset->GetHairGroupsCards()[GroupIndex].Textures.Textures.Num();
 					TSharedPtr<IPropertyHandle> TextureArrayHandle = ChildHandle->GetChildHandle(1);
-					switch (LayoutType)
+					for (uint32 TextureIt = 0; TextureIt < TextureCount; ++TextureIt)
 					{
-						case EHairTextureLayout::Layout0:
-						{
-							check(TextureCount == GetHairTextureLayoutTextureCount(EHairTextureLayout::Layout0));
-							TextureGroup.AddPropertyRow(TextureArrayHandle->GetChildHandle(0).ToSharedRef()).DisplayName(FText(LOCTEXT("Layout0_0", "Depth\n R8")));
-							TextureGroup.AddPropertyRow(TextureArrayHandle->GetChildHandle(1).ToSharedRef()).DisplayName(FText(LOCTEXT("Layout0_1", "Coverage\n R8")));
-							TextureGroup.AddPropertyRow(TextureArrayHandle->GetChildHandle(2).ToSharedRef()).DisplayName(FText(LOCTEXT("Layout0_2", "Tangent\n RGB8")));
-							TextureGroup.AddPropertyRow(TextureArrayHandle->GetChildHandle(3).ToSharedRef()).DisplayName(FText(LOCTEXT("Layout0_3", "Attributes\n RootUV | CoordU | Seed\n RGB8")));
-							TextureGroup.AddPropertyRow(TextureArrayHandle->GetChildHandle(4).ToSharedRef()).DisplayName(FText(LOCTEXT("Layout0_4", "Material\n Color | Roughess\n RBGA8 ")));
-							TextureGroup.AddPropertyRow(TextureArrayHandle->GetChildHandle(5).ToSharedRef()).DisplayName(FText(LOCTEXT("Layout0_5", "Auxiliary\n RGBA8")));
-						} break;
-						case EHairTextureLayout::Layout1:
-						{
-							check(TextureCount == GetHairTextureLayoutTextureCount(EHairTextureLayout::Layout1));
-							TextureGroup.AddPropertyRow(TextureArrayHandle->GetChildHandle(0).ToSharedRef()).DisplayName(FText(LOCTEXT("Layout1_0", "Depth\n R8")));
-							TextureGroup.AddPropertyRow(TextureArrayHandle->GetChildHandle(1).ToSharedRef()).DisplayName(FText(LOCTEXT("Layout1_1", "Coverage\n R8")));
-							TextureGroup.AddPropertyRow(TextureArrayHandle->GetChildHandle(2).ToSharedRef()).DisplayName(FText(LOCTEXT("Layout1_2", "Tangent\n RGB8")));
-							TextureGroup.AddPropertyRow(TextureArrayHandle->GetChildHandle(3).ToSharedRef()).DisplayName(FText(LOCTEXT("Layout1_3", "Attributes\n RootUV | CoordU | Seed\n RGB8")));
-							TextureGroup.AddPropertyRow(TextureArrayHandle->GetChildHandle(4).ToSharedRef()).DisplayName(FText(LOCTEXT("Layout1_4", "Material\n Color | GroupID\n RBGA8")));
-							TextureGroup.AddPropertyRow(TextureArrayHandle->GetChildHandle(5).ToSharedRef()).DisplayName(FText(LOCTEXT("Layout1_5", "Auxiliary\n RGBA8")));
-						} break;
-						case EHairTextureLayout::Layout2:
-						{
-							check(TextureCount == GetHairTextureLayoutTextureCount(EHairTextureLayout::Layout2));
-							TextureGroup.AddPropertyRow(TextureArrayHandle->GetChildHandle(0).ToSharedRef()).DisplayName(FText(LOCTEXT("Layout2_0", "Tangent | CoordU\n RGBA8")));
-							TextureGroup.AddPropertyRow(TextureArrayHandle->GetChildHandle(2).ToSharedRef()).DisplayName(FText(LOCTEXT("Layout2_1", "Coverage | Depth | Seed\n RGB8")));
-							TextureGroup.AddPropertyRow(TextureArrayHandle->GetChildHandle(3).ToSharedRef()).DisplayName(FText(LOCTEXT("Layout2_2", "Color | Roughness\n RGBA8")));
-						} break;
-						case EHairTextureLayout::Layout3:
-						{
-							check(TextureCount == GetHairTextureLayoutTextureCount(EHairTextureLayout::Layout3));
-							TextureGroup.AddPropertyRow(TextureArrayHandle->GetChildHandle(0).ToSharedRef()).DisplayName(FText(LOCTEXT("Layout3_0", "Tangent\n RGB8")));
-							TextureGroup.AddPropertyRow(TextureArrayHandle->GetChildHandle(1).ToSharedRef()).DisplayName(FText(LOCTEXT("Layout3_1", "RooUV | CoordU | GroupID\n RGBA8")));
-							TextureGroup.AddPropertyRow(TextureArrayHandle->GetChildHandle(2).ToSharedRef()).DisplayName(FText(LOCTEXT("Layout3_2", "Coverage | Depth | Seed\n RGB8")));
-							TextureGroup.AddPropertyRow(TextureArrayHandle->GetChildHandle(3).ToSharedRef()).DisplayName(FText(LOCTEXT("Layout3_3", "Color | Roughness\n RGBA8")));
-						} break;
+						TextureGroup.AddPropertyRow(TextureArrayHandle->GetChildHandle(TextureIt).ToSharedRef()).DisplayName(FTextStringHelper::CreateFromBuffer(GetHairTextureLayoutTextureName(LayoutType, TextureIt, true /*bDetail*/)));
 					}
 				}
 			}
diff --git a/Engine/Shaders/Private/HairStrands/HairCardsAttributeCommon.ush b/Engine/Shaders/Private/HairStrands/HairCardsAttributeCommon.ush
--- a/Engine/Shaders/Private/HairStrands/HairCardsAttributeCommon.ush
+++ b/Engine/Shaders/Private/HairStrands/HairCardsAttributeCommon.ush
@@ -2,10 +2,6 @@
 
 #pragma once
 
-#if HAIR_CARD_MESH_FACTORY
-
-#include "/Engine/Private/HairStrands/HairStrandsVertexFactoryCommon.ush"
-
 ///////////////////////////////////////////////////////////////////////////////////////////////////
 
 // Vertex factory Flags
@@ -29,65 +25,73 @@
 #define EHairCardsVFAttribute_Coverage     8
 #define EHairCardsVFAttribute_Auxilary     9
 
-///////////////////////////////////////////////////////////////////////////////////////////////////
-// Helper functions for fetching attributes
-
-#define LAYOUT_VALUES(L0, L1, L2, L3) \
-	(HairCardsVF.LayoutIndex == 0 ? L0 :  \
-	(HairCardsVF.LayoutIndex == 1 ? L1 :  \
-	(HairCardsVF.LayoutIndex == 2 ? L2 : L3)))
+#define LAYOUT_VALUES(LayoutIdx, L0, L1, L2, L3) \
+	(LayoutIdx == 0 ? L0 :  \
+	(LayoutIdx == 1 ? L1 :  \
+	(LayoutIdx == 2 ? L2 : L3)))
 
-#define HAIRCARDSVF_USE_LEGACY_LAYOUT 1
-#if HAIRCARDSVF_USE_LEGACY_LAYOUT
-uint GetHairCardTextureIndex(uint In) 
+uint GetHairCardTextureIndex(uint InLayoutIndex, uint InAttribute) 
 { 
-	// See GroomAssetDetails.cpp, EHairTextureLayout::Layout0/1/2/3 for mapping details. (9 = not present/not available)
-	switch (In)
+	// See GroomAssetCards.cpp, EHairTextureLayout::Layout0/1/2/3 for mapping details. (9 = not present/not available)
+	switch (InAttribute)
 	{
-		case EHairCardsVFAttribute_Depth      : return LAYOUT_VALUES(0, 0, 1, 2);
-		case EHairCardsVFAttribute_Seed       : return LAYOUT_VALUES(3, 3, 1, 2);
-		case EHairCardsVFAttribute_CoordU     : return LAYOUT_VALUES(3, 3, 0, 1);
-		case EHairCardsVFAttribute_GroupIndex : return LAYOUT_VALUES(9, 4, 9, 1); // Only used for mesh, not for cards
-		case EHairCardsVFAttribute_Roughness  : return LAYOUT_VALUES(9, 9, 9, 3);
-		case EHairCardsVFAttribute_RootUV     : return LAYOUT_VALUES(9, 3, 9, 1);
-		case EHairCardsVFAttribute_Color      : return LAYOUT_VALUES(9, 4, 9, 3);
-		case EHairCardsVFAttribute_Tangent    : return LAYOUT_VALUES(2, 2, 0, 0);
-		case EHairCardsVFAttribute_Coverage   : return LAYOUT_VALUES(1, 1, 1, 2);
-		case EHairCardsVFAttribute_Auxilary   : return LAYOUT_VALUES(5, 5, 9, 9); 
+		case EHairCardsVFAttribute_Depth      : return LAYOUT_VALUES(InLayoutIndex, 0, 0, 1, 2);
+		case EHairCardsVFAttribute_Seed       : return LAYOUT_VALUES(InLayoutIndex, 3, 3, 1, 2);
+		case EHairCardsVFAttribute_CoordU     : return LAYOUT_VALUES(InLayoutIndex, 3, 3, 0, 1);
+		case EHairCardsVFAttribute_GroupIndex : return LAYOUT_VALUES(InLayoutIndex, 9, 4, 9, 1); // Only used for mesh, not for cards
+		case EHairCardsVFAttribute_Roughness  : return LAYOUT_VALUES(InLayoutIndex, 9, 9, 9, 3);
+		case EHairCardsVFAttribute_RootUV     : return LAYOUT_VALUES(InLayoutIndex, 9, 3, 9, 1);
+		case EHairCardsVFAttribute_Color      : return LAYOUT_VALUES(InLayoutIndex, 9, 4, 9, 3);
+		case EHairCardsVFAttribute_Tangent    : return LAYOUT_VALUES(InLayoutIndex, 2, 2, 0, 0);
+		case EHairCardsVFAttribute_Coverage   : return LAYOUT_VALUES(InLayoutIndex, 1, 1, 1, 2);
+		case EHairCardsVFAttribute_Auxilary   : return LAYOUT_VALUES(InLayoutIndex, 5, 5, 9, 9); 
 	}
 	return 0;
 }
 
-bool HasTextureBound(uint In)
+bool IsHairCardTextureIndexValid(uint InLayoutIndex, uint InAttribute)
 {
-	// HairCardsVF.Flags stores InverttUV on the first bit, texture bound after
-	const uint TextureIndex = GetHairCardTextureIndex(In);
-	return (TextureIndex < HairCardsVF.TextureCount) && (HairCardsVF.Flags & (1u << (TextureIndex + 1))) != 0;
+	return GetHairCardTextureIndex(InLayoutIndex, InAttribute) < HAIR_CARDS_MAX_TEXTURE_COUNT;
 }
 
-uint GetHairCardChannelIndex(uint In) 
+uint GetHairCardChannelIndex(uint InLayoutIndex, uint InAttribute) 
 { 
-	// See GroomAssetDetails.cpp, EHairTextureLayout::Layout0/1/2/3 for mapping details
-	switch (In)
+	// See GroomAssetCards.cpp, EHairTextureLayout::Layout0/1/2/3 for mapping details
+	switch (InAttribute)
 	{
-		case EHairCardsVFAttribute_Depth      : return LAYOUT_VALUES(0, 0, 1, 1);
-		case EHairCardsVFAttribute_Seed       : return LAYOUT_VALUES(3, 3, 2, 2);
-		case EHairCardsVFAttribute_CoordU     : return LAYOUT_VALUES(2, 2, 3, 1);
-		case EHairCardsVFAttribute_GroupIndex : return LAYOUT_VALUES(3, 3, 0, 2);
-		case EHairCardsVFAttribute_Roughness  : return LAYOUT_VALUES(3, 0, 3, 3);
-		case EHairCardsVFAttribute_RootUV     : return LAYOUT_VALUES(0, 0, 0, 0);
-		case EHairCardsVFAttribute_Color      : return LAYOUT_VALUES(0, 0, 0, 0);
-		case EHairCardsVFAttribute_Tangent    : return LAYOUT_VALUES(0, 0, 0, 0);
-		case EHairCardsVFAttribute_Coverage   : return LAYOUT_VALUES(0, 0, 0, 0);
-		case EHairCardsVFAttribute_Auxilary   : return LAYOUT_VALUES(0, 0, 0, 0); 
+		case EHairCardsVFAttribute_Depth      : return LAYOUT_VALUES(InLayoutIndex, 0, 0, 1, 1);
+		case EHairCardsVFAttribute_Seed       : return LAYOUT_VALUES(InLayoutIndex, 3, 3, 2, 2);
+		case EHairCardsVFAttribute_CoordU     : return LAYOUT_VALUES(InLayoutIndex, 2, 2, 3, 1);
+		case EHairCardsVFAttribute_GroupIndex : return LAYOUT_VALUES(InLayoutIndex, 3, 3, 0, 2);
+		case EHairCardsVFAttribute_Roughness  : return LAYOUT_VALUES(InLayoutIndex, 3, 0, 3, 3);
+		case EHairCardsVFAttribute_RootUV     : return LAYOUT_VALUES(InLayoutIndex, 0, 0, 0, 0);
+		case EHairCardsVFAttribute_Color      : return LAYOUT_VALUES(InLayoutIndex, 0, 0, 0, 0);
+		case EHairCardsVFAttribute_Tangent    : return LAYOUT_VALUES(InLayoutIndex, 0, 0, 0, 0);
+		case EHairCardsVFAttribute_Coverage   : return LAYOUT_VALUES(InLayoutIndex, 0, 0, 0, 0);
+		case EHairCardsVFAttribute_Auxilary   : return LAYOUT_VALUES(InLayoutIndex, 0, 0, 0, 0); 
 	}
 	return 0;
 }
-#else	
-uint GetHairCardTextureIndex(uint In)	{ return ((HairCards.AttributeTextureIndex) >> (In*3u) & 0x7u); }
-uint GetHairCardChannelIndex(uint In)	{ return ((HairCards.AttributeChannelIndex) >> (In*2u) & 0x3u); }
-#endif
-float2 GetAtlasUV(float2 In) 			{ return float2(In.x, (HairCardsVF.Flags & EHairCardsVFlags_InvertedUV) ? 1 - In.y : In.y); }
+
+#if HAIR_CARD_MESH_FACTORY
+
+#include "/Engine/Private/HairStrands/HairStrandsVertexFactoryCommon.ush"
+
+///////////////////////////////////////////////////////////////////////////////////////////////////
+// Helper functions for fetching attributes
+
+// If we need custom dynamic layout:
+// uint GetHairCardTextureIndex(uint In)	{ return ((HairCards.AttributeTextureIndex) >> (In*3u) & 0x7u); }
+// uint GetHairCardChannelIndex(uint In)	{ return ((HairCards.AttributeChannelIndex) >> (In*2u) & 0x3u); }
+
+float2 GetAtlasUV(float2 In) 				{ return float2(In.x, (HairCardsVF.Flags & EHairCardsVFlags_InvertedUV) ? 1 - In.y : In.y); }
+
+bool HasTextureBound(uint In)
+{
+	// HairCardsVF.Flags stores InverttUV on the first bit, texture bound after
+	const uint TextureIndex = GetHairCardTextureIndex(HairCardsVF.LayoutIndex, In);
+	return (TextureIndex < HairCardsVF.TextureCount) && (HairCardsVF.Flags & (1u << (TextureIndex + 1))) != 0;
+}
 
 // Generic attribute texture fetches
 float4 FetchAttributeTexture(uint InTextureIndex, float2 InAtlasUV, float InBias)
@@ -115,11 +119,11 @@ float4 FetchAttributeTexture(uint InTextureIndex, float2 InAtlasUV)
 }
 
 // Generic attribute fetch functions
-float  FetchAttribute1(uint In, float2 InAtlasUV, float InBias) { const float4 Value = FetchAttributeTexture(GetHairCardTextureIndex(In), InAtlasUV, InBias); return Value[GetHairCardChannelIndex(In)];                     }
-float  FetchAttribute1(uint In, float2 InAtlasUV)               { const float4 Value = FetchAttributeTexture(GetHairCardTextureIndex(In), InAtlasUV);         return Value[GetHairCardChannelIndex(In)];                     }
-float2 FetchAttribute2(uint In, float2 InAtlasUV)               { const float4 Value = FetchAttributeTexture(GetHairCardTextureIndex(In), InAtlasUV);         return GetHairCardChannelIndex(In) == 0 ? Value.xy : Value.zw; } // float2 alignemnt is forced onto .xy or .zw
-float3 FetchAttribute3(uint In, float2 InAtlasUV)               { const float4 Value = FetchAttributeTexture(GetHairCardTextureIndex(In), InAtlasUV);         return Value.xyz;                                              } // float3 alignemnt forced onto .xyz
-float4 FetchAttribute4(uint In, float2 InAtlasUV)               { const float4 Value = FetchAttributeTexture(GetHairCardTextureIndex(In), InAtlasUV);         return Value;                                                  }
+float  FetchAttribute1(uint In, float2 InAtlasUV, float InBias) { const float4 Value = FetchAttributeTexture(GetHairCardTextureIndex(HairCardsVF.LayoutIndex, In), InAtlasUV, InBias); return Value[GetHairCardChannelIndex(HairCardsVF.LayoutIndex, In)];                     }
+float  FetchAttribute1(uint In, float2 InAtlasUV)               { const float4 Value = FetchAttributeTexture(GetHairCardTextureIndex(HairCardsVF.LayoutIndex, In), InAtlasUV);         return Value[GetHairCardChannelIndex(HairCardsVF.LayoutIndex, In)];                     }
+float2 FetchAttribute2(uint In, float2 InAtlasUV)               { const float4 Value = FetchAttributeTexture(GetHairCardTextureIndex(HairCardsVF.LayoutIndex, In), InAtlasUV);         return GetHairCardChannelIndex(HairCardsVF.LayoutIndex, In) == 0 ? Value.xy : Value.zw; } // float2 alignemnt is forced onto .xy or .zw
+float3 FetchAttribute3(uint In, float2 InAtlasUV)               { const float4 Value = FetchAttributeTexture(GetHairCardTextureIndex(HairCardsVF.LayoutIndex, In), InAtlasUV);         return Value.xyz;                                                                       } // float3 alignemnt forced onto .xyz
+float4 FetchAttribute4(uint In, float2 InAtlasUV)               { const float4 Value = FetchAttributeTexture(GetHairCardTextureIndex(HairCardsVF.LayoutIndex, In), InAtlasUV);         return Value;                                                                           }
 
 // Special attribute fetch functions
 float  GetDepthOffset (float2 InAtlasUV)                        { float  Out = 0;         if (HasTextureBound(EHairCardsVFAttribute_Depth))		 { Out = FetchAttribute1(EHairCardsVFAttribute_Depth,     InAtlasUV);                          } return Out; }
diff --git a/Engine/Shaders/Private/HairStrands/HairStrandsTexturesGeneration.usf b/Engine/Shaders/Private/HairStrands/HairStrandsTexturesGeneration.usf
--- a/Engine/Shaders/Private/HairStrands/HairStrandsTexturesGeneration.usf
+++ b/Engine/Shaders/Private/HairStrands/HairStrandsTexturesGeneration.usf
@@ -11,6 +11,64 @@
 #include "../ShaderPrint.ush"
 #endif
 
+/////////////////////////////////////////////////////////////////////////////////////////////////////////////////////////////
+// Helpers
+
+#if HAIR_CARDS_MAX_TEXTURE_COUNT != 6
+#error Max cards texture count is != 6. This codes needs to be updated.
+#endif
+
+#define WriteValue1(InLayout, InAttribute, In)\
+{ \
+	if (IsHairCardTextureIndexValid(InLayout, InAttribute))\
+	{\
+		const uint TexIndex = GetHairCardTextureIndex(InLayout, InAttribute);\
+		const uint ChannelIndex = GetHairCardChannelIndex(InLayout, InAttribute);\
+		switch (TexIndex)\
+		{\
+			case 0: ChannelIndex == 0 ? OutTexture0.x = In : (ChannelIndex == 1 ? OutTexture0.y = In : (ChannelIndex == 2 ? OutTexture0.z = In : OutTexture0.w = In)); break;\
+			case 1: ChannelIndex == 0 ? OutTexture1.x = In : (ChannelIndex == 1 ? OutTexture1.y = In : (ChannelIndex == 2 ? OutTexture1.z = In : OutTexture1.w = In)); break;\
+			case 2: ChannelIndex == 0 ? OutTexture2.x = In : (ChannelIndex == 1 ? OutTexture2.y = In : (ChannelIndex == 2 ? OutTexture2.z = In : OutTexture2.w = In)); break;\
+			case 3: ChannelIndex == 0 ? OutTexture3.x = In : (ChannelIndex == 1 ? OutTexture3.y = In : (ChannelIndex == 2 ? OutTexture3.z = In : OutTexture3.w = In)); break;\
+			case 4: ChannelIndex == 0 ? OutTexture4.x = In : (ChannelIndex == 1 ? OutTexture4.y = In : (ChannelIndex == 2 ? OutTexture4.z = In : OutTexture4.w = In)); break;\
+			case 5: ChannelIndex == 0 ? OutTexture5.x = In : (ChannelIndex == 1 ? OutTexture5.y = In : (ChannelIndex == 2 ? OutTexture5.z = In : OutTexture5.w = In)); break;\
+		}\
+	}\
+}
+#define WriteValue2(InLayout, InAttribute, In)\
+{\
+	if (IsHairCardTextureIndexValid(InLayout, InAttribute))\
+	{\
+		const uint TexIndex = GetHairCardTextureIndex(InLayout, InAttribute);\
+		const uint ChannelIndex = GetHairCardChannelIndex(InLayout, InAttribute);\
+		switch (TexIndex)\
+		{\
+			case 0: ChannelIndex < 2 ? OutTexture0.xy = In : OutTexture0.zw = In; break;\
+			case 1: ChannelIndex < 2 ? OutTexture1.xy = In : OutTexture1.zw = In; break;\
+			case 2: ChannelIndex < 2 ? OutTexture2.xy = In : OutTexture2.zw = In; break;\
+			case 3: ChannelIndex < 2 ? OutTexture3.xy = In : OutTexture3.zw = In; break;\
+			case 4: ChannelIndex < 2 ? OutTexture4.xy = In : OutTexture4.zw = In; break;\
+			case 5: ChannelIndex < 2 ? OutTexture5.xy = In : OutTexture5.zw = In; break;\
+		}\
+	}\
+}
+#define WriteValue3(InLayout, InAttribute, In)\
+{\
+	if (IsHairCardTextureIndexValid(InLayout, InAttribute))\
+	{\
+	const uint TexIndex = GetHairCardTextureIndex(InLayout, InAttribute);\
+		switch (TexIndex)\
+		{\
+			case 0: OutTexture0.xyz = In; break;\
+			case 1: OutTexture1.xyz = In; break;\
+			case 2: OutTexture2.xyz = In; break;\
+			case 3: OutTexture3.xyz = In; break;\
+			case 4: OutTexture4.xyz = In; break;\
+			case 5: OutTexture5.xyz = In; break;\
+		}\
+	}\
+}
+
 /////////////////////////////////////////////////////////////////////////////////////////////////////////////////////////////
 
 #if SHADER_VERTEX
@@ -61,6 +119,7 @@ void MainVS(
 
 #include "HairStrandsVertexFactoryCommon.ush"
 #include "HairStrandsAttributeCommon.ush"
+#include "HairCardsAttributeCommon.ush"
 
 HAIR_STRANDS_INSTANCE_PARAMETERS(InVF)
 
@@ -85,6 +144,7 @@ float3 ComputeTangent(
 	return normalize(CP1.Position - CP0.Position);
 }
 
+uint LayoutIndex;
 int2 OutputResolution;
 uint VertexCount;
 uint PrimitiveCount;
@@ -108,6 +168,7 @@ uint Voxel_Data_MaxCount;
 
 StructuredBuffer<uint2> Voxel_OffsetAndCount;
 StructuredBuffer<uint2> Voxel_Data;
+RWTexture2DArray<float4> OutTexture;
 
 uint3 PositionToCoord(float3 P, float3 InMinAABB, float3 InMaxAABB, uint3 InResolution)
 {
@@ -342,12 +403,7 @@ void MainPS(
 	in float3  InTangent	: MESH_TANGENT,
 	in float3  InBitangent	: MESH_BITANGENT,
 	in float3  InNormal		: MESH_NORMAL,
-	out float  OutDepth		: SV_Target0,
-	out float4 OutTangent	: SV_Target1,
-	out float  OutCoverage	: SV_Target2,
-	out float4 OutAttribute : SV_Target3,
-	out float4 OutMaterial  : SV_Target4,
-	out uint   OutTriangleMask : SV_Target5,
+	out uint   OutTriangleMask : SV_Target0,
 	out float  OutDepthTest	: SV_Depth)
 {
 	const float3 T = normalize(InTangent);
@@ -357,12 +413,22 @@ void MainPS(
 	const float3 PixelAxisX = ddx_fine(InPosition);
 	const float3 PixelAxisY = ddy_fine(InPosition);
 
-	OutDepth = max(0, MaxDistance);
-	OutCoverage = 0;
-	OutTangent = 0;
-	OutAttribute = 0;
-	OutMaterial = 0;
-	OutDepthTest = 1;
+	float4 OutTexture0 = 0;
+	float4 OutTexture1	= 0;
+	float4 OutTexture2	= 0;
+	float4 OutTexture3	= 0;
+	float4 OutTexture4	= 0;
+	float4 OutTexture5	= 0;
+
+	float  OutDepth = max(0, MaxDistance);
+	float  OutSeed = 0;
+	float  OutCoordU = 0;
+	float  OutGroupIndex = 0;
+	float  OutRoughness = 0;
+	float2 OutRootUV = 0;
+	float3 OutColor = 0;
+	float3 OutTangent = 0;
+	float  OutCoverage = 0;
 
 	float3 Tangent_ObjectSpace = 0;
 	const int2 PixelCoord = int2(SvPosition.xy);
@@ -414,10 +480,11 @@ void MainPS(
 		{
 			if (SampleDepth < OutDepth)
 			{
-				OutAttribute.xy = SampleRootUV;
-				OutAttribute.z = SampleU;
-				OutAttribute.w = SampleSeed;
-				OutMaterial.xyz= SampleMaterial.xyz;
+				OutRootUV    = SampleRootUV;
+				OutCoordU    = SampleU;
+				OutSeed      = SampleSeed;
+				OutColor     = SampleMaterial.xyz;
+				OutRoughness = SampleMaterial.w;
 			}
 
 			OutDepth = min(OutDepth, SampleDepth);
@@ -436,22 +503,42 @@ void MainPS(
 			dot(Tangent_ObjectSpace, N));
 
 		OutCoverage /= SampleCount;
-		OutTangent.xyz = (Tangent_FrameSpace + 1) * 0.5f;
-		OutTangent.w = 1;
-		OutMaterial.w = InVF_GroupIndex / 255.f;
+		OutTangent = (Tangent_FrameSpace + 1) * 0.5f;
+		OutGroupIndex = InVF_GroupIndex / 255.f;
 	}
 	else
 	{
 		OutDepth = MaxDistance;
-		OutCoverage = 0;
+		OutSeed = 0;
+		OutCoordU = 0;
+		OutGroupIndex = 0;
+		OutRoughness = 0;
+		OutRootUV = 0;
+		OutColor = 0;
 		OutTangent = 0;
-		OutAttribute = 0;
-		OutMaterial = 0;
+		OutCoverage = 0;
 	}
 
 	OutDepth = saturate(OutDepth / MaxDistance);
 	OutDepthTest = OutDepth;
 	OutTriangleMask = 1;
+
+	WriteValue1(LayoutIndex, EHairCardsVFAttribute_Depth,      OutDepth);
+	WriteValue1(LayoutIndex, EHairCardsVFAttribute_Seed,       OutSeed);
+	WriteValue1(LayoutIndex, EHairCardsVFAttribute_CoordU,     OutCoordU);
+	WriteValue1(LayoutIndex, EHairCardsVFAttribute_GroupIndex, OutGroupIndex);
+	WriteValue1(LayoutIndex, EHairCardsVFAttribute_Roughness,  OutRoughness);
+	WriteValue2(LayoutIndex, EHairCardsVFAttribute_RootUV,     OutRootUV);
+	WriteValue3(LayoutIndex, EHairCardsVFAttribute_Color,      OutColor);
+	WriteValue3(LayoutIndex, EHairCardsVFAttribute_Tangent,    OutTangent);
+	WriteValue1(LayoutIndex, EHairCardsVFAttribute_Coverage,   OutCoverage);
+
+	OutTexture[uint3(PixelCoord, 0)] = OutTexture0;
+	OutTexture[uint3(PixelCoord, 1)] = OutTexture1;
+	OutTexture[uint3(PixelCoord, 2)] = OutTexture2;
+	OutTexture[uint3(PixelCoord, 3)] = OutTexture3;
+	OutTexture[uint3(PixelCoord, 4)] = OutTexture4;
+	OutTexture[uint3(PixelCoord, 5)] = OutTexture5;
 }
 
 #endif // SHADER_PIXEL
@@ -460,62 +547,61 @@ void MainPS(
 
 #if SHADER_TEXTURE_DILATION
 
-uint PackDistance(float x)							{ return x * 1000.f; }
-float UnpackDistance(float x, inout bool bIsValid)	{ bIsValid = x != 0; return float(x) / 1000.f; }
+#include "HairCardsAttributeCommon.ush"
 
+uint LayoutIndex;
+uint bClearPass;
+uint TextureCount;
 int2 Resolution;
 int ClearDistance;
-RWTexture2D<uint> TriangleMaskTexture;
-
-Texture2D<float>  Source_DepthTexture;
-Texture2D<float>  Source_CoverageTexture;
-Texture2D<float4> Source_TangentTexture;
-Texture2D<float4> Source_AttributeTexture;
-Texture2D<float4> Source_MaterialTexture;
-
-RWTexture2D<float>  Target_DepthTexture;
-RWTexture2D<float>  Target_CoverageTexture;
-RWTexture2D<float4> Target_TangentTexture;
-RWTexture2D<float4> Target_AttributeTexture;
-RWTexture2D<float4> Target_MaterialTexture;
+Texture2D<float> DepthTestTexture;
+Texture2D<uint> SourceTriangleMaskTexture;
+RWTexture2D<uint> TargetTriangleMaskTexture;
+Texture2DArray<float4> SourceTexture;
+RWTexture2DArray<float4> TargetTexture;
 
 // Dilate the source texture (Depth/Coverage/Tanget/Attribute), using a distance to triangle edge metric. 
 // Original valid pixel are marked into the triangleMasktexture during the texture generation. Then 
 // successive dilation pass are run. At each pass the triangleMaskTexure stored the distance to triangle
 // edge in order to propagate correctly the pixel edge information.
-[numthreads(8, 4, 1)]
+[numthreads(8, 8, 1)]
 void MainCS(int2 DispatchThreadId : SV_DispatchThreadID)
 {
-	if (DispatchThreadId.x >= Resolution.x || DispatchThreadId.y >= Resolution.y)
-		return;
-
-	const float MaxDistance = 1000000.f;
-
 	const uint2 PixelCoord = DispatchThreadId.xy;
-	const uint  PackedDistance = TriangleMaskTexture.Load(PixelCoord);
-
-	bool bIsValid = false;
-	const float Distance = UnpackDistance(PackedDistance, bIsValid);
+	if (any(PixelCoord >= Resolution))
+	{
+		return;
+	}
 
-	if (bIsValid)
+	const bool bIsValid = SourceTriangleMaskTexture.Load(uint3(PixelCoord,0)) > 0;
+	if (bClearPass || bIsValid)
 	{
 		// Fast path: the pixel is valid, we copy value
-		Target_DepthTexture[PixelCoord]		= Source_DepthTexture.Load(uint3(PixelCoord,0));
-		Target_CoverageTexture[PixelCoord]	= Source_CoverageTexture.Load(uint3(PixelCoord, 0));
-		Target_TangentTexture[PixelCoord]	= Source_TangentTexture.Load(uint3(PixelCoord, 0));
-		Target_AttributeTexture[PixelCoord] = Source_AttributeTexture.Load(uint3(PixelCoord, 0));
-		Target_MaterialTexture[PixelCoord]  = Source_MaterialTexture.Load(uint3(PixelCoord, 0));
-		return;
+		for (uint TextureIt = 0; TextureIt < TextureCount; ++TextureIt)
+		{
+			TargetTexture[uint3(PixelCoord, TextureIt)] = SourceTexture[uint3(PixelCoord, TextureIt)];
+		}
+
+		if (bClearPass)
+		{
+			TargetTriangleMaskTexture[PixelCoord] = DepthTestTexture.Load(uint3(PixelCoord,0)) < 1.f ? 1u : 0u;
+		}
+		else
+		{
+			TargetTriangleMaskTexture[PixelCoord] = SourceTriangleMaskTexture.Load(uint3(PixelCoord,0));
+		}
 	}
 	else
 	{
-		float Target_Depth = 0;
-		float Target_Coverage = 0;
-		float4 Target_Tangent = 0;
-		float4 Target_Attribute = 0;
-		float4 Target_Material = 0;
+		const uint CoverageTextureIndex = GetHairCardTextureIndex(LayoutIndex, EHairCardsVFAttribute_Coverage);
+		const uint CoverageChannelIndex = GetHairCardChannelIndex(LayoutIndex, EHairCardsVFAttribute_Coverage);
 
 		// Find the closest neighbor from triangle, amoung valid neighbord
+		float4 Values[HAIR_CARDS_MAX_TEXTURE_COUNT];
+		for (uint TextureIt = 0; TextureIt < HAIR_CARDS_MAX_TEXTURE_COUNT; ++TextureIt)
+		{
+			Values[TextureIt] = 0;
+		}
 		bool bIsDilatedPixelValid = false;
 		float MinDistance = 9999999;
 		for (int y = -1; y <= 1; ++y)
@@ -524,20 +610,15 @@ void MainCS(int2 DispatchThreadId : SV_DispatchThreadID)
 			if (x == 0 && y == 0)
 				continue;
 
-			const int2 NPixelOffset = int2(x, y);
-			const int2 NPixelCoord = clamp(PixelCoord + NPixelOffset, 0, Resolution);
+			const uint2 NPixelCoord = clamp(int2(PixelCoord) + int2(x, y), int2(0,0), int2(Resolution-1));
 
-			bool bNIsValid = false;
-			const uint NPackedDistance = TriangleMaskTexture.Load(uint2(NPixelCoord));
-			const float NDistance = UnpackDistance(NPackedDistance, bNIsValid);
-			if (bNIsValid && NDistance <= MinDistance)
+			const bool bNIsValid = SourceTriangleMaskTexture.Load(uint3(NPixelCoord,0)) > 0;
+			if (bNIsValid)
 			{
-				Target_Depth		= Source_DepthTexture.Load(uint3(NPixelCoord,0));
-				Target_Coverage		= Source_CoverageTexture.Load(uint3(NPixelCoord, 0));
-				Target_Tangent		= Source_TangentTexture.Load(uint3(NPixelCoord, 0));
-				Target_Attribute	= Source_AttributeTexture.Load(uint3(NPixelCoord, 0));
-				Target_Material		= Source_MaterialTexture.Load(uint3(NPixelCoord, 0));
-				MinDistance			= NDistance;
+				for (uint TextureIt = 0; TextureIt < TextureCount; ++TextureIt)
+				{
+					Values[TextureIt] = SourceTexture[uint3(NPixelCoord, TextureIt)];
+				}
 				bIsDilatedPixelValid= true;
 			}
 		}
@@ -545,13 +626,23 @@ void MainCS(int2 DispatchThreadId : SV_DispatchThreadID)
 		// If a valid neighbord has been bound, update pixel value
 		if (bIsDilatedPixelValid)
 		{
-			Target_DepthTexture[PixelCoord]		= Target_Depth;
-			Target_CoverageTexture[PixelCoord]	= Target_Coverage;
-			Target_TangentTexture[PixelCoord]	= Target_Tangent;
-			Target_AttributeTexture[PixelCoord] = Target_Attribute;
-			Target_MaterialTexture[PixelCoord]  = Target_Material;
-
-			TriangleMaskTexture[PixelCoord] = PackDistance(MinDistance);
+			for (uint TextureIt = 0; TextureIt < TextureCount; ++TextureIt)
+			{
+				// Write all attributes but coverage, as coverage should remain untouched.
+				float4 WriteOut = Values[TextureIt];
+				if (TextureIt == CoverageTextureIndex)
+				{
+					switch (CoverageChannelIndex)
+					{
+						case 0 : WriteOut.x = 0.f; break;
+						case 1 : WriteOut.y = 0.f; break;
+						case 2 : WriteOut.z = 0.f; break;
+						case 3 : WriteOut.z = 0.f; break;
+					}
+				}
+				TargetTexture[uint3(PixelCoord, TextureIt)] = WriteOut;
+			}
+			TargetTriangleMaskTexture[PixelCoord] = 1.f;
 		}
 	}
 }
@@ -580,4 +671,66 @@ void MainCS(int2 DispatchThreadId : SV_DispatchThreadID)
 		OutPositions[VertexIndex] = float4(CP.Position, CP.Type);
 	}
 }
-#endif // SHADER_READ_POSITIONS
\ No newline at end of file
+#endif // SHADER_READ_POSITIONS
+
+/////////////////////////////////////////////////////////////////////////////////////////////////////////////////////////////
+
+#if SHADER_TEXTURE_CLEAR
+
+#include "HairCardsAttributeCommon.ush"
+int2 Resolution;
+uint LayoutIndex;
+RWTexture2DArray<float4> OutTexture;
+
+[numthreads(8, 8, 1)]
+void MainCS(int2 DispatchThreadId : SV_DispatchThreadID)
+{
+	int2 PixelCoord = DispatchThreadId.xy;
+	if (all(PixelCoord < Resolution))
+	{	
+		float4 OutTexture0 = 0;
+		float4 OutTexture1 = 0;
+		float4 OutTexture2 = 0;
+		float4 OutTexture3 = 0;
+		float4 OutTexture4 = 0;
+		float4 OutTexture5 = 0;
+
+		WriteValue1(LayoutIndex, EHairCardsVFAttribute_Depth,      1.f);
+		WriteValue1(LayoutIndex, EHairCardsVFAttribute_Seed,       0.f);
+		WriteValue1(LayoutIndex, EHairCardsVFAttribute_CoordU,     0.f);
+		WriteValue1(LayoutIndex, EHairCardsVFAttribute_GroupIndex, 0.f);
+		WriteValue1(LayoutIndex, EHairCardsVFAttribute_Roughness,  0.f);
+		WriteValue2(LayoutIndex, EHairCardsVFAttribute_RootUV,     0.f);
+		WriteValue3(LayoutIndex, EHairCardsVFAttribute_Color,      0.f);
+		WriteValue3(LayoutIndex, EHairCardsVFAttribute_Tangent,    0.f);
+		WriteValue1(LayoutIndex, EHairCardsVFAttribute_Coverage,   0.f);
+
+		OutTexture[uint3(PixelCoord, 0)] = OutTexture0;
+		OutTexture[uint3(PixelCoord, 1)] = OutTexture1;
+		OutTexture[uint3(PixelCoord, 2)] = OutTexture2;
+		OutTexture[uint3(PixelCoord, 3)] = OutTexture3;
+		OutTexture[uint3(PixelCoord, 4)] = OutTexture4;
+		OutTexture[uint3(PixelCoord, 5)] = OutTexture5;
+	}
+}
+#endif // SHADER_TEXTURE_CLEAR
+
+/////////////////////////////////////////////////////////////////////////////////////////////////////////////////////////////
+
+#if SHADER_TEXTURE_COPY
+
+uint TextureIndex;
+int2 Resolution;
+Texture2DArray<float4> InTexture;
+RWTexture2D<float4> OutTexture;
+
+[numthreads(8, 8, 1)]
+void MainCS(int2 DispatchThreadId : SV_DispatchThreadID)
+{
+	int2 PixelCoord = DispatchThreadId.xy;
+	if (all(PixelCoord < Resolution))
+	{	
+		OutTexture[PixelCoord] = InTexture[uint3(PixelCoord, TextureIndex)];
+	}
+}
+#endif // SHADER_TEXTURE_COPY
\ No newline at end of file
diff --git a/Engine/Shaders/Shared/HairStrandsDefinitions.h b/Engine/Shaders/Shared/HairStrandsDefinitions.h
--- a/Engine/Shaders/Shared/HairStrandsDefinitions.h
+++ b/Engine/Shaders/Shared/HairStrandsDefinitions.h
@@ -96,4 +96,6 @@
 #define HAIR_POINT_LOD_COUNT_PER_UINT (32u/4u)
 #define HAIR_POINT_LOD_COUNT_PER_UINT_DIV_AS_SHIFT 3
 
-#define HAIR_RBF_ENTRY_COUNT(InSampleCount)	((InSampleCount)+4u)
\ No newline at end of file
+#define HAIR_RBF_ENTRY_COUNT(InSampleCount)	((InSampleCount)+4u)
+
+#define HAIR_CARDS_MAX_TEXTURE_COUNT 6
\ No newline at end of file
