Change groom binding building to support properly platform specific settings for cooked data.

diff --git a/Engine/Plugins/Runtime/HairStrands/Source/HairStrandsCore/Private/GroomBindingAsset.cpp b/Engine/Plugins/Runtime/HairStrands/Source/HairStrandsCore/Private/GroomBindingAsset.cpp
--- a/Engine/Plugins/Runtime/HairStrands/Source/HairStrandsCore/Private/GroomBindingAsset.cpp
+++ b/Engine/Plugins/Runtime/HairStrands/Source/HairStrandsCore/Private/GroomBindingAsset.cpp
@@ -12,6 +12,8 @@
 #include "UObject/AnimObjectVersion.h"
 #include "UObject/ObjectSaveContext.h"
 #include "UObject/DevObjectVersion.h"
+#include "Misc/CoreMisc.h"
+#include "Interfaces/ITargetPlatformManagerModule.h"
 
 #include UE_INLINE_GENERATED_CPP_BY_NAME(GroomBindingAsset)
 
@@ -41,24 +43,13 @@ static void InternalSerializeGuides(FArchive& Ar, UObject* Owner, FHairStrandsRo
 	}
 }
 
-static void InternalSerializeStrands(FArchive& Ar, UObject* Owner, FHairStrandsRootBulkData& Data, uint32 Flags, bool bHeader, bool bData, bool* bOutHasDataInCache = nullptr)
+static void InternalSerializeStrands(FArchive& Ar, UObject* Owner, FHairStrandsRootBulkData& Data, uint32 Flags, bool bHeader, bool bData)
 {
 	Ar.UsingCustomVersion(FAnimObjectVersion::GUID);
 
 	const bool bStripped = (Flags & UGroomAsset::CDSF_StrandsStripped);
 	if (!bStripped)
 	{
-		// When cooking data, force loading of *all* bulk data prior to saving them
-		// Note: bFillBulkdata is true for filling in the bulkdata container prior to serialization. This also forces the resources loading 
-		// from the 'start' (i.e., without offset)
-		if (Ar.IsCooking() && Ar.IsSaving())
-		{
-			for (int32 LODIndex = 0, LODCount = Data.GetLODCount(); LODIndex < LODCount; ++LODIndex)
-			{
-				FHairStreamingRequest R; R.Request(HAIR_MAX_NUM_CURVE_PER_GROUP, HAIR_MAX_NUM_POINT_PER_GROUP, LODIndex, Data, true /*bWait*/, true /*bFillBulkdata*/, true /*bWarmCache*/, Owner->GetFName());
-			}
-		}
-
 		if (bHeader){ Data.SerializeHeader(Ar, Owner); }
 		if (bData)
 		{
@@ -67,21 +58,6 @@ static void InternalSerializeStrands(FArchive& Ar, UObject* Owner, FHairStrandsR
 				Data.SerializeData(Ar, Owner, LODIndex);
 			}
 		}
-
-		// Pre-warm DDC cache
-		#if WITH_EDITORONLY_DATA
-		const bool bPreWarmCache = Ar.IsLoading() && bHeader && !bData;
-		if (bPreWarmCache)
-		{
-			bool bHasDataInCache = true;
-			for (int32 LODIndex = 0, LODCount = Data.GetLODCount(); LODIndex < LODCount; ++LODIndex)
-			{
-				FHairStreamingRequest R; bHasDataInCache &= R.WarmCache(HAIR_MAX_NUM_CURVE_PER_GROUP, HAIR_MAX_NUM_POINT_PER_GROUP, LODIndex, Data);
-			}
-
-			if (bOutHasDataInCache) { *bOutHasDataInCache = bHasDataInCache; }
-		}
-		#endif
 	}
 }
 
@@ -103,7 +79,7 @@ static void InternalSerializeCards(FArchive& Ar, UObject* Owner, TArray<FHairStr
 	}
 }
 
-static void InternalSerializePlatformData(FArchive& Ar, UObject* Owner, UGroomBindingAsset::FHairGroupPlatformData& GroupData, uint32 Flags, bool bHeader, bool bData, bool* bOutHasDataInCache = nullptr)
+static void InternalSerializePlatformData(FArchive& Ar, UObject* Owner, UGroomBindingAsset::FHairGroupPlatformData& GroupData, uint32 Flags, bool bHeader, bool bData)
 {
 	Ar.UsingCustomVersion(FAnimObjectVersion::GUID);
 
@@ -111,7 +87,7 @@ static void InternalSerializePlatformData(FArchive& Ar, UObject* Owner, UGroomBi
 	InternalSerializeGuides(Ar, Owner, GroupData.SimRootBulkData);
 
 	// Strands
-	InternalSerializeStrands(Ar, Owner, GroupData.RenRootBulkData, Flags, bHeader, bData, bOutHasDataInCache);
+	InternalSerializeStrands(Ar, Owner, GroupData.RenRootBulkData, Flags, bHeader, bData);
 
 	// Cards
 	InternalSerializeCards(Ar, Owner, GroupData.CardsRootBulkData);
@@ -147,12 +123,28 @@ void UGroomBindingAsset::Serialize(FArchive& Ar)
 	Super::Serialize(Ar);
 	Ar.UsingCustomVersion(FAnimObjectVersion::GUID);
 #if WITH_EDITOR
+	// When using editor:
+	// * The header are loaded in CacheDerivedData(), and the data are streamed from DDC
+	// * When cooking, we write out data from the cached cooked platform data
 	if (Ar.IsCooking())
-#endif
+	{
+		if (TArray<UGroomBindingAsset::FHairGroupPlatformData>* CookedDatas = GetCachedCookedPlatformData(Ar.CookingTarget()))
+		{
+			InternalSerializePlatformDatas(Ar, this /*Owner*/, *CookedDatas, Flags);
+			bIsValid = true;
+		}
+		else
+		{
+			UE_LOG(LogHairStrands, Error, TEXT("[Groom] The binding asset (%s) has missing cooked platform data."), *GetName());
+		}
+	}
+#else
+	// Always loaded dara from the archive when not using the editor
 	{
 		InternalSerializePlatformDatas(Ar, this, GetHairGroupsPlatformData(), Flags);
 		bIsValid = true;
 	}
+#endif
 }
 
 void UGroomBindingAsset::InitResource()
@@ -363,10 +355,14 @@ void UGroomBindingAsset::PostLoad()
 	#endif
 	}
 
+	// * When running with the editor, InitResource is called in CacheDerivedDatas
+	// * When running without the editor, InitResource is explicitely called here
+#if !WITH_EDITOR
 	if (!IsTemplate() && IsValid())
 	{
 		InitResource();
 	}
+#endif
 }
 
 void UGroomBindingAsset::PreSave(const class ITargetPlatform* TargetPlatform)
@@ -378,12 +374,6 @@ void UGroomBindingAsset::PreSave(const class ITargetPlatform* TargetPlatform)
 
 void UGroomBindingAsset::PreSave(FObjectPreSaveContext ObjectSaveContext)
 {
-#if WITH_EDITOR
-	while (QueryStatus == EQueryStatus::Submitted)
-	{
-		FPlatformProcess::Sleep(1);
-	}
-#endif
 	Super::PreSave(ObjectSaveContext);
 #if WITH_EDITOR
 	OnGroomBindingAssetChanged.Broadcast();
@@ -773,7 +763,7 @@ namespace GroomBindingDerivedDataCacheUtils
 	}
 }
 
-static FString BuildDerivedDataKeySuffix(const UGroomBindingAsset& BindingAsset)
+static FString BuildDerivedDataKeySuffix(const UGroomBindingAsset& BindingAsset, const ITargetPlatform* TargetPlatform)
 {
 	FString BindingType;
 	FString SourceKey;
@@ -782,8 +772,8 @@ static FString BuildDerivedDataKeySuffix(const UGroomBindingAsset& BindingAsset)
 	if (BindingAsset.GetGroomBindingType() == EGroomBindingMeshType::SkeletalMesh)
 	{
 		// Binding type is implicitly SkeletalMesh so keep BindingType empty to prevent triggering rebuild of old binding for nothing
-		SourceKey = BindingAsset.GetSourceSkeletalMesh() ? BindingAsset.GetSourceSkeletalMesh()->GetDerivedDataKey() : FString();
-		TargetKey = BindingAsset.GetTargetSkeletalMesh() ? BindingAsset.GetTargetSkeletalMesh()->GetDerivedDataKey() : FString();
+		SourceKey = BindingAsset.GetSourceSkeletalMesh() ? BindingAsset.GetSourceSkeletalMesh()->BuildDerivedDataKey(TargetPlatform) : FString();
+		TargetKey = BindingAsset.GetTargetSkeletalMesh() ? BindingAsset.GetTargetSkeletalMesh()->BuildDerivedDataKey(TargetPlatform) : FString();
 	}
 	else
 	{
@@ -805,6 +795,20 @@ static FString BuildDerivedDataKeySuffix(const UGroomBindingAsset& BindingAsset)
 	return KeySuffix;
 }
 
+static FString BuildDerivedDataKeyGroup(const UGroomBindingAsset& BindingAsset, const ITargetPlatform* TargetPlatform, uint32 InGroupIndex)
+{
+	const FString DeriveDataKeySuffix = BuildDerivedDataKeySuffix(BindingAsset, TargetPlatform);
+	return GroomBindingDerivedDataCacheUtils::BuildGroomBindingDerivedDataKey(DeriveDataKeySuffix + FString(TEXT("_Group")) + FString::FromInt(InGroupIndex));
+}
+
+static FString BuildDerivedDataKeyGroup(const FString& InDeriveDataKeySuffix , uint32 InGroupIndex)
+{
+	return GroomBindingDerivedDataCacheUtils::BuildGroomBindingDerivedDataKey(InDeriveDataKeySuffix + FString(TEXT("_Group")) + FString::FromInt(InGroupIndex));
+}
+
+static TArray<FString> GetGroupDerivedDataKeys(const UGroomBindingAsset* In, const ITargetPlatform* TargetPlatform);
+static void CacheDerivedDatas(UGroomBindingAsset* In, const uint32 InGroupIndex, const FString& DerivedDataKey, bool& bOutValid, const ITargetPlatform* TargetPlatform, UGroomBindingAsset::FHairGroupPlatformData& OutPlatformData);
+
 void UGroomBindingAsset::CacheDerivedDatas()
 {
 	if (!GetGroom() || !GetGroom()->IsValid())
@@ -818,21 +822,61 @@ void UGroomBindingAsset::CacheDerivedDatas()
 	CachedDerivedDataKey.SetNum(GroupCount);
 	GetGroupInfos().SetNum(GroupCount);
 
-	// 2. Prepare main cache key
-	const FString KeySuffix = BuildDerivedDataKeySuffix(*this);
+	// 2.1 Prepare main cache key
+	// Cache derived data for the running platform.
+	ITargetPlatform* RunningPlatform = GetTargetPlatformManagerRef().GetRunningTargetPlatform();
+	check(RunningPlatform);
+	const FString KeySuffix = BuildDerivedDataKeySuffix(*this, RunningPlatform);
+
+	// 2.2 Build the key for each group and check if any group needs to be rebuilt
+	const TArray<FString> GroupDerivedDataKeys = GetGroupDerivedDataKeys(this, RunningPlatform);
+	const bool bAnyGroupNeedRebuild = GroupDerivedDataKeys != CachedDerivedDataKey;
 
 	// 3. Build or retrieve from cache, binding data for each group
 	bIsValid = true;
 	bool bReloadResource = false;
-	for (uint32 GroupIndex = 0; GroupIndex < GroupCount; ++GroupIndex)
+	if (bAnyGroupNeedRebuild)
 	{
-		bool bGroupReloadResource = false;
-		bool bGroupValid = true;
+		FGroomComponentRecreateRenderStateContext RecreateRenderContext(GetGroom());
 
-		CacheDerivedDatas(GroupIndex, KeySuffix, bGroupValid, bGroupReloadResource);
+		for (uint32 GroupIndex = 0; GroupIndex < GroupCount; ++GroupIndex)
+		{
+			// 1. Build data
+			bool bGroupValid = true;
+			bool bGroupReloadResource = false;
+			if (GroupDerivedDataKeys[GroupIndex] != CachedDerivedDataKey[GroupIndex])
+			{
+				::CacheDerivedDatas(this, GroupIndex, GroupDerivedDataKeys[GroupIndex], bGroupValid, RunningPlatform, GetHairGroupsPlatformData()[GroupIndex]);
+	
+				if (bGroupValid)
+				{
+					bGroupReloadResource = true;
+					CachedDerivedDataKey[GroupIndex] = GroupDerivedDataKeys[GroupIndex] ;
+				}
+				else
+				{
+					UE_LOG(LogHairStrands, Error, TEXT("[Groom] The binding asset (%s) couldn't be built. This binding asset won't be used."), *GetName());
+				}
+			}
+	
+			// 2. Release existing resources data
+			if (bGroupReloadResource)
+			{
+				UGroomBindingAsset::FHairGroupResources& OutHairGroupResources = GetHairGroupResources();
+				if (OutHairGroupResources.Num() > 0)
+				{
+					for (UGroomBindingAsset::FHairGroupResource& GroupResources : OutHairGroupResources)
+					{
+						AddHairGroupResourcesToDelete(GroupResources);
+					}
+					OutHairGroupResources.Empty();
+				}
+				check(OutHairGroupResources.Num() == 0);
+			}
 
-		bIsValid = bIsValid && bGroupValid;
-		bReloadResource = bReloadResource || bGroupReloadResource;
+			bIsValid = bIsValid && bGroupValid;
+			bReloadResource = bReloadResource || bGroupReloadResource;
+		}
 	}
 
 	// 4. Reload resources if needed
@@ -845,19 +889,14 @@ void UGroomBindingAsset::CacheDerivedDatas()
 	UpdateGroomBindingAssetInfos(this);
 }
 
-void UGroomBindingAsset::CacheDerivedDatas(uint32 InGroupIndex, const FString KeySuffix, bool& bOutValid, bool& bOutReloadResource)
+static void CacheDerivedDatas(UGroomBindingAsset* In, const uint32 InGroupIndex, const FString& DerivedDataKey, bool& bOutValid, const ITargetPlatform* TargetPlatform, UGroomBindingAsset::FHairGroupPlatformData& OutPlatformData)
 {
-	const FString DerivedDataKey = GroomBindingDerivedDataCacheUtils::BuildGroomBindingDerivedDataKey(KeySuffix + FString(TEXT("_Group")) + FString::FromInt(InGroupIndex));
-
-	bOutValid = true;
-	bOutReloadResource = false;
-	if (DerivedDataKey != CachedDerivedDataKey[InGroupIndex])
 	{
 		bOutValid = false;
 		using namespace UE::DerivedData;
 
 		const FCacheKey HeaderKey = ConvertLegacyCacheKey(DerivedDataKey + FString(TEXT("_Header")));
-		const FSharedString Name = MakeStringView(GetPathName());
+		const FSharedString Name = MakeStringView(In->GetPathName());
 		FSharedBuffer Data;
 		{
 			FRequestOwner Owner(EPriority::Blocking);
@@ -868,45 +907,61 @@ void UGroomBindingAsset::CacheDerivedDatas(uint32 InGroupIndex, const FString Ke
 			Owner.Wait();
 		}
 
-		UGroomBindingAsset::FHairGroupPlatformData& PlatformData = GetHairGroupsPlatformData()[InGroupIndex];
-
 		// Populate key/name for streaming data request
 		auto FillDrivedDataKey = [&DerivedDataKey, &Name](UGroomBindingAsset::FHairGroupPlatformData& In)
 		{
 			In.RenRootBulkData.DerivedDataKey = DerivedDataKey + FString(TEXT("_Data"));
 		};
 
-		FGroomComponentRecreateRenderStateContext RecreateRenderContext(GetGroom());
-
-		bool bSuccess = false;
+		bool bHasDataInCache = false;
 		if (Data)
 		{
-			UE_CLOG(IsHairStrandsDDCLogEnable(), LogHairStrands, Log, TEXT("[GroomBinding/DDC] Found (GroomBinding:%s)."), *GetName());
+			UE_CLOG(IsHairStrandsDDCLogEnable(), LogHairStrands, Log, TEXT("[GroomBinding/DDC] Found (GroomBinding:%s)."), *In->GetName());
 
-			FillDrivedDataKey(PlatformData);
+			FillDrivedDataKey(OutPlatformData);
 
 			// Header
 			FMemoryReaderView Ar(Data, /*bIsPersistent*/ true);
-			InternalSerializePlatformData(Ar, this, PlatformData, 0 /*Flags*/, true /*bHeader*/, false /*bData*/, &bSuccess);
+			InternalSerializePlatformData(Ar, In, OutPlatformData, 0 /*Flags*/, true /*bHeader*/, false /*bData*/);
+			bHasDataInCache = true;
+
+			// Verify that all strands data are correctly cached into the DDC
+			{
+				for (int32 LODIndex = 0, LODCount = OutPlatformData.RenRootBulkData.GetLODCount(); LODIndex < LODCount; ++LODIndex)
+				{
+					FHairStreamingRequest R; bHasDataInCache &= R.WarmCache(HAIR_MAX_NUM_CURVE_PER_GROUP, HAIR_MAX_NUM_POINT_PER_GROUP, LODIndex, OutPlatformData.RenRootBulkData);
+				}
+			}
 
 			bOutValid = true;
 		}
-		if (!bSuccess)
+		if (!bHasDataInCache)
 		{
-			UE_CLOG(IsHairStrandsDDCLogEnable(), LogHairStrands, Log, TEXT("[GroomBinding/DDC] Not found (GroomBinding:%s)."), *GetName());
+			UE_CLOG(IsHairStrandsDDCLogEnable(), LogHairStrands, Log, TEXT("[GroomBinding/DDC] Not found (GroomBinding:%s)."), *In->GetName());
 
 			// Build groom binding data
-			bOutValid = FGroomBindingBuilder::BuildBinding(this, InGroupIndex);
+			FGroomBindingBuilder::FInput BuilderInput;
+			BuilderInput.BindingType = In->GetGroomBindingType();
+			BuilderInput.NumInterpolationPoints = In->GetNumInterpolationPoints();
+			BuilderInput.MatchingSection = In->GetMatchingSection();
+			BuilderInput.bHasValidTarget = In->HasValidTarget();
+			BuilderInput.GroomAsset = In->GetGroom();
+			BuilderInput.SourceSkeletalMesh = In->GetSourceSkeletalMesh();
+			BuilderInput.TargetSkeletalMesh = In->GetTargetSkeletalMesh();
+			BuilderInput.SourceGeometryCache = In->GetSourceGeometryCache();
+			BuilderInput.TargetGeometryCache = In->GetTargetGeometryCache();
+
+			bOutValid = FGroomBindingBuilder::BuildBinding(BuilderInput, InGroupIndex, TargetPlatform, OutPlatformData);
 
 			if (bOutValid)
 			{
-				FillDrivedDataKey(PlatformData);
+				FillDrivedDataKey(OutPlatformData);
 
 				// Header
 				{
 					TArray<uint8> WriteData;
 					FMemoryWriter Ar(WriteData, /*bIsPersistent*/ true);
-					InternalSerializePlatformData(Ar, this, PlatformData, 0 /*Flags*/, true /*bHeader*/, false /*bData*/);
+					InternalSerializePlatformData(Ar, In, OutPlatformData, 0 /*Flags*/, true /*bHeader*/, false /*bData*/);
 	
 					FRequestOwner AsyncOwner(EPriority::Normal);
 					GetCache().PutValue({ {Name, HeaderKey, FValue::Compress(MakeSharedBufferFromArray(MoveTemp(WriteData)))} }, AsyncOwner);
@@ -914,10 +969,10 @@ void UGroomBindingAsset::CacheDerivedDatas(uint32 InGroupIndex, const FString Ke
 				}
 	
 				// Data
-				for (uint32 LODIt=0, LODCount = PlatformData.RenRootBulkData.GetLODCount(); LODIt<LODCount;++LODIt)
+				for (uint32 LODIt=0, LODCount = OutPlatformData.RenRootBulkData.GetLODCount(); LODIt<LODCount;++LODIt)
 				{
 					TArray<FCachePutValueRequest> Out;
-					PlatformData.RenRootBulkData.Write_DDC(this, Out, LODIt);
+					OutPlatformData.RenRootBulkData.Write_DDC(In, Out, LODIt);
 
 					FRequestOwner AsyncOwner(EPriority::Normal);
 					GetCache().PutValue(Out, AsyncOwner);
@@ -925,18 +980,111 @@ void UGroomBindingAsset::CacheDerivedDatas(uint32 InGroupIndex, const FString Ke
 				}
 			}
 		}
+	}
+}
 
-		if (bOutValid)
+static TArray<FString> GetGroupDerivedDataKeys(const UGroomBindingAsset* In, const ITargetPlatform* TargetPlatform)
+{
+	check(In);
+	check(TargetPlatform);
+
+	const FString KeySuffix = BuildDerivedDataKeySuffix(*In, TargetPlatform);
+	const uint32 GroupCount = In->GetGroupInfos().Num();
+
+	TArray<FString> Out;
+	Out.SetNum(GroupCount);
+	for (uint32 GroupIndex = 0; GroupIndex < GroupCount; ++GroupIndex)
+	{
+		Out[GroupIndex] = BuildDerivedDataKeyGroup(KeySuffix, GroupIndex);
+	}
+	return Out;
+}
+
+static UGroomBindingAsset::FCachedCookedPlatformData* FindCachedCookedPlatformData(const TArray<FString>& InGroupKeys, TArray<UGroomBindingAsset::FCachedCookedPlatformData*>& InCachedCookedData)
+{
+	for (UGroomBindingAsset::FCachedCookedPlatformData* CookedPlatformData : InCachedCookedData)
+	{
+		if (CookedPlatformData->GroupDerivedDataKeys == InGroupKeys)
 		{
-			bOutReloadResource = true;
-			CachedDerivedDataKey[InGroupIndex] = DerivedDataKey;
+			return CookedPlatformData;
 		}
-		else
+	}
+	return nullptr;
+}
+
+void UGroomBindingAsset::BeginCacheForCookedPlatformData(const ITargetPlatform* TargetPlatform)
+{
+	Super::BeginCacheForCookedPlatformData(TargetPlatform);
+
+	// 1. Build the key for each group
+	const TArray<FString> GroupDerivedDataKeys = GetGroupDerivedDataKeys(this, TargetPlatform);
+
+	// 2. Find existing cached cooked data
+	UGroomBindingAsset::FCachedCookedPlatformData* TargetPlatformData = FindCachedCookedPlatformData(GroupDerivedDataKeys, CachedCookedPlatformDatas);
+
+	// 3. If the target cooked data does not already exist, we build it
+	if (TargetPlatformData == nullptr)
+	{
+		// 3.1 Build cooked derived data
+		const uint32 GroupCount = GroupDerivedDataKeys.Num();
+		TargetPlatformData = new FCachedCookedPlatformData();
+		TargetPlatformData->GroupDerivedDataKeys = GroupDerivedDataKeys;
+		TargetPlatformData->GroupPlatformDatas.SetNum(GroupCount);
+		for (uint32 GroupIndex = 0; GroupIndex < GroupCount; ++GroupIndex)
+		{
+			bool bGroupValid = true;
+			::CacheDerivedDatas(this, GroupIndex, TargetPlatformData->GroupDerivedDataKeys[GroupIndex], bGroupValid, TargetPlatform, TargetPlatformData->GroupPlatformDatas[GroupIndex]);
+			if (!bGroupValid)
+			{
+				UE_LOG(LogHairStrands, Error, TEXT("[Groom] The binding asset (%s) couldn't be built. This binding asset won't be used."), *GetName());
+			}
+		}
+
+		// 3.2 Place cooked derived data into their bulk data. 
+		// This is done only for strands, which support DDC streaming
+		// When cooking data, force loading of *all* bulk data prior to saving them
+		// Note: bFillBulkdata is true for filling in the bulkdata container prior to serialization. This also forces the resources loading 
+		// from the 'start' (i.e., without offset)
+		for (uint32 GroupIndex = 0; GroupIndex < GroupCount; ++GroupIndex)
 		{
-			UE_LOG(LogHairStrands, Error, TEXT("[Groom] The binding asset (%s) couldn't be built. This binding asset won't be used."), *GetName());
+			FHairStrandsRootBulkData& RenRootBulkData = TargetPlatformData->GroupPlatformDatas[GroupIndex].RenRootBulkData;
+			for (int32 LODIndex = 0, LODCount = RenRootBulkData.GetLODCount(); LODIndex < LODCount; ++LODIndex)
+			{
+				FHairStreamingRequest R; R.Request(HAIR_MAX_NUM_CURVE_PER_GROUP, HAIR_MAX_NUM_POINT_PER_GROUP, LODIndex, RenRootBulkData, true /*bWait*/, true /*bFillBulkdata*/, true /*bWarmCache*/, GetFName());
+			}
 		}
+
+		CachedCookedPlatformDatas.Add(TargetPlatformData);
+	}
+}
+
+TArray<UGroomBindingAsset::FHairGroupPlatformData>* UGroomBindingAsset::GetCachedCookedPlatformData(const ITargetPlatform* TargetPlatform)
+{
+	// 1. Build the key for each group
+	const TArray<FString> GroupDerivedDataKeys = GetGroupDerivedDataKeys(this, TargetPlatform);
+
+	// 2. Find existing cached cooked data
+	if (UGroomBindingAsset::FCachedCookedPlatformData* CachedCookedPlatformData = FindCachedCookedPlatformData(GroupDerivedDataKeys, CachedCookedPlatformDatas))
+	{
+		return &CachedCookedPlatformData->GroupPlatformDatas;
+	}
+	else
+	{
+		return nullptr;
 	}
 }
+
+void UGroomBindingAsset::ClearAllCachedCookedPlatformData()
+{
+	for (UGroomBindingAsset::FCachedCookedPlatformData* CookedPlatformData : CachedCookedPlatformDatas)
+	{
+		delete CookedPlatformData;
+	}
+	CachedCookedPlatformDatas.Empty();
+
+	Super::ClearAllCachedCookedPlatformData();
+}
+
 #endif
 
 void UGroomBindingAsset::GetResourceSizeEx(FResourceSizeEx& CumulativeResourceSize)
diff --git a/Engine/Plugins/Runtime/HairStrands/Source/HairStrandsCore/Private/GroomBindingBuilder.cpp b/Engine/Plugins/Runtime/HairStrands/Source/HairStrandsCore/Private/GroomBindingBuilder.cpp
--- a/Engine/Plugins/Runtime/HairStrands/Source/HairStrandsCore/Private/GroomBindingBuilder.cpp
+++ b/Engine/Plugins/Runtime/HairStrands/Source/HairStrandsCore/Private/GroomBindingBuilder.cpp
@@ -12,6 +12,8 @@
 #include "GlobalShader.h"
 #include "Misc/ScopedSlowTask.h"
 #include "GroomRBFDeformer.h"
+#include "Engine/SkinnedAssetAsyncCompileUtils.h"
+#include "Interfaces/ITargetPlatform.h"
 
 #if WITH_EDITORONLY_DATA
 
@@ -58,7 +60,7 @@ static FAutoConsoleVariableRef CVarHairStrandsBindingBuilderWarningEnable(TEXT("
 FString FGroomBindingBuilder::GetVersion()
 {
 	// Important to update the version when groom building changes
-	return TEXT("3p");
+	return TEXT("3p_7");
 }
 
 ///////////////////////////////////////////////////////////////////////////////////////////////////
@@ -293,15 +295,15 @@ private:
 class FSkeletalMeshData : public IMeshData
 {
 public:
-	FSkeletalMeshData(USkeletalMesh* InSkeletalMesh) 
+	FSkeletalMeshData(const USkeletalMesh* InSkeletalMesh, const FSkeletalMeshRenderData* InRenderData) 
 	{
 		SkeletalMesh = InSkeletalMesh;
+		MeshData = InRenderData;
 		if (SkeletalMesh)
 		{
-			MeshData = SkeletalMesh->GetResourceForRendering();
 			if (MeshData)
 			{
-				const uint32 LODCount = SkeletalMesh->GetLODNum();
+				LODCount = SkeletalMesh->GetLODNum();
 				MeshesLODData.Reserve(LODCount);
 				for (uint32 LODIt=0; LODIt<LODCount; ++LODIt)
 				{
@@ -322,7 +324,7 @@ public:
 
 	virtual uint32 GetNumLODs() const override
 	{
-		return SkeletalMesh->GetLODNum();
+		return LODCount;
 	}
 
 	virtual const IMeshLODData& GetMeshLODData(uint32 InLODIndex) const override
@@ -332,6 +334,7 @@ public:
 	}
 
 private:
+	uint32 LODCount = 0;
 	const USkeletalMesh* SkeletalMesh = nullptr;
 	const FSkeletalMeshRenderData* MeshData = nullptr;
 	TArray<FSkeletalMeshLODData> MeshesLODData;
@@ -1439,6 +1442,7 @@ namespace GroomBinding_RootProjection
 
 			// Update the root mesh projection data with unique valid mesh section IDs, based on the projection data
 			OutRootData.MeshProjectionLODs[LODIt].UniqueSectionIds = UniqueSectionId;
+			OutRootData.MeshProjectionLODs[LODIt].MeshSectionCount = SectionCount;
 		}
 
 		return true;
@@ -2001,6 +2005,7 @@ static void BuildRootBulkData(
 		Out.Header.LODs[MeshLODIt].SampleCount 			= bHasValidSamples ? In.MeshProjectionLODs[MeshLODIt].SampleCount : 0u;
 		Out.Header.LODs[MeshLODIt].UniqueTriangleCount 	= In.MeshProjectionLODs[MeshLODIt].UniqueTriangleIndexBuffer.Num();
 		Out.Header.LODs[MeshLODIt].UniqueSectionIndices = In.MeshProjectionLODs[MeshLODIt].UniqueSectionIds;
+		Out.Header.LODs[MeshLODIt].MeshSectionCount		= In.MeshProjectionLODs[MeshLODIt].MeshSectionCount;
 	}
 
 	// Data
@@ -2051,6 +2056,7 @@ static void BuildRootData(
 		Out.MeshProjectionLODs[MeshLODIt].LODIndex = In.Header.LODs[MeshLODIt].LODIndex;
 		Out.MeshProjectionLODs[MeshLODIt].SampleCount = bHasValidSamples ? In.Header.LODs[MeshLODIt].SampleCount : 0u;
 		Out.MeshProjectionLODs[MeshLODIt].UniqueSectionIds = In.Header.LODs[MeshLODIt].UniqueSectionIndices;
+		Out.MeshProjectionLODs[MeshLODIt].MeshSectionCount = In.Header.LODs[MeshLODIt].MeshSectionCount;
 
 		CopyFromBulkData<FHairStrandsUniqueTriangleIndexFormat>(Out.MeshProjectionLODs[MeshLODIt].UniqueTriangleIndexBuffer, In.Data.LODs[MeshLODIt].UniqueTriangleIndexBuffer);
 		CopyFromBulkData<FHairStrandsRootToUniqueTriangleIndexFormat>(Out.MeshProjectionLODs[MeshLODIt].RootToUniqueTriangleIndexBuffer, In.Data.LODs[MeshLODIt].RootToUniqueTriangleIndexBuffer);
@@ -2096,13 +2102,10 @@ static void BuildRootBulkData(
 
   ///////////////////////////////////////////////////////////////////////////////////////////////////
 // Main entry (CPU path)
-static bool InternalBuildBinding_CPU(UGroomBindingAsset* BindingAsset, uint32 InGroupIndex)
+static bool InternalBuildBinding_CPU(const FGroomBindingBuilder::FInput& In, uint32 InGroupIndex, const ITargetPlatform* TargetPlatform, UGroomBindingAsset::FHairGroupPlatformData& OutPlatformData)
 {
 #if WITH_EDITORONLY_DATA
-	if (!BindingAsset ||
-		!BindingAsset->GetGroom() ||
-		!BindingAsset->HasValidTarget() ||
-		BindingAsset->GetGroom()->GetNumHairGroups() == 0)
+	if (!In.GroomAsset || !In.bHasValidTarget || In.GroomAsset->GetNumHairGroups() == 0)
 	{
 		UE_LOG(LogHairStrands, Error, TEXT("[Groom] Binding asset cannot be created/rebuilt."));
 		return false;
@@ -2111,38 +2114,73 @@ static bool InternalBuildBinding_CPU(UGroomBindingAsset* BindingAsset, uint32 In
 	// 1. Build groom root data
 	FHairRootGroupData OutData;
 	{
-		BindingAsset->GetGroom()->ConditionalPostLoad();
+		In.GroomAsset->ConditionalPostLoad();
 
-		const int32 NumInterpolationPoints = BindingAsset->GetNumInterpolationPoints();
-		UGroomAsset* GroomAsset = BindingAsset->GetGroom();
+		// Ensure the skeletal meshes / geom caches are built
+		if (In.BindingType == EGroomBindingMeshType::SkeletalMesh)
+		{
+			check(!In.TargetSkeletalMesh->IsCompiling());
+			check(In.TargetSkeletalMesh->IsAsyncTaskComplete());
+			In.TargetSkeletalMesh->ConditionalPostLoad();
+			if (In.SourceSkeletalMesh)
+			{
+				check(!In.SourceSkeletalMesh->IsCompiling());
+				check(In.SourceSkeletalMesh->IsAsyncTaskComplete());
+				In.SourceSkeletalMesh->ConditionalPostLoad();
+			}
+		}
+		else
+		{
+			In.TargetGeometryCache->ConditionalPostLoad();
+			if (In.SourceGeometryCache)
+			{
+				In.SourceGeometryCache->ConditionalPostLoad();
+			}
+		}
+
+		// * Only for SkeletalMesh: Take scoped lock on the skeletal render mesh data during the entire groom binding building
+		// * Then use an async build scope to allow accessing skeletal mesh property safely.
+		//   If skel.meshes are nullptr, this will act as a NOP
+		USkeletalMesh* InSourceSkeletalMesh = In.SourceSkeletalMesh == In.TargetSkeletalMesh ? nullptr : In.SourceSkeletalMesh;
+		FScopedSkeletalMeshRenderData SourceSkeletalMeshScopedData(InSourceSkeletalMesh);
+		FScopedSkeletalMeshRenderData TargetSkeletalMeshScopedData(In.TargetSkeletalMesh);
 
 		TUniquePtr<GroomBinding_Mesh::IMeshData> SourceMeshData;
 		TUniquePtr<GroomBinding_Mesh::IMeshData> TargetMeshData;
-		if (BindingAsset->GetGroomBindingType() == EGroomBindingMeshType::SkeletalMesh)
+		if (In.BindingType == EGroomBindingMeshType::SkeletalMesh)
 		{
-			check(!BindingAsset->GetTargetSkeletalMesh()->IsCompiling());
-			check(BindingAsset->GetTargetSkeletalMesh()->IsAsyncTaskComplete());
-			BindingAsset->GetTargetSkeletalMesh()->ConditionalPostLoad();
-			if (BindingAsset->GetSourceSkeletalMesh())
+			if (InSourceSkeletalMesh)
 			{
-				check(!BindingAsset->GetSourceSkeletalMesh()->IsCompiling());
-				check(BindingAsset->GetSourceSkeletalMesh()->IsAsyncTaskComplete());
-				BindingAsset->GetSourceSkeletalMesh()->ConditionalPostLoad();
+				FSkinnedAssetAsyncBuildScope AsyncBuildScope(InSourceSkeletalMesh);
+				USkeletalMesh::GetPlatformSkeletalMeshRenderData(TargetPlatform, SourceSkeletalMeshScopedData);
+				SourceMeshData = TUniquePtr<GroomBinding_Mesh::FSkeletalMeshData, TDefaultDelete<GroomBinding_Mesh::IMeshData>>(new GroomBinding_Mesh::FSkeletalMeshData(InSourceSkeletalMesh, SourceSkeletalMeshScopedData.GetData()));
+			}
+			else
+			{
+				SourceMeshData = TUniquePtr<GroomBinding_Mesh::FSkeletalMeshData, TDefaultDelete<GroomBinding_Mesh::IMeshData>>(new GroomBinding_Mesh::FSkeletalMeshData(nullptr, nullptr));
 			}
 
-			SourceMeshData = TUniquePtr<GroomBinding_Mesh::FSkeletalMeshData, TDefaultDelete<GroomBinding_Mesh::IMeshData>>(new GroomBinding_Mesh::FSkeletalMeshData(BindingAsset->GetSourceSkeletalMesh()));
-			TargetMeshData = TUniquePtr<GroomBinding_Mesh::FSkeletalMeshData, TDefaultDelete<GroomBinding_Mesh::IMeshData>>(new GroomBinding_Mesh::FSkeletalMeshData(BindingAsset->GetTargetSkeletalMesh()));
+			if (In.TargetSkeletalMesh)
+			{
+				FSkinnedAssetAsyncBuildScope AsyncBuildScope(In.TargetSkeletalMesh);
+				USkeletalMesh::GetPlatformSkeletalMeshRenderData(TargetPlatform, TargetSkeletalMeshScopedData);
+				TargetMeshData = TUniquePtr<GroomBinding_Mesh::FSkeletalMeshData, TDefaultDelete<GroomBinding_Mesh::IMeshData>>(new GroomBinding_Mesh::FSkeletalMeshData(In.TargetSkeletalMesh, TargetSkeletalMeshScopedData.GetData()));
+			}
+			else
+			{
+				TargetMeshData = TUniquePtr<GroomBinding_Mesh::FSkeletalMeshData, TDefaultDelete<GroomBinding_Mesh::IMeshData>>(new GroomBinding_Mesh::FSkeletalMeshData(nullptr, nullptr));
+			}
 		}
 		else
 		{
-			BindingAsset->GetTargetGeometryCache()->ConditionalPostLoad();
-			if (BindingAsset->GetSourceGeometryCache())
+			In.TargetGeometryCache->ConditionalPostLoad();
+			if (In.SourceGeometryCache)
 			{
-				BindingAsset->GetSourceGeometryCache()->ConditionalPostLoad();
+				In.SourceGeometryCache->ConditionalPostLoad();
 			}
 
-			SourceMeshData = TUniquePtr<GroomBinding_Mesh::FGeometryCacheData, TDefaultDelete<GroomBinding_Mesh::IMeshData>>(new GroomBinding_Mesh::FGeometryCacheData(BindingAsset->GetSourceGeometryCache()));
-			TargetMeshData = TUniquePtr<GroomBinding_Mesh::FGeometryCacheData, TDefaultDelete<GroomBinding_Mesh::IMeshData>>(new GroomBinding_Mesh::FGeometryCacheData(BindingAsset->GetTargetGeometryCache()));
+			SourceMeshData = TUniquePtr<GroomBinding_Mesh::FGeometryCacheData, TDefaultDelete<GroomBinding_Mesh::IMeshData>>(new GroomBinding_Mesh::FGeometryCacheData(In.SourceGeometryCache));
+			TargetMeshData = TUniquePtr<GroomBinding_Mesh::FGeometryCacheData, TDefaultDelete<GroomBinding_Mesh::IMeshData>>(new GroomBinding_Mesh::FGeometryCacheData(In.TargetGeometryCache));
 		}
 
 		if (!TargetMeshData->IsValid())
@@ -2150,14 +2188,14 @@ static bool InternalBuildBinding_CPU(UGroomBindingAsset* BindingAsset, uint32 In
 			UE_LOG(LogHairStrands, Error, TEXT("[Groom] Binding asset could not be built. Target mesh is not valid."));
 			return false;
 		}
-		const uint32 GroupCount = GroomAsset->GetNumHairGroups();
+		const uint32 GroupCount = In.GroomAsset->GetNumHairGroups();
 		const uint32 MeshLODCount = TargetMeshData->GetNumLODs();
 
-		check(InGroupIndex < uint32(GroomAsset->GetHairGroupsPlatformData().Num()));
+		check(InGroupIndex < uint32(In.GroomAsset->GetHairGroupsPlatformData().Num()));
 
 		// Check if root data are needs for strands
 		bool bNeedStrandsRoot = false;
-		for (const FHairLODSettings& LODSettings : GroomAsset->GetHairGroupsLOD()[InGroupIndex].LODs)
+		for (const FHairLODSettings& LODSettings : In.GroomAsset->GetHairGroupsLOD()[InGroupIndex].LODs)
 		{
 			if (LODSettings.GeometryType == EGroomGeometryType::Strands)
 			{
@@ -2169,18 +2207,18 @@ static bool InternalBuildBinding_CPU(UGroomBindingAsset* BindingAsset, uint32 In
 		// 1.1 Build guide/strands data
 		FHairStrandsDatas StrandsData;
 		FHairStrandsDatas GuidesData;
-		GroomAsset->GetHairStrandsDatas(InGroupIndex, StrandsData, GuidesData);
+		In.GroomAsset->GetHairStrandsDatas(InGroupIndex, StrandsData, GuidesData);
 
 		// 1.2 Init root data for guides/strands/cards
-		const FHairGroupPlatformData& GroupData = GroomAsset->GetHairGroupsPlatformData()[InGroupIndex];
+		const FHairGroupPlatformData& GroupData = In.GroomAsset->GetHairGroupsPlatformData()[InGroupIndex];
 		{
 			// Guides
-			InitHairStrandsRootData(OutData.SimRootData, &GuidesData, MeshLODCount, NumInterpolationPoints);
+			InitHairStrandsRootData(OutData.SimRootData, &GuidesData, MeshLODCount, In.NumInterpolationPoints);
 
 			// Strands
 			if (bNeedStrandsRoot)
 			{
-				InitHairStrandsRootData(OutData.RenRootData, &StrandsData, MeshLODCount, NumInterpolationPoints);
+				InitHairStrandsRootData(OutData.RenRootData, &StrandsData, MeshLODCount, In.NumInterpolationPoints);
 			}
 
 			// Cards
@@ -2191,10 +2229,10 @@ static bool InternalBuildBinding_CPU(UGroomBindingAsset* BindingAsset, uint32 In
 				if (GroupData.Cards.IsValid(CardsLODIt))
 				{
 					FHairStrandsDatas LODGuidesData;
-					const bool bIsValid = GroomAsset->GetHairCardsGuidesDatas(InGroupIndex, CardsLODIt, LODGuidesData);
+					const bool bIsValid = In.GroomAsset->GetHairCardsGuidesDatas(InGroupIndex, CardsLODIt, LODGuidesData);
 					if (bIsValid)
 					{
-						InitHairStrandsRootData(OutData.CardsRootData[CardsLODIt], &LODGuidesData, MeshLODCount, NumInterpolationPoints);
+						InitHairStrandsRootData(OutData.CardsRootData[CardsLODIt], &LODGuidesData, MeshLODCount, In.NumInterpolationPoints);
 					}
 				}
 			}
@@ -2225,7 +2263,7 @@ static bool InternalBuildBinding_CPU(UGroomBindingAsset* BindingAsset, uint32 In
 			if (!GroomBinding_Transfer::Transfer(
 				SourceMeshData.Get(),
 				TargetMeshData.Get(),
-				TransferredPositions, BindingAsset->GetMatchingSection()))
+				TransferredPositions, In.MatchingSection))
 			{
 				UE_LOG(LogHairStrands, Error, TEXT("[Groom] Binding asset could not be built. Positions transfer between source and target mesh failed."));
 				return false;
@@ -2268,10 +2306,10 @@ static bool InternalBuildBinding_CPU(UGroomBindingAsset* BindingAsset, uint32 In
 			const uint32 CardsLODCount = OutData.CardsRootData.Num();
 			for (uint32 CardsLODIt = 0; CardsLODIt < CardsLODCount; ++CardsLODIt)
 			{
-				if (BindingAsset->GetGroom()->GetHairGroupsPlatformData()[InGroupIndex].Cards.IsValid(CardsLODIt))
+				if (In.GroomAsset->GetHairGroupsPlatformData()[InGroupIndex].Cards.IsValid(CardsLODIt))
 				{
 					FHairStrandsDatas LODGuidesData;
-					const bool bIsValid = GroomAsset->GetHairCardsGuidesDatas(InGroupIndex, CardsLODIt, LODGuidesData);
+					const bool bIsValid = In.GroomAsset->GetHairCardsGuidesDatas(InGroupIndex, CardsLODIt, LODGuidesData);
 					if (bIsValid)
 					{
 						if (!GroomBinding_RootProjection::Project(
@@ -2291,59 +2329,36 @@ static bool InternalBuildBinding_CPU(UGroomBindingAsset* BindingAsset, uint32 In
 		
 		// 1.5 RBF building
 		{
-			GroomBinding_RBFWeighting::ComputeInterpolationWeights(OutData, bNeedStrandsRoot, BindingAsset->GetNumInterpolationPoints(), BindingAsset->GetMatchingSection(), TargetMeshData.Get(), TransferredPositions);
+			GroomBinding_RBFWeighting::ComputeInterpolationWeights(OutData, bNeedStrandsRoot, In.NumInterpolationPoints, In.MatchingSection, TargetMeshData.Get(), TransferredPositions);
 			SlowTask.EnterProgressFrame();
 		}
 	}
 
-	// 2. Release existing resources data
-	UGroomBindingAsset::FHairGroupResources& OutHairGroupResources = BindingAsset->GetHairGroupResources();
-	if (OutHairGroupResources.Num() > 0)
-	{
-		for (UGroomBindingAsset::FHairGroupResource& GroupResources : OutHairGroupResources)
-		{
-			BindingAsset->AddHairGroupResourcesToDelete(GroupResources);
-		}
-		OutHairGroupResources.Empty();
-	}
-	check(OutHairGroupResources.Num() == 0);
-
 	// 3. Convert data to bulk data
-	GroomBinding_BulkCopy::BuildRootBulkData(BindingAsset->GetHairGroupsPlatformData()[InGroupIndex], OutData);
-
-	BindingAsset->QueryStatus = UGroomBindingAsset::EQueryStatus::Completed;
+	GroomBinding_BulkCopy::BuildRootBulkData(OutPlatformData, OutData);
 #endif
 	return true;
 }
-void UpdateGroomBindingAssetInfos(UGroomBindingAsset* In);
 
-bool FGroomBindingBuilder::BuildBinding(UGroomBindingAsset* BindingAsset, uint32 InGroupIndex)
+bool FGroomBindingBuilder::BuildBinding(const FGroomBindingBuilder::FInput& In, uint32 InGroupIndex, const ITargetPlatform* TargetPlatform, UGroomBindingAsset::FHairGroupPlatformData& Out)
 {
-	return InternalBuildBinding_CPU(BindingAsset, InGroupIndex);
+	return InternalBuildBinding_CPU(In, InGroupIndex, TargetPlatform, Out);
 }
 
-bool FGroomBindingBuilder::BuildBinding(UGroomBindingAsset* BindingAsset, bool bInitResources)
+bool FGroomBindingBuilder::BuildBinding(class UGroomBindingAsset* BindingAsset, bool bInitResource)
 {
-	bool bOutValid = true;
-
-	// 1. Build binding asset
-	const uint32 GroupCount = BindingAsset->GetHairGroupsPlatformData().Num();
-	BindingAsset->GetGroupInfos().SetNum(GroupCount);
-	for (uint32 InGroupIndex = 0; InGroupIndex < GroupCount; ++InGroupIndex)
-	{
-		bOutValid = InternalBuildBinding_CPU(BindingAsset, InGroupIndex) && bOutValid;
-	}
-
-	// 2. Optionnally update resources
-	if (bOutValid && bInitResources)
-	{
-		BindingAsset->InitResource();
-	}
-
-	// 3. Update GroomBindingAsset infos
-	UpdateGroomBindingAssetInfos(BindingAsset);
+#if WITH_EDITORONLY_DATA
+	BindingAsset->CacheDerivedDatas();
+#endif
+	return true;
+}
 
-	return bOutValid;
+bool FGroomBindingBuilder::BuildBinding(class UGroomBindingAsset* BindingAsset, uint32 InGroupIndex)
+{
+#if WITH_EDITORONLY_DATA
+	BindingAsset->CacheDerivedDatas();
+#endif
+	return true;
 }
 
 void FGroomBindingBuilder::GetRootData(
diff --git a/Engine/Plugins/Runtime/HairStrands/Source/HairStrandsCore/Private/HairStrandsDatas.cpp b/Engine/Plugins/Runtime/HairStrands/Source/HairStrandsCore/Private/HairStrandsDatas.cpp
--- a/Engine/Plugins/Runtime/HairStrands/Source/HairStrandsCore/Private/HairStrandsDatas.cpp
+++ b/Engine/Plugins/Runtime/HairStrands/Source/HairStrandsCore/Private/HairStrandsDatas.cpp
@@ -876,6 +876,7 @@ void FHairStrandsRootBulkData::SerializeHeader(FArchive& Ar, UObject* Owner)
 		Ar << Header.LODs[LODIt].UniqueTriangleCount;
 		Ar << Header.LODs[LODIt].SampleCount;
 		Ar << Header.LODs[LODIt].UniqueSectionIndices;
+		Ar << Header.LODs[LODIt].MeshSectionCount;
 	}
 }
 
diff --git a/Engine/Plugins/Runtime/HairStrands/Source/HairStrandsCore/Public/GroomBindingAsset.h b/Engine/Plugins/Runtime/HairStrands/Source/HairStrandsCore/Public/GroomBindingAsset.h
--- a/Engine/Plugins/Runtime/HairStrands/Source/HairStrandsCore/Public/GroomBindingAsset.h
+++ b/Engine/Plugins/Runtime/HairStrands/Source/HairStrandsCore/Public/GroomBindingAsset.h
@@ -237,14 +237,6 @@ public:
 	
 #endif // WITH_EDITOR
 
-	enum class EQueryStatus
-	{
-		None,
-		Submitted,
-		Completed
-	};
-	volatile EQueryStatus QueryStatus = EQueryStatus::None;
-
 	/** Initialize resources. */
 	void InitResource();
 
@@ -269,10 +261,22 @@ public:
 	void Build();
 
 	void CacheDerivedDatas();
-	void CacheDerivedDatas(uint32 InGroupIndex, const FString KeySuffix, bool& bOutValid, bool& bOutReloadResource);
+
+	virtual void BeginCacheForCookedPlatformData(const ITargetPlatform* TargetPlatform);
+	virtual void ClearAllCachedCookedPlatformData();
+	TArray<FHairGroupPlatformData>* GetCachedCookedPlatformData(const ITargetPlatform* TargetPlatform);
 
 	void InvalidateBinding();
 	void InvalidateBinding(class USkeletalMesh*);
+
+	struct FCachedCookedPlatformData
+	{
+		TArray<FString> GroupDerivedDataKeys;
+		TArray<FHairGroupPlatformData> GroupPlatformDatas;
+	};
+private:
+	TArray<FCachedCookedPlatformData*> CachedCookedPlatformDatas;
+
 	bool bRegisterSourceMeshCallback = false;
 	bool bRegisterTargetMeshCallback = false;
 	bool bRegisterGroomAssetCallback = false;
diff --git a/Engine/Plugins/Runtime/HairStrands/Source/HairStrandsCore/Public/GroomBindingBuilder.h b/Engine/Plugins/Runtime/HairStrands/Source/HairStrandsCore/Public/GroomBindingBuilder.h
--- a/Engine/Plugins/Runtime/HairStrands/Source/HairStrandsCore/Public/GroomBindingBuilder.h
+++ b/Engine/Plugins/Runtime/HairStrands/Source/HairStrandsCore/Public/GroomBindingBuilder.h
@@ -9,17 +9,38 @@
 #include "HairStrandsInterface.h"
 #include "HairStrandsMeshProjection.h"
 #include "HairStrandsDatas.h"
+#include "GroomBindingAsset.h"
 
-struct FHairStrandsRestRootResource;
+class ITargetPlatform;
 
 struct HAIRSTRANDSCORE_API FGroomBindingBuilder
 {
+	struct FInput
+	{
+		EGroomBindingMeshType BindingType = EGroomBindingMeshType::SkeletalMesh;
+		int32 NumInterpolationPoints = 0;
+		int32 MatchingSection = 0;
+		bool bHasValidTarget = false;
+		UGroomAsset* GroomAsset = nullptr;
+		USkeletalMesh* SourceSkeletalMesh = nullptr;
+		USkeletalMesh* TargetSkeletalMesh = nullptr;
+		UGeometryCache* SourceGeometryCache;
+		UGeometryCache* TargetGeometryCache;
+	};
+
 	static FString GetVersion();
 
 	// Build binding asset data
+	UE_DEPRECATED(5.4, "Please do not access this funciton; but rather call BindingAsset->CacheDerivedDatas()")
 	static bool BuildBinding(class UGroomBindingAsset* BindingAsset, bool bInitResource);
+
+	// Build binding asset data
+	UE_DEPRECATED(5.4, "Please do not access this funciton; but rather call BindingAsset->CacheDerivedDatas()")
 	static bool BuildBinding(class UGroomBindingAsset* BindingAsset, uint32 InGroupIndex);
 
+	// Build binding asset data
+	static bool BuildBinding(const FInput& In, uint32 InGroupIndex, const ITargetPlatform* TargetPlatform, UGroomBindingAsset::FHairGroupPlatformData& Out);
+
 	// Extract root data from bulk data
 	static void GetRootData(FHairStrandsRootData& Out, const FHairStrandsRootBulkData& In);
 };
diff --git a/Engine/Plugins/Runtime/HairStrands/Source/HairStrandsCore/Public/HairStrandsDatas.h b/Engine/Plugins/Runtime/HairStrands/Source/HairStrandsCore/Public/HairStrandsDatas.h
--- a/Engine/Plugins/Runtime/HairStrands/Source/HairStrandsCore/Public/HairStrandsDatas.h
+++ b/Engine/Plugins/Runtime/HairStrands/Source/HairStrandsCore/Public/HairStrandsDatas.h
@@ -791,6 +791,9 @@ struct FHairStrandsRootData
 
 		/* Store the mesh section indices which are relevant for this root LOD data */
 		TArray<uint32> UniqueSectionIds;
+
+		/* Number of render section of the target mesh */
+		uint32 MeshSectionCount = 0;
 	};
 
 	/* Number of roots */
@@ -840,6 +843,7 @@ struct FHairStrandsRootBulkData : FHairStrandsBulkCommon
 			int32  LODIndex = -1;
 			uint32 UniqueTriangleCount = 0;			
 			uint32 SampleCount = 0; 				// Number of samples used for the mesh interpolation
+			uint32 MeshSectionCount = 0;			// Number of section of the target mesh
 			TArray<uint32> UniqueSectionIndices; 	// Store the mesh section indices which are relevant for this root LOD data
 		};
 
diff --git a/Engine/Plugins/Runtime/HairStrands/Source/HairStrandsEditor/Private/GroomBindingActions.cpp b/Engine/Plugins/Runtime/HairStrands/Source/HairStrandsEditor/Private/GroomBindingActions.cpp
--- a/Engine/Plugins/Runtime/HairStrands/Source/HairStrandsEditor/Private/GroomBindingActions.cpp
+++ b/Engine/Plugins/Runtime/HairStrands/Source/HairStrandsEditor/Private/GroomBindingActions.cpp
@@ -47,7 +47,7 @@ void ExecuteRebuildBindingAsset(const FToolMenuContext& InContext)
 						BindingAsset->GetSourceGeometryCache()->ConditionalPostLoad();
 					}
 				}
-				FGroomBindingBuilder::BuildBinding(BindingAsset, true);
+				BindingAsset->CacheDerivedDatas();
 				BindingAsset->GetOutermost()->MarkPackageDirty();
 			}
 		}
