By pass sim cache update if there is no sim cache instance.

diff --git a/Engine/Plugins/Runtime/HairStrands/Source/HairStrandsCore/Private/GroomManager.cpp b/Engine/Plugins/Runtime/HairStrands/Source/HairStrandsCore/Private/GroomManager.cpp
--- a/Engine/Plugins/Runtime/HairStrands/Source/HairStrandsCore/Private/GroomManager.cpp
+++ b/Engine/Plugins/Runtime/HairStrands/Source/HairStrandsCore/Private/GroomManager.cpp
@@ -433,6 +433,7 @@ static void RunHairStrandsInterpolation_Guide(
 	};
 
 	// Gather all instances which require guides or RBF deformations
+	bool bHasAnySimCacheInstances = false;
 	TArray<FInstanceData> InstanceDatas;
 	InstanceDatas.Reserve(Instances.Num());
 	for (FHairStrandsInstance* AbstractInstance : Instances)
@@ -471,6 +472,8 @@ static void RunHairStrandsInterpolation_Guide(
 		{
 			check(InstanceData.Instance->Guides.IsValid());
 		}
+
+		if (InstanceData.ActiveGroomCacheType == EGroomCacheType::Guides)  { bHasAnySimCacheInstances = true; }
 	}	
 
 	// Update dynamic mesh triangles
@@ -650,45 +653,47 @@ static void RunHairStrandsInterpolation_Guide(
 	}
 
 	// Apply guide cache update
-	for (FInstanceData& InstanceData : InstanceDatas)
+	if (bHasAnySimCacheInstances)
 	{
-		if (InstanceData.ActiveGroomCacheType == EGroomCacheType::Guides)
+		for (FInstanceData& InstanceData : InstanceDatas)
 		{
-			FScopeLock Lock(InstanceData.Instance->Debug.GroomCacheBuffers->GetCriticalSection());
-			FGroomCacheGroupData* GroomCacheData0 = const_cast<FGroomCacheGroupData*>(&InstanceData.Instance->Debug.GroomCacheBuffers->GetCurrentFrameBuffer().GroupsData[InstanceData.Instance->Debug.GroupIndex]);
-			FGroomCacheGroupData* GroomCacheData1 = const_cast<FGroomCacheGroupData*>(&InstanceData.Instance->Debug.GroomCacheBuffers->GetNextFrameBuffer().GroupsData[InstanceData.Instance->Debug.GroupIndex]);
-
-			const float InterpolationFactor = InstanceData.Instance->Debug.GroomCacheBuffers->GetInterpolationFactor();
-
-			// Update CPU position offset to GPU
-			InstanceData.Instance->Guides.DeformedResource->SetPositionOffset(FHairStrandsDeformedResource::EFrameType::Current, FMath::Lerp(GroomCacheData0->BoundingBox.GetCenter(), GroomCacheData1->BoundingBox.GetCenter(), InterpolationFactor));
-			AddHairStrandUpdatePositionOffsetPass(
-				GraphBuilder,
-				ShaderMap,
-				EHairPositionUpdateType::Guides,
-				InstanceData.Instance->RegisteredIndex,
-				InstanceData.HairLODIndex,
-				InstanceData.MeshLODIndex,
-				InstanceData.Instance->Guides.DeformedRootResource,
-				InstanceData.Instance->Guides.DeformedResource);
-
-			// Pass to upload GroomCache guide positions
-			FGroomCacheResources CacheResources0 = CreateGroomCacheBuffer(GraphBuilder, GroomCacheData0->VertexData);
-			FGroomCacheResources CacheResources1 = CreateGroomCacheBuffer(GraphBuilder, GroomCacheData1->VertexData);
-			AddGroomCacheUpdatePass(
-				GraphBuilder,
-				ShaderMap,
-				InstanceData.Instance->RegisteredIndex,
-				InstanceData.Instance->Guides.RestResource->GetPointCount(),
-				InterpolationFactor,
-				CacheResources0,
-				CacheResources1,
-				RegisterAsSRV(GraphBuilder, InstanceData.Instance->Guides.RestResource->PositionBuffer),
-				RegisterAsSRV(GraphBuilder, InstanceData.Instance->Guides.DeformedResource->GetPositionOffsetBuffer(FHairStrandsDeformedResource::EFrameType::Current)),
-				RegisterAsUAV(GraphBuilder, InstanceData.Instance->Guides.DeformedResource->GetBuffer(FHairStrandsDeformedResource::Current)));
+			if (InstanceData.ActiveGroomCacheType == EGroomCacheType::Guides)
+			{
+				FScopeLock Lock(InstanceData.Instance->Debug.GroomCacheBuffers->GetCriticalSection());
+				FGroomCacheGroupData* GroomCacheData0 = const_cast<FGroomCacheGroupData*>(&InstanceData.Instance->Debug.GroomCacheBuffers->GetCurrentFrameBuffer().GroupsData[InstanceData.Instance->Debug.GroupIndex]);
+				FGroomCacheGroupData* GroomCacheData1 = const_cast<FGroomCacheGroupData*>(&InstanceData.Instance->Debug.GroomCacheBuffers->GetNextFrameBuffer().GroupsData[InstanceData.Instance->Debug.GroupIndex]);
+	
+				const float InterpolationFactor = InstanceData.Instance->Debug.GroomCacheBuffers->GetInterpolationFactor();
+	
+				// Update CPU position offset to GPU
+				InstanceData.Instance->Guides.DeformedResource->SetPositionOffset(FHairStrandsDeformedResource::EFrameType::Current, FMath::Lerp(GroomCacheData0->BoundingBox.GetCenter(), GroomCacheData1->BoundingBox.GetCenter(), InterpolationFactor));
+				AddHairStrandUpdatePositionOffsetPass(
+					GraphBuilder,
+					ShaderMap,
+					EHairPositionUpdateType::Guides,
+					InstanceData.Instance->RegisteredIndex,
+					InstanceData.HairLODIndex,
+					InstanceData.MeshLODIndex,
+					InstanceData.Instance->Guides.DeformedRootResource,
+					InstanceData.Instance->Guides.DeformedResource);
+	
+				// Pass to upload GroomCache guide positions
+				FGroomCacheResources CacheResources0 = CreateGroomCacheBuffer(GraphBuilder, GroomCacheData0->VertexData);
+				FGroomCacheResources CacheResources1 = CreateGroomCacheBuffer(GraphBuilder, GroomCacheData1->VertexData);
+				AddGroomCacheUpdatePass(
+					GraphBuilder,
+					ShaderMap,
+					InstanceData.Instance->RegisteredIndex,
+					InstanceData.Instance->Guides.RestResource->GetPointCount(),
+					InterpolationFactor,
+					CacheResources0,
+					CacheResources1,
+					RegisterAsSRV(GraphBuilder, InstanceData.Instance->Guides.RestResource->PositionBuffer),
+					RegisterAsSRV(GraphBuilder, InstanceData.Instance->Guides.DeformedResource->GetPositionOffsetBuffer(FHairStrandsDeformedResource::EFrameType::Current)),
+					RegisterAsUAV(GraphBuilder, InstanceData.Instance->Guides.DeformedResource->GetBuffer(FHairStrandsDeformedResource::Current)));
+			}
 		}
 	}
-
 }
 
 void AddDrawDebugClusterPass(
