Remove invalid check

diff --git a/Engine/Plugins/Runtime/HairStrands/Source/HairStrandsCore/Private/GroomBindingBuilder.cpp b/Engine/Plugins/Runtime/HairStrands/Source/HairStrandsCore/Private/GroomBindingBuilder.cpp
--- a/Engine/Plugins/Runtime/HairStrands/Source/HairStrandsCore/Private/GroomBindingBuilder.cpp
+++ b/Engine/Plugins/Runtime/HairStrands/Source/HairStrandsCore/Private/GroomBindingBuilder.cpp
@@ -2119,14 +2119,12 @@ static bool InternalBuildBinding_CPU(const FGroomBindingBuilder::FInput& In, uin
 		// Ensure the skeletal meshes / geom caches are built
 		if (In.BindingType == EGroomBindingMeshType::SkeletalMesh)
 		{
-			check(!In.TargetSkeletalMesh->IsCompiling());
-			check(In.TargetSkeletalMesh->IsAsyncTaskComplete());
 			In.TargetSkeletalMesh->ConditionalPostLoad();
+			In.TargetSkeletalMesh->GetLODNum();
 			if (In.SourceSkeletalMesh)
 			{
-				check(!In.SourceSkeletalMesh->IsCompiling());
-				check(In.SourceSkeletalMesh->IsAsyncTaskComplete());
 				In.SourceSkeletalMesh->ConditionalPostLoad();
+				In.SourceSkeletalMesh->GetLODNum();
 			}
 		}
 		else
