Pack RBF sample index and section into uint buffer

diff --git a/Engine/Plugins/Runtime/HairStrands/Source/HairStrandsCore/Private/GroomBindingBuilder.cpp b/Engine/Plugins/Runtime/HairStrands/Source/HairStrandsCore/Private/GroomBindingBuilder.cpp
--- a/Engine/Plugins/Runtime/HairStrands/Source/HairStrandsCore/Private/GroomBindingBuilder.cpp
+++ b/Engine/Plugins/Runtime/HairStrands/Source/HairStrandsCore/Private/GroomBindingBuilder.cpp
@@ -60,7 +60,7 @@ static FAutoConsoleVariableRef CVarHairStrandsBindingBuilderWarningEnable(TEXT("
 FString FGroomBindingBuilder::GetVersion()
 {
 	// Important to update the version when groom building changes
-	return TEXT("3p_7");
+	return TEXT("3p_14");
 }
 
 ///////////////////////////////////////////////////////////////////////////////////////////////////
@@ -144,7 +144,7 @@ public:
 	virtual const IMeshSectionData& GetSection(uint32 SectionIndex) const = 0;
 	virtual const FVector3f& GetVertexPosition(uint32 VertexIndex) const = 0;
 	virtual FVector2f GetVertexUV(uint32 VertexIndex, uint32 ChannelIndex) const = 0;
-	virtual void GetSectionFromVertexIndex(uint32 InVertIndex, int32& OutSectionIndex) const = 0;
+	virtual int32 GetSectionFromVertexIndex(uint32 InVertIndex) const = 0;
 	virtual ~IMeshLODData() {}
 };
 
@@ -276,12 +276,14 @@ public:
 		return FVector2f(MeshData->LODRenderData[LODIndex].StaticVertexBuffers.StaticMeshVertexBuffer.GetVertexUV(InVertexIndex, InChannelIndex));
 	}
 
-	virtual void GetSectionFromVertexIndex(uint32 InVertIndex, int32& OutSectionIndex) const override
+	virtual int32 GetSectionFromVertexIndex(uint32 InVertIndex) const override
 	{
+		int32 OutSectionIndex = 0;
 		check(MeshData);
 		check(MeshData->LODRenderData.IsValidIndex(LODIndex));
 		int32 OutVertIndex = 0;
-		return MeshData->LODRenderData[LODIndex].GetSectionFromVertexIndex(InVertIndex, OutSectionIndex, OutVertIndex);
+		MeshData->LODRenderData[LODIndex].GetSectionFromVertexIndex(InVertIndex, OutSectionIndex, OutVertIndex);
+		return OutSectionIndex;
 	}
 
 private:
@@ -473,17 +475,17 @@ public:
 		return FVector2f(MeshData.TextureCoordinates[VertexIndex]);
 	}
 
-	virtual void GetSectionFromVertexIndex(uint32 InVertIndex, int32& OutSectionIndex) const override
+	virtual int32 GetSectionFromVertexIndex(uint32 InVertIndex) const override
 	{
 		for (int32 SectionIndex = 0; SectionIndex < SectionRanges.Num(); ++SectionIndex)
 		{
 			const FRange& SectionRange = SectionRanges[SectionIndex];
 			if (InVertIndex >= SectionRange.Min && InVertIndex <= SectionRange.Max)
 			{
-				OutSectionIndex = SectionIndex;
-				return;
+				return SectionIndex;
 			}
 		}
+		return 0;
 	}
 
 private:
@@ -533,8 +535,9 @@ namespace GroomBinding_RBFWeighting
 
 	void FPointsSampler::BuildPositions(const FVector3f* PointPositions)
 	{
-		SamplePositions.SetNum(SampleIndices.Num());
-		for (int32 i = 0; i < SampleIndices.Num(); ++i)
+		const int32 NumSamples = SampleIndices.Num();
+		SamplePositions.SetNum(NumSamples);
+		for (int32 i = 0; i < NumSamples; ++i)
 		{
 			SamplePositions[i] = PointPositions[SampleIndices[i]];
 		}
@@ -686,28 +689,36 @@ namespace GroomBinding_RBFWeighting
 		}
 	}
 
-	void UpdateInterpolationWeights(const FWeightsBuilder& InterpolationWeights, const FPointsSampler& PointsSampler, const uint32 LODIndex, FHairStrandsRootData& RootDatas)
+	void UpdateInterpolationWeights(
+		const FWeightsBuilder& InterpolationWeights,
+		const FPointsSampler& PointsSampler,
+		const GroomBinding_Mesh::IMeshLODData& MeshData,
+		FHairStrandsRootData::FMeshProjectionLOD& OutRootLODData)
 	{
-		FHairStrandsRootData::FMeshProjectionLOD& CPULOD = RootDatas.MeshProjectionLODs[LODIndex];
-		CPULOD.MeshSampleIndicesBuffer.SetNum(PointsSampler.SampleIndices.Num());
-		CPULOD.MeshInterpolationWeightsBuffer.SetNum(InterpolationWeights.InverseEntries.Num());
-		CPULOD.RestSamplePositionsBuffer.SetNum(PointsSampler.SampleIndices.Num());
-
-		CPULOD.SampleCount = PointsSampler.SampleIndices.Num();
-		CPULOD.MeshSampleIndicesBuffer = PointsSampler.SampleIndices;
-		CPULOD.MeshInterpolationWeightsBuffer = InterpolationWeights.InverseEntries;
+		OutRootLODData.MeshSampleIndicesBuffer.SetNum(PointsSampler.SampleIndices.Num());
+		OutRootLODData.MeshInterpolationWeightsBuffer.SetNum(InterpolationWeights.InverseEntries.Num());
+		OutRootLODData.RestSamplePositionsBuffer.SetNum(PointsSampler.SampleIndices.Num());
+		OutRootLODData.MeshSampleSectionsBuffer.SetNum(PointsSampler.SampleIndices.Num());
+
+		OutRootLODData.SampleCount = PointsSampler.SampleIndices.Num();
+		OutRootLODData.MeshSampleIndicesBuffer = PointsSampler.SampleIndices;
+		OutRootLODData.MeshInterpolationWeightsBuffer = InterpolationWeights.InverseEntries;
 		for (int32 i = 0; i < PointsSampler.SamplePositions.Num(); ++i)
-		{
-			CPULOD.RestSamplePositionsBuffer[i] = FVector4f(PointsSampler.SamplePositions[i], 1.0f);
+		{			
+			OutRootLODData.RestSamplePositionsBuffer[i] = FVector4f(PointsSampler.SamplePositions[i], 1.0f);
+			OutRootLODData.MeshSampleSectionsBuffer[i] = MeshData.GetSectionFromVertexIndex(PointsSampler.SampleIndices[i]);
 		}
 	}
 
-	void FillLocalValidPoints(const GroomBinding_Mesh::IMeshLODData& MeshLODData, const int32 TargetSection,
-		const FHairStrandsRootData::FMeshProjectionLOD& ProjectionLOD, TArray<bool>& ValidPoints)
+	void FillLocalValidPoints(
+		const GroomBinding_Mesh::IMeshLODData& MeshLODData, 
+		const int32 TargetSection,
+		const FHairStrandsRootData::FMeshProjectionLOD& ProjectionLOD, 
+		TArray<bool>& OutValidPoints)
 	{
 		const TArray<uint32>& TriangleIndices = MeshLODData.GetIndexBuffer();
 
-		ValidPoints.Init(false, MeshLODData.GetNumVertices());
+		OutValidPoints.Init(false, MeshLODData.GetNumVertices());
 
 		const bool ValidSection = (TargetSection >= 0 && TargetSection < MeshLODData.GetNumSections());
 
@@ -722,18 +733,18 @@ namespace GroomBinding_RBFWeighting
 				for (uint32 VertexIt = 0; VertexIt < 3; ++VertexIt)
 				{
 					const uint32 VertexIndex = TriangleIndices[Section.GetBaseIndex() + 3 * TriangleIndex + VertexIt];
-					ValidPoints[VertexIndex] = (VertexIndex >= Section.GetBaseVertexIndex()) && (VertexIndex < Section.GetBaseVertexIndex() + Section.GetNumVertices());
+					OutValidPoints[VertexIndex] = (VertexIndex >= Section.GetBaseVertexIndex()) && (VertexIndex < Section.GetBaseVertexIndex() + Section.GetNumVertices());
 				}
 			}
 		}
 	}
 
-	void FillGlobalValidPoints(const GroomBinding_Mesh::IMeshLODData& MeshLODData, const int32 TargetSection, TArray<bool>& ValidPoints)
+	void FillGlobalValidPoints(const GroomBinding_Mesh::IMeshLODData& MeshLODData, const int32 TargetSection, TArray<bool>& OutValidPoints)
 	{
 		const TArray<uint32>& TriangleIndices = MeshLODData.GetIndexBuffer(); 
 		if (TargetSection >= 0 && TargetSection < MeshLODData.GetNumSections())
 		{
-			ValidPoints.Init(false, MeshLODData.GetNumVertices());
+			OutValidPoints.Init(false, MeshLODData.GetNumVertices());
 
 			const GroomBinding_Mesh::IMeshSectionData& Section = MeshLODData.GetSection(TargetSection);
 			for (uint32 TriangleIt = 0; TriangleIt < Section.GetNumTriangles(); ++TriangleIt)
@@ -741,13 +752,13 @@ namespace GroomBinding_RBFWeighting
 				for (uint32 VertexIt = 0; VertexIt < 3; ++VertexIt)
 				{
 					const uint32 VertexIndex = TriangleIndices[Section.GetBaseIndex() + 3 * TriangleIt + VertexIt];
-					ValidPoints[VertexIndex] = true;
+					OutValidPoints[VertexIndex] = true;
 				}
 			}
 		}
 		else
 		{
-			ValidPoints.Init(true, MeshLODData.GetNumVertices());
+			OutValidPoints.Init(true, MeshLODData.GetNumVertices());
 		}
 	}
 
@@ -801,7 +812,7 @@ namespace GroomBinding_RBFWeighting
 					FWeightsBuilder InterpolationWeights(SampleCount, SampleCount, PointsSampler.SamplePositions.GetData(), PointsSampler.SamplePositions.GetData());
 
 					// Guides
-					UpdateInterpolationWeights(InterpolationWeights, PointsSampler, LODIndex, Out.SimRootData);
+					UpdateInterpolationWeights(InterpolationWeights, PointsSampler, MeshLODData, Out.SimRootData.MeshProjectionLODs[LODIndex]);
 
 					// Strands
 					// No sample data, only used/available for guides
@@ -823,7 +834,7 @@ namespace GroomBinding_RBFWeighting
 				FWeightsBuilder InterpolationWeights(SampleCount, SampleCount, PointsSampler.SamplePositions.GetData(), PointsSampler.SamplePositions.GetData());
 
 				// Guides
-				UpdateInterpolationWeights(InterpolationWeights, PointsSampler, LODIndex, Out.SimRootData);
+				UpdateInterpolationWeights(InterpolationWeights, PointsSampler, MeshLODData, Out.SimRootData.MeshProjectionLODs[LODIndex]);
 
 				// Strands 
 				// No sample data, only used/available for guides
@@ -1861,8 +1872,7 @@ namespace GroomBinding_Transfer
 			for (uint32 TargetVertexIt = 0; TargetVertexIt < TargetVertexCount; ++TargetVertexIt)
 #endif
 			{
-				int32 SectionIt = 0;
-				TargetMeshLODData.GetSectionFromVertexIndex(TargetVertexIt, SectionIt);
+				const int32 SectionIt = TargetMeshLODData.GetSectionFromVertexIndex(TargetVertexIt);
 				if (SectionIt != LocalTargetSectionId)
 				{
 					OutTransferredPositions[TargetLODIndex][TargetVertexIt] = FVector3f(0,0,0);
@@ -1990,7 +2000,7 @@ static void BuildRootBulkData(
 	Out.Header.Strides.RestUniqueTrianglePositionBufferStride 	= FHairStrandsMeshTrianglePositionFormat::SizeInByte * 3; // 3 vertices per triangle
 
 	Out.Header.Strides.MeshInterpolationWeightsBufferStride 	= FHairStrandsWeightFormat::SizeInByte;
-	Out.Header.Strides.MeshSampleIndicesBufferStride 			= FHairStrandsIndexFormat::SizeInByte;
+	Out.Header.Strides.MeshSampleIndicesAndSectionsBufferStride = FHairStrandsRBFSampleIndexFormat::SizeInByte;
 	Out.Header.Strides.RestSamplePositionsBufferStride 			= FHairStrandsMeshTrianglePositionFormat::SizeInByte;
 
 	Out.Header.LODs.SetNum(MeshLODCount);
@@ -2024,14 +2034,26 @@ static void BuildRootBulkData(
 
 		if (bHasValidSamples)
 		{
+			check(In.MeshProjectionLODs[MeshLODIt].MeshSampleIndicesBuffer.Num() == In.MeshProjectionLODs[MeshLODIt].MeshSampleSectionsBuffer.Num());
+
+			const uint32 SampleCount = In.MeshProjectionLODs[MeshLODIt].MeshSampleIndicesBuffer.Num();
+			TArray<uint32> MeshSampleIndicesAndSectionBuffer;
+			MeshSampleIndicesAndSectionBuffer.SetNum(SampleCount);
+			for (uint32 SampleIt=0;SampleIt<SampleCount;++SampleIt)
+			{
+				const uint32 Index = In.MeshProjectionLODs[MeshLODIt].MeshSampleIndicesBuffer[SampleIt];
+				const uint32 SectionIndex = In.MeshProjectionLODs[MeshLODIt].MeshSampleSectionsBuffer[SampleIt];
+				MeshSampleIndicesAndSectionBuffer[SampleIt] = FHairStrandsRootUtils::PackTriangleIndex(Index, SectionIndex);
+			}
+
 			CopyToBulkData<FHairStrandsWeightFormat>(Out.Data.LODs[MeshLODIt].MeshInterpolationWeightsBuffer, In.MeshProjectionLODs[MeshLODIt].MeshInterpolationWeightsBuffer);
-			CopyToBulkData<FHairStrandsIndexFormat>(Out.Data.LODs[MeshLODIt].MeshSampleIndicesBuffer, In.MeshProjectionLODs[MeshLODIt].MeshSampleIndicesBuffer);
+			CopyToBulkData<FHairStrandsRBFSampleIndexFormat>(Out.Data.LODs[MeshLODIt].MeshSampleIndicesAndSectionsBuffer, MeshSampleIndicesAndSectionBuffer);
 			CopyToBulkData<FHairStrandsMeshTrianglePositionFormat>(Out.Data.LODs[MeshLODIt].RestSamplePositionsBuffer, In.MeshProjectionLODs[MeshLODIt].RestSamplePositionsBuffer);
 		}
 		else
 		{
 			Out.Data.LODs[MeshLODIt].MeshInterpolationWeightsBuffer.RemoveBulkData();
-			Out.Data.LODs[MeshLODIt].MeshSampleIndicesBuffer.RemoveBulkData();
+			Out.Data.LODs[MeshLODIt].MeshSampleIndicesAndSectionsBuffer.RemoveBulkData();
 			Out.Data.LODs[MeshLODIt].RestSamplePositionsBuffer.RemoveBulkData();
 		}
 	}
@@ -2065,14 +2087,29 @@ static void BuildRootData(
 
 		if (bHasValidSamples)
 		{
+			TArray<uint32> MeshSampleIndicesAndSectionBuffer;
+
 			CopyFromBulkData<FHairStrandsWeightFormat>(Out.MeshProjectionLODs[MeshLODIt].MeshInterpolationWeightsBuffer, In.Data.LODs[MeshLODIt].MeshInterpolationWeightsBuffer);
-			CopyFromBulkData<FHairStrandsIndexFormat>(Out.MeshProjectionLODs[MeshLODIt].MeshSampleIndicesBuffer, In.Data.LODs[MeshLODIt].MeshSampleIndicesBuffer);
+			CopyFromBulkData<FHairStrandsRBFSampleIndexFormat>(MeshSampleIndicesAndSectionBuffer, In.Data.LODs[MeshLODIt].MeshSampleIndicesAndSectionsBuffer);
 			CopyFromBulkData<FHairStrandsMeshTrianglePositionFormat>(Out.MeshProjectionLODs[MeshLODIt].RestSamplePositionsBuffer, In.Data.LODs[MeshLODIt].RestSamplePositionsBuffer);
+
+			// Split indices and sections
+			const uint32 SampleCount = MeshSampleIndicesAndSectionBuffer.Num();
+			Out.MeshProjectionLODs[MeshLODIt].MeshSampleIndicesBuffer.SetNum(SampleCount);
+			Out.MeshProjectionLODs[MeshLODIt].MeshSampleSectionsBuffer.SetNum(SampleCount);
+			for (uint32 SampleIt=0;SampleIt<SampleCount;++SampleIt)
+			{
+				FHairStrandsRootUtils::UnpackTriangleIndex(
+					MeshSampleIndicesAndSectionBuffer[SampleIt], 
+					Out.MeshProjectionLODs[MeshLODIt].MeshSampleIndicesBuffer[SampleIt],
+					Out.MeshProjectionLODs[MeshLODIt].MeshSampleSectionsBuffer[SampleIt]);
+			}
 		}
 		else
 		{
 			Out.MeshProjectionLODs[MeshLODIt].MeshInterpolationWeightsBuffer.Empty();
 			Out.MeshProjectionLODs[MeshLODIt].MeshSampleIndicesBuffer.Empty();
+			Out.MeshProjectionLODs[MeshLODIt].MeshSampleSectionsBuffer.Empty();
 			Out.MeshProjectionLODs[MeshLODIt].RestSamplePositionsBuffer.Empty();
 		}
 	}
diff --git a/Engine/Plugins/Runtime/HairStrands/Source/HairStrandsCore/Private/GroomResources.cpp b/Engine/Plugins/Runtime/HairStrands/Source/HairStrandsCore/Private/GroomResources.cpp
--- a/Engine/Plugins/Runtime/HairStrands/Source/HairStrandsCore/Private/GroomResources.cpp
+++ b/Engine/Plugins/Runtime/HairStrands/Source/HairStrandsCore/Private/GroomResources.cpp
@@ -1424,7 +1424,7 @@ void FHairStrandsRestRootResource::InternalAllocate(FRDGBuilder& GraphBuilder, u
 			check(WeightCount == InteroplationWeightCount); // Sanity check
 
 			InternalCreateStructuredBufferRDG_FromHairBulkData<FHairStrandsWeightFormat>(GraphBuilder, CPUData.MeshInterpolationWeightsBuffer, InteroplationWeightCount, GPUData.MeshInterpolationWeightsBuffer, ToHairResourceDebugName(HAIRSTRANDS_RESOUCE_NAME(CurveType, Hair.StrandsRestRoot_MeshInterpolationWeightsBuffer), ResourceName), OwnerName, EHairResourceUsageType::Static);
-			InternalCreateVertexBufferRDG_FromHairBulkData<FHairStrandsIndexFormat>(GraphBuilder, CPUData.MeshSampleIndicesBuffer, LODHeader.SampleCount, GPUData.MeshSampleIndicesBuffer, ToHairResourceDebugName(HAIRSTRANDS_RESOUCE_NAME(CurveType, Hair.StrandsRestRoot_MeshSampleIndicesBuffer), ResourceName), OwnerName, EHairResourceUsageType::Static);
+			InternalCreateVertexBufferRDG_FromHairBulkData<FHairStrandsRBFSampleIndexFormat>(GraphBuilder, CPUData.MeshSampleIndicesAndSectionsBuffer, LODHeader.SampleCount, GPUData.MeshSampleIndicesAndSectionsBuffer, ToHairResourceDebugName(HAIRSTRANDS_RESOUCE_NAME(CurveType, Hair.StrandsRestRoot_MeshSampleIndicesAndSectionsBuffer), ResourceName), OwnerName, EHairResourceUsageType::Static);
 			InternalCreateStructuredBufferRDG_FromHairBulkData<FHairStrandsMeshTrianglePositionFormat>(GraphBuilder, CPUData.RestSamplePositionsBuffer, LODHeader.SampleCount, GPUData.RestSamplePositionsBuffer, ToHairResourceDebugName(HAIRSTRANDS_RESOUCE_NAME(CurveType, Hair.StrandsRestRoot_RestSamplePositionsBuffer), ResourceName), OwnerName, EHairResourceUsageType::Static);
 		}
 
@@ -1445,7 +1445,7 @@ void FHairStrandsRestRootResource::InternalRelease()
 		GPUData.RestUniqueTrianglePositionBuffer.Release();
 		GPUData.SampleCount = 0;
 		GPUData.MeshInterpolationWeightsBuffer.Release();
-		GPUData.MeshSampleIndicesBuffer.Release();
+		GPUData.MeshSampleIndicesAndSectionsBuffer.Release();
 		GPUData.RestSamplePositionsBuffer.Release();
 		GPUData.AvailableCurveCount = 0;
 	}
diff --git a/Engine/Plugins/Runtime/HairStrands/Source/HairStrandsCore/Private/HairStrandsDatas.cpp b/Engine/Plugins/Runtime/HairStrands/Source/HairStrandsCore/Private/HairStrandsDatas.cpp
--- a/Engine/Plugins/Runtime/HairStrands/Source/HairStrandsCore/Private/HairStrandsDatas.cpp
+++ b/Engine/Plugins/Runtime/HairStrands/Source/HairStrandsCore/Private/HairStrandsDatas.cpp
@@ -854,7 +854,7 @@ void FHairStrandsRootBulkData::SerializeHeader(FArchive& Ar, UObject* Owner)
 		Ar << Header.Strides.RestUniqueTrianglePositionBufferStride;
 
 		Ar << Header.Strides.MeshInterpolationWeightsBufferStride;
-		Ar << Header.Strides.MeshSampleIndicesBufferStride;
+		Ar << Header.Strides.MeshSampleIndicesAndSectionsBufferStride;
 		Ar << Header.Strides.RestSamplePositionsBufferStride;
 	}
 
@@ -920,9 +920,9 @@ void FHairStrandsRootBulkData::GetResources(FQuery& Out)
 	{
 		const uint32 RBFWeightCount = FGroomRBFDeformer::GetWeightCount(HeaderLOD.SampleCount); 
 
-		Out.Add(DataLOD.MeshInterpolationWeightsBuffer, TEXT("_MeshInterpolationWeightsBuffer"), 		DataLOD.MeshInterpolationWeightsBuffer.LoadedSize, 		RBFWeightCount * Header.Strides.MeshInterpolationWeightsBufferStride); 	// Load all data
-		Out.Add(DataLOD.MeshSampleIndicesBuffer, 		TEXT("_MeshSampleIndicesBuffer"), 				DataLOD.MeshSampleIndicesBuffer.LoadedSize, 			HeaderLOD.SampleCount * Header.Strides.MeshSampleIndicesBufferStride); 	// Load all data
-		Out.Add(DataLOD.RestSamplePositionsBuffer, 		TEXT("_RestSamplePositionsBuffer"), 			DataLOD.RestSamplePositionsBuffer.LoadedSize, 			HeaderLOD.SampleCount * Header.Strides.RestSamplePositionsBufferStride);// Load all data
+		Out.Add(DataLOD.MeshInterpolationWeightsBuffer, 	TEXT("_MeshInterpolationWeightsBuffer"), 	DataLOD.MeshInterpolationWeightsBuffer.LoadedSize, 		RBFWeightCount * Header.Strides.MeshInterpolationWeightsBufferStride); 				// Load all data
+		Out.Add(DataLOD.MeshSampleIndicesAndSectionsBuffer, TEXT("_MeshSampleIndicesAndSectionsBuffer"),DataLOD.MeshSampleIndicesAndSectionsBuffer.LoadedSize, 	HeaderLOD.SampleCount * Header.Strides.MeshSampleIndicesAndSectionsBufferStride); 	// Load all data
+		Out.Add(DataLOD.RestSamplePositionsBuffer, 			TEXT("_RestSamplePositionsBuffer"), 		DataLOD.RestSamplePositionsBuffer.LoadedSize, 			HeaderLOD.SampleCount * Header.Strides.RestSamplePositionsBufferStride);			// Load all data
 	}
 }
 
@@ -950,12 +950,12 @@ void FHairStrandsRootBulkData::Reset()
 
 		// RBF
 		LOD.MeshInterpolationWeightsBuffer.RemoveBulkData();
-		LOD.MeshSampleIndicesBuffer.RemoveBulkData();
+		LOD.MeshSampleIndicesAndSectionsBuffer.RemoveBulkData();
 		LOD.RestSamplePositionsBuffer.RemoveBulkData();
 
-		LOD.MeshInterpolationWeightsBuffer 	 = FHairBulkContainer();
-		LOD.MeshSampleIndicesBuffer 		 = FHairBulkContainer();
-		LOD.RestSamplePositionsBuffer 		 = FHairBulkContainer();
+		LOD.MeshInterpolationWeightsBuffer 		 = FHairBulkContainer();
+		LOD.MeshSampleIndicesAndSectionsBuffer 	 = FHairBulkContainer();
+		LOD.RestSamplePositionsBuffer 			 = FHairBulkContainer();
 	}
 	Data.LODs.Empty();
 }
@@ -971,8 +971,8 @@ void FHairStrandsRootBulkData::ResetLoadedSize()
 		LOD.RestUniqueTrianglePositionBuffer.LoadedSize = 0;
 		      
 		// RBF
-		LOD.MeshInterpolationWeightsBuffer.LoadedSize	= 0;
-		LOD.MeshSampleIndicesBuffer.LoadedSize 			= 0;
-		LOD.RestSamplePositionsBuffer.LoadedSize 		= 0;
+		LOD.MeshInterpolationWeightsBuffer.LoadedSize	 = 0;
+		LOD.MeshSampleIndicesAndSectionsBuffer.LoadedSize= 0;
+		LOD.RestSamplePositionsBuffer.LoadedSize 		 = 0;
 	}
 }
\ No newline at end of file
diff --git a/Engine/Plugins/Runtime/HairStrands/Source/HairStrandsCore/Private/HairStrandsMeshProjection.cpp b/Engine/Plugins/Runtime/HairStrands/Source/HairStrandsCore/Private/HairStrandsMeshProjection.cpp
--- a/Engine/Plugins/Runtime/HairStrands/Source/HairStrandsCore/Private/HairStrandsMeshProjection.cpp
+++ b/Engine/Plugins/Runtime/HairStrands/Source/HairStrandsCore/Private/HairStrandsMeshProjection.cpp
@@ -502,6 +502,7 @@ private:
 	using FPermutationDomain = TShaderPermutationDomain<FPositionType>;
 
 	BEGIN_SHADER_PARAMETER_STRUCT(FParameters, )
+		SHADER_PARAMETER(uint32, MaxSectionCount)
 		SHADER_PARAMETER(uint32, MaxSampleCount)
 		SHADER_PARAMETER(uint32, MaxVertexCount)
 		SHADER_PARAMETER(uint32, PassSectionCount)
@@ -526,7 +527,7 @@ private:
 		SHADER_PARAMETER_RDG_BUFFER_SRV(Buffer, RDGVertexPositionsBuffer6)
 		SHADER_PARAMETER_RDG_BUFFER_SRV(Buffer, RDGVertexPositionsBuffer7)
 
-		SHADER_PARAMETER_RDG_BUFFER_SRV(Buffer, SampleIndicesBuffer)
+	SHADER_PARAMETER_RDG_BUFFER_SRV(Buffer, SampleIndicesAndSectionsBuffer)
 		SHADER_PARAMETER_RDG_BUFFER_UAV(RWStructuredBuffer, OutSamplePositionsBuffer)
 	END_SHADER_PARAMETER_STRUCT()
 
@@ -718,9 +719,10 @@ void AddHairStrandInitMeshSamplesPass(
 			{
 				return;
 			}
+			Parameters.MaxSectionCount = MeshData.Sections.Num();
 			Parameters.MaxVertexCount = MeshData.Sections[0].TotalVertexCount;
 			Parameters.MaxSampleCount = RestLODData.SampleCount;
-			Parameters.SampleIndicesBuffer = RegisterAsSRV(GraphBuilder, RestLODData.MeshSampleIndicesBuffer);
+			Parameters.SampleIndicesAndSectionsBuffer = RegisterAsSRV(GraphBuilder, RestLODData.MeshSampleIndicesAndSectionsBuffer);
 			Parameters.OutSamplePositionsBuffer = OutBuffer.UAV;
 
 			FHairInitMeshSamplesCS::FParameters* PassParameters = GraphBuilder.AllocParameters<FHairInitMeshSamplesCS::FParameters>();
@@ -757,12 +759,11 @@ private:
 	BEGIN_SHADER_PARAMETER_STRUCT(FParameters, )
 		SHADER_PARAMETER(uint32, MaxSampleCount)
 
-		SHADER_PARAMETER_RDG_BUFFER_SRV(Buffer, SampleIndicesBuffer)
 		SHADER_PARAMETER_RDG_BUFFER_SRV(StructuredBuffer, InterpolationWeightsBuffer)
 		SHADER_PARAMETER_RDG_BUFFER_SRV(StructuredBuffer, SampleRestPositionsBuffer)
 		SHADER_PARAMETER_RDG_BUFFER_SRV(StructuredBuffer, SampleDeformedPositionsBuffer)
 		SHADER_PARAMETER_RDG_BUFFER_UAV(RWStructuredBuffer, OutSampleDeformationsBuffer)
-		END_SHADER_PARAMETER_STRUCT()
+	END_SHADER_PARAMETER_STRUCT()
 
 public:
 	static bool ShouldCompilePermutation(const FGlobalShaderPermutationParameters& Parameters) { return IsHairStrandsSupported(EHairStrandsShaderType::All, Parameters.Platform); }
@@ -801,7 +802,6 @@ void AddHairStrandUpdateMeshSamplesPass(
 		FRDGImportedBuffer OutWeightsBuffer = Register(GraphBuilder, DeformedLODData.GetMeshSampleWeightsBuffer(FHairStrandsDeformedRootResource::FLOD::Current), ERDGImportedBufferFlags::CreateUAV);
 
 		Parameters->MaxSampleCount					= RestLODData.SampleCount;
-		Parameters->SampleIndicesBuffer				= RegisterAsSRV(GraphBuilder, RestLODData.MeshSampleIndicesBuffer);
 		Parameters->InterpolationWeightsBuffer		= RegisterAsSRV(GraphBuilder, RestLODData.MeshInterpolationWeightsBuffer);
 		Parameters->SampleRestPositionsBuffer		= RegisterAsSRV(GraphBuilder, RestLODData.RestSamplePositionsBuffer);
 		Parameters->SampleDeformedPositionsBuffer	= RegisterAsSRV(GraphBuilder, DeformedLODData.GetDeformedSamplePositionsBuffer(FHairStrandsDeformedRootResource::FLOD::Current));
diff --git a/Engine/Plugins/Runtime/HairStrands/Source/HairStrandsCore/Public/GroomResources.h b/Engine/Plugins/Runtime/HairStrands/Source/HairStrandsCore/Public/GroomResources.h
--- a/Engine/Plugins/Runtime/HairStrands/Source/HairStrandsCore/Public/GroomResources.h
+++ b/Engine/Plugins/Runtime/HairStrands/Source/HairStrandsCore/Public/GroomResources.h
@@ -416,7 +416,7 @@ struct FHairStrandsRestRootResource : public FHairCommonResource
 			Total += GetBufferTotalNumBytes(LOD.RootBarycentricBuffer);
 			Total += GetBufferTotalNumBytes(LOD.RestUniqueTrianglePositionBuffer);
 			Total += GetBufferTotalNumBytes(LOD.MeshInterpolationWeightsBuffer);
-			Total += GetBufferTotalNumBytes(LOD.MeshSampleIndicesBuffer);
+			Total += GetBufferTotalNumBytes(LOD.MeshSampleIndicesAndSectionsBuffer);
 			Total += GetBufferTotalNumBytes(LOD.RestSamplePositionsBuffer);
 		}
 		return Total;
@@ -445,7 +445,7 @@ struct FHairStrandsRestRootResource : public FHairCommonResource
 		/* Strand hair mesh interpolation matrix and sample indices */
 		uint32 SampleCount = 0;
 		FRDGExternalBuffer MeshInterpolationWeightsBuffer;
-		FRDGExternalBuffer MeshSampleIndicesBuffer;
+		FRDGExternalBuffer MeshSampleIndicesAndSectionsBuffer;
 		FRDGExternalBuffer RestSamplePositionsBuffer;
 	};
 
diff --git a/Engine/Plugins/Runtime/HairStrands/Source/HairStrandsCore/Public/HairStrandsDatas.h b/Engine/Plugins/Runtime/HairStrands/Source/HairStrandsCore/Public/HairStrandsDatas.h
--- a/Engine/Plugins/Runtime/HairStrands/Source/HairStrandsCore/Public/HairStrandsDatas.h
+++ b/Engine/Plugins/Runtime/HairStrands/Source/HairStrandsCore/Public/HairStrandsDatas.h
@@ -183,6 +183,17 @@ struct HAIRSTRANDSCORE_API FHairStrandsIndexFormat
 	static const EPixelFormat Format = PF_R32_UINT;
 };
 
+/** Hair strands RBF sample index format */
+struct HAIRSTRANDSCORE_API FHairStrandsRBFSampleIndexFormat
+{
+	using Type = uint32;
+	using BulkType = uint32;
+	static const uint32 ComponentCount = 1;
+	static const uint32 SizeInByte = sizeof(Type);
+	static const EVertexElementType VertexElementType = VET_UInt;
+	static const EPixelFormat Format = PF_R32_UINT;
+};
+
 /** Hair strands weights format */
 struct FHairStrandsWeightFormat
 {
@@ -784,11 +795,14 @@ struct FHairStrandsRootData
 		TArray<FHairStrandsWeightFormat::Type> MeshInterpolationWeightsBuffer;
 
 		/* Store the samples vertex indices */
-		TArray<FHairStrandsIndexFormat::Type> MeshSampleIndicesBuffer;
+		TArray<uint32> MeshSampleIndicesBuffer;
 
 		/* Store the samples rest positions */
 		TArray<FHairStrandsMeshTrianglePositionFormat::Type> RestSamplePositionsBuffer;
 
+		/* Store the samples rest section */
+		TArray<uint32> MeshSampleSectionsBuffer;
+
 		/* Store the mesh section indices which are relevant for this root LOD data */
 		TArray<uint32> UniqueSectionIds;
 
@@ -830,7 +844,7 @@ struct FHairStrandsRootBulkData : FHairStrandsBulkCommon
 			Total += LOD.RootBarycentricBuffer.IsBulkDataLoaded() ?				LOD.RootBarycentricBuffer.GetBulkDataSize() : 0u;
 			Total += LOD.RestUniqueTrianglePositionBuffer.IsBulkDataLoaded() ?	LOD.RestUniqueTrianglePositionBuffer.GetBulkDataSize() : 0u;
 			Total += LOD.MeshInterpolationWeightsBuffer.IsBulkDataLoaded() ?	LOD.MeshInterpolationWeightsBuffer.GetBulkDataSize() : 0u;
-			Total += LOD.MeshSampleIndicesBuffer.IsBulkDataLoaded() ?			LOD.MeshSampleIndicesBuffer.GetBulkDataSize() : 0u;
+			Total += LOD.MeshSampleIndicesAndSectionsBuffer.IsBulkDataLoaded()?	LOD.MeshSampleIndicesAndSectionsBuffer.GetBulkDataSize() : 0u;
 			Total += LOD.RestSamplePositionsBuffer.IsBulkDataLoaded() ?			LOD.RestSamplePositionsBuffer.GetBulkDataSize() : 0u;
 		}
 		return Total;
@@ -855,7 +869,7 @@ struct FHairStrandsRootBulkData : FHairStrandsBulkCommon
 			uint32 RestUniqueTrianglePositionBufferStride = 0;
 
 			uint32 MeshInterpolationWeightsBufferStride = 0;
-			uint32 MeshSampleIndicesBufferStride = 0;
+			uint32 MeshSampleIndicesAndSectionsBufferStride = 0;
 			uint32 RestSamplePositionsBufferStride = 0;
 		};
 
@@ -876,9 +890,9 @@ struct FHairStrandsRootBulkData : FHairStrandsBulkCommon
 			FHairBulkContainer RestUniqueTrianglePositionBuffer;// Rest triangle positions (per-unique-triangle)
 
 			// RBF
-			FHairBulkContainer MeshInterpolationWeightsBuffer; 	// Store the hair interpolation weights | Size = SamplesCount * SamplesCount (per-sample
-			FHairBulkContainer MeshSampleIndicesBuffer; 		// Store the samples vertex indices (per-sample)
-			FHairBulkContainer RestSamplePositionsBuffer; 		// Store the samples rest positions (per-sample)
+			FHairBulkContainer MeshInterpolationWeightsBuffer; 		// Store the hair interpolation weights | Size = SamplesCount * SamplesCount (per-sample
+			FHairBulkContainer MeshSampleIndicesAndSectionsBuffer;	// Store the samples vertex indices (per-sample)
+			FHairBulkContainer RestSamplePositionsBuffer; 			// Store the samples rest positions (per-sample)
 		};
 		TArray<FLOD> LODs;
 	} Data;
diff --git a/Engine/Shaders/Private/HairStrands/HairStrandsMesh.usf b/Engine/Shaders/Private/HairStrands/HairStrandsMesh.usf
--- a/Engine/Shaders/Private/HairStrands/HairStrandsMesh.usf
+++ b/Engine/Shaders/Private/HairStrands/HairStrandsMesh.usf
@@ -79,37 +79,25 @@ Buffer<float> RDGVertexPositionsBuffer7;
 #define VertexPositionsBuffer7 RDGVertexPositionsBuffer7
 #endif
 
+uint MaxSectionCount;
 uint MaxVertexCount;
 uint MaxSampleCount;
-Buffer<uint> SampleIndicesBuffer;
+Buffer<uint> SampleIndicesAndSectionsBuffer;
 RWStructuredBuffer<float4> OutSamplePositionsBuffer;
 
 [numthreads(128, 1, 1)]
 void MainCS(uint3 DispatchThreadId : SV_DispatchThreadID)
-{ 
+{
 	const uint SampleIndex = DispatchThreadId.x;
 	if (SampleIndex >= MaxSampleCount) 
 		return;
 
-	const uint VertexIndex = SampleIndicesBuffer[SampleIndex];
-	if (VertexIndex >= MaxVertexCount) 
+	const FHairTriangleIndex IndexAndSection = UnpackTriangleIndex(SampleIndicesAndSectionsBuffer[SampleIndex]);
+	const uint VertexIndex = IndexAndSection.TriangleIndex;
+	const uint SectionIndex = IndexAndSection.SectionIndex;
+	if (VertexIndex >= MaxVertexCount || SectionIndex >= MaxSectionCount) 
 		return;
 	
-	bool bValid = false;  
-	uint SectionIndex = 0;
-	for (uint It = 0; It < PassSectionCount; ++It)
-	{
-		if ((VertexIndex >= SectionVertexOffset(It)) && (VertexIndex < (SectionVertexOffset(It) + SectionVertexCount(It))))
-		{
-			SectionIndex = It;
-			bValid = true;
-			break;
-		}
-	}
-	
-	if (!bValid)
-		return;
-
 	float3 SamplePosition = float3(0,0,0);
 	const uint BufferIndex = SectionBufferIndex(SectionIndex);
 	const int LocalIndex = VertexIndex;
@@ -177,7 +165,6 @@ void MainCS(uint3 DispatchThreadId : SV_DispatchThreadID)
 #include "../Common.ush"
 
 uint MaxSampleCount;
-Buffer<uint> SampleIndicesBuffer;
 StructuredBuffer<float> InterpolationWeightsBuffer;
 
 StructuredBuffer<float4> SampleRestPositionsBuffer;
