Simplify RBF init sample.

diff --git a/Engine/Plugins/Runtime/HairStrands/Source/HairStrandsCore/Private/HairStrandsMeshProjection.cpp b/Engine/Plugins/Runtime/HairStrands/Source/HairStrandsCore/Private/HairStrandsMeshProjection.cpp
--- a/Engine/Plugins/Runtime/HairStrands/Source/HairStrandsCore/Private/HairStrandsMeshProjection.cpp
+++ b/Engine/Plugins/Runtime/HairStrands/Source/HairStrandsCore/Private/HairStrandsMeshProjection.cpp
@@ -238,6 +238,8 @@ bool AddHairStrandUpdateMeshTrianglesPass(
 	CommonParameters.OutUniqueTriangleCurrPosition	= OutputCurrUAV;
 	CommonParameters.OutUniqueTrianglePrevPosition	= OutputPrevUAV;
 
+	const bool bUseRDGPositionBuffer = CommonParameters.RDGMeshPositionBuffer != nullptr;
+
 	struct FSectionData
 	{
 		uint32 TotalIndexCount;
@@ -258,24 +260,19 @@ bool AddHairStrandUpdateMeshTrianglesPass(
 		SectionDatas[SectionIt].IndexBaseIndex	= MeshSectionData.IndexBaseIndex;
 		SectionDatas[SectionIt].UVsChannelOffset= MeshSectionData.UVsChannelOffset;
 		SectionDatas[SectionIt].UVsChannelCount	= MeshSectionData.UVsChannelCount;
-		SectionDatas[SectionIt].bIsSwapped		= MeshData.Sections[SectionIt].PositionBuffer != CommonParameters.MeshPositionBuffer ? 1u : 0u;
 		SectionDatas[SectionIt].Pad				= 0u;
+		SectionDatas[SectionIt].bIsSwapped		= bUseRDGPositionBuffer ? 
+			(MeshData.Sections[SectionIt].RDGPositionBuffer != CommonParameters.RDGMeshPositionBuffer ? 1u : 0u) : 
+			(MeshData.Sections[SectionIt].PositionBuffer != CommonParameters.MeshPositionBuffer ? 1u : 0u);
 
 		// Sanity check
 		check(MeshSectionData.SectionIndex == SectionIt);
 		check(MeshSectionData.UVsChannelOffset < 255);
 		check(MeshSectionData.UVsChannelCount < 255);
-#if 1 // Relaxed check, with optional buffer swap
-		check(CommonParameters.RDGMeshPositionBuffer == MeshData.Sections[SectionIt].RDGPositionBuffer || CommonParameters.RDGMeshPreviousPositionBuffer == MeshData.Sections[SectionIt].RDGPreviousPositionBuffer);
+		check(CommonParameters.RDGMeshPositionBuffer == MeshData.Sections[SectionIt].RDGPositionBuffer || CommonParameters.RDGMeshPositionBuffer == MeshData.Sections[SectionIt].RDGPreviousPositionBuffer);
 		check(CommonParameters.MeshPositionBuffer    == MeshData.Sections[SectionIt].PositionBuffer    || CommonParameters.MeshPositionBuffer    == MeshData.Sections[SectionIt].PreviousPositionBuffer);
-#else
-		check(CommonParameters.RDGMeshPositionBuffer		== MeshData.Sections[SectionIt].RDGPositionBuffer);
-		check(CommonParameters.RDGMeshPreviousPositionBuffer== MeshData.Sections[SectionIt].RDGPreviousPositionBuffer);
-		check(CommonParameters.MeshPositionBuffer			== MeshData.Sections[SectionIt].PositionBuffer);
-		check(CommonParameters.MeshPreviousPositionBuffer	== MeshData.Sections[SectionIt].PreviousPositionBuffer);
-#endif
-		check(CommonParameters.MeshIndexBuffer				== MeshData.Sections[SectionIt].IndexBuffer);
-		check(CommonParameters.MeshUVsBuffer				== MeshData.Sections[SectionIt].UVsBuffer);
+		check(CommonParameters.MeshIndexBuffer		 == MeshData.Sections[SectionIt].IndexBuffer);
+		check(CommonParameters.MeshUVsBuffer		 == MeshData.Sections[SectionIt].UVsBuffer);
 	}
 
 	// If no previous position buffer available, reusing the current position buffers
@@ -291,7 +288,6 @@ bool AddHairStrandUpdateMeshTrianglesPass(
 	FRDGBufferRef SectionBuffer = CreateStructuredBuffer(GraphBuilder, TEXT("Hair.SkelMeshSectionBuffer"), sizeof(FSectionData),  SectionDatas.Num(), SectionDatas.GetData(), sizeof(FSectionData) * SectionDatas.Num());
 	CommonParameters.MeshSectionBuffer = GraphBuilder.CreateSRV(SectionBuffer);
 
-	const bool bUseRDGPositionBuffer = CommonParameters.RDGMeshPositionBuffer != nullptr;
 	const bool bComputePreviousDeformedPosition = OutputPrevUAV != nullptr;
 	{
 		FHairUpdateMeshTriangleCS::FParameters* PassParameters = GraphBuilder.AllocParameters<FHairUpdateMeshTriangleCS::FParameters>();
@@ -493,7 +489,6 @@ void AddHairCardsRBFInterpolationPass(
 class FHairInitMeshSamplesCS : public FGlobalShader
 {
 public:
-	const static uint32 SectionArrayCount = 16; // This defines the number of sections managed for each iteration pass
 private:
 	DECLARE_GLOBAL_SHADER(FHairInitMeshSamplesCS);
 	SHADER_USE_PARAMETER_STRUCT(FHairInitMeshSamplesCS, FGlobalShader);
@@ -505,52 +500,29 @@ private:
 		SHADER_PARAMETER(uint32, MaxSectionCount)
 		SHADER_PARAMETER(uint32, MaxSampleCount)
 		SHADER_PARAMETER(uint32, MaxVertexCount)
-		SHADER_PARAMETER(uint32, PassSectionCount)
-
-		SHADER_PARAMETER_ARRAY(FUintVector4, PackedSectionScalars, [SectionArrayCount])
-
-		SHADER_PARAMETER_SRV(Buffer, VertexPositionsBuffer0)
-		SHADER_PARAMETER_SRV(Buffer, VertexPositionsBuffer1)
-		SHADER_PARAMETER_SRV(Buffer, VertexPositionsBuffer2)
-		SHADER_PARAMETER_SRV(Buffer, VertexPositionsBuffer3)
-		SHADER_PARAMETER_SRV(Buffer, VertexPositionsBuffer4)
-		SHADER_PARAMETER_SRV(Buffer, VertexPositionsBuffer5)
-		SHADER_PARAMETER_SRV(Buffer, VertexPositionsBuffer6)
-		SHADER_PARAMETER_SRV(Buffer, VertexPositionsBuffer7)
-
-		SHADER_PARAMETER_RDG_BUFFER_SRV(Buffer, RDGVertexPositionsBuffer0)
-		SHADER_PARAMETER_RDG_BUFFER_SRV(Buffer, RDGVertexPositionsBuffer1)
-		SHADER_PARAMETER_RDG_BUFFER_SRV(Buffer, RDGVertexPositionsBuffer2)
-		SHADER_PARAMETER_RDG_BUFFER_SRV(Buffer, RDGVertexPositionsBuffer3)
-		SHADER_PARAMETER_RDG_BUFFER_SRV(Buffer, RDGVertexPositionsBuffer4)
-		SHADER_PARAMETER_RDG_BUFFER_SRV(Buffer, RDGVertexPositionsBuffer5)
-		SHADER_PARAMETER_RDG_BUFFER_SRV(Buffer, RDGVertexPositionsBuffer6)
-		SHADER_PARAMETER_RDG_BUFFER_SRV(Buffer, RDGVertexPositionsBuffer7)
-
-	SHADER_PARAMETER_RDG_BUFFER_SRV(Buffer, SampleIndicesAndSectionsBuffer)
+
+		SHADER_PARAMETER_SRV(Buffer, MeshPositionBuffer0)
+		SHADER_PARAMETER_SRV(Buffer, MeshPositionBuffer1)
+		SHADER_PARAMETER_RDG_BUFFER_SRV(Buffer, RDGMeshPositionBuffer0)
+		SHADER_PARAMETER_RDG_BUFFER_SRV(Buffer, RDGMeshPositionBuffer1)
+		SHADER_PARAMETER_RDG_BUFFER_SRV(Buffer, SampleIndicesAndSectionsBuffer)
+		SHADER_PARAMETER_RDG_BUFFER_SRV(StructuredBuffer, MeshSectionBuffer)
 		SHADER_PARAMETER_RDG_BUFFER_UAV(RWStructuredBuffer, OutSamplePositionsBuffer)
 	END_SHADER_PARAMETER_STRUCT()
 
 public:
+	static uint32 GetGroupSize() { return 128; }
 	static bool ShouldCompilePermutation(const FGlobalShaderPermutationParameters& Parameters) { return IsHairStrandsSupported(EHairStrandsShaderType::All, Parameters.Platform); }
 	static void ModifyCompilationEnvironment(const FGlobalShaderPermutationParameters& Parameters, FShaderCompilerEnvironment& OutEnvironment)
 	{
 		FGlobalShader::ModifyCompilationEnvironment(Parameters, OutEnvironment);
-		OutEnvironment.SetDefine(TEXT("MAX_SECTION_COUNT"), SectionArrayCount);
 		OutEnvironment.SetDefine(TEXT("SHADER_SAMPLE_INIT"), 1);
+		OutEnvironment.SetDefine(TEXT("GROUP_SIZE"), GetGroupSize());
 	}
 };
 
 IMPLEMENT_GLOBAL_SHADER(FHairInitMeshSamplesCS, "/Engine/Private/HairStrands/HairStrandsMesh.usf", "MainCS", SF_Compute);
 
-// Manual packing for mesh section scalars.  MUST MATCH WITH HairStrandsMesh.usf
-// PackedSectionScalars: SectionVertexOffset, SectionVertexCount, SectionBufferIndex, *free*
-inline void SetSectionVertexOffset(FHairInitMeshSamplesCS::FParameters& Params, uint32 SectionIndex, uint32 SectionVertexOffset) { Params.PackedSectionScalars[SectionIndex].X = SectionVertexOffset; }
-inline void SetSectionVertexCount(FHairInitMeshSamplesCS::FParameters& Params, uint32 SectionIndex, uint32 SectionVertexCount)   { Params.PackedSectionScalars[SectionIndex].Y = SectionVertexCount; }
-inline void SetSectionBufferIndex(FHairInitMeshSamplesCS::FParameters& Params, uint32 SectionIndex, uint32 SectionBufferIndex)   { Params.PackedSectionScalars[SectionIndex].Z = SectionBufferIndex; }
-
-
-
 void AddHairStrandInitMeshSamplesPass(
 	FRDGBuilder& GraphBuilder,
 	FGlobalShaderMap* ShaderMap,
@@ -559,191 +531,81 @@ void AddHairStrandInitMeshSamplesPass(
 	FHairStrandsRestRootResource* RestResources,
 	FHairStrandsDeformedRootResource* DeformedResources)
 {
-	if (LODIndex < 0)
+	if (!RestResources->LODs.IsValidIndex(LODIndex))
 	{
 		return;
 	}
-
-	if ((LODIndex >= RestResources->LODs.Num() || LODIndex >= DeformedResources->LODs.Num()))
-	{
-		return;
-	}
-
+	
 	FHairStrandsRestRootResource::FLOD& RestLODData = RestResources->LODs[LODIndex];
 	check(RestLODData.LODIndex == LODIndex);
 
 	const uint32 SectionCount = MeshData.Sections.Num();
-	const uint32 MaxSupportedSectionCount = GetHairStrandsMaxSectionCount();
-	check(SectionCount < MaxSupportedSectionCount);
-	if (SectionCount == 0 || SectionCount >= MaxSupportedSectionCount)
+	if (SectionCount == 0 || RestLODData.SampleCount == 0)
 	{
 		return;
 	}
 
-	// When the number of section of a mesh is above FHairUpdateMeshTriangleCS::SectionArrayCount, the update is split into several passes
-	const TArray<uint32>& ValidSectionIndices = RestResources->BulkData.GetValidSectionIndices(LODIndex);
-	const uint32 ValidSectionCount = ValidSectionIndices.Num();
-	const uint32 PassCount = FMath::DivideAndRoundUp(ValidSectionCount, FHairInitMeshSamplesCS::SectionArrayCount);
-
-	if (SectionCount > 0 && RestLODData.SampleCount > 0)
+	FRDGImportedBuffer OutBuffer = Register(GraphBuilder, DeformedResources->LODs[LODIndex].GetDeformedSamplePositionsBuffer(FHairStrandsDeformedRootResource::FLOD::Current), ERDGImportedBufferFlags::CreateUAV);
+	
+	FHairInitMeshSamplesCS::FParameters CommonParameters;
+	CommonParameters.MaxSampleCount 				= RestLODData.SampleCount;
+	CommonParameters.MaxSectionCount 				= MeshData.Sections.Num();
+	CommonParameters.MaxVertexCount 				= MeshData.Sections[0].TotalVertexCount;
+	CommonParameters.RDGMeshPositionBuffer0			= MeshData.Sections[0].RDGPositionBuffer;
+	CommonParameters.RDGMeshPositionBuffer1		 	= MeshData.Sections[0].RDGPreviousPositionBuffer;
+	CommonParameters.MeshPositionBuffer0			= MeshData.Sections[0].PositionBuffer;
+	CommonParameters.MeshPositionBuffer1			= MeshData.Sections[0].PreviousPositionBuffer;
+	CommonParameters.SampleIndicesAndSectionsBuffer = RegisterAsSRV(GraphBuilder, RestLODData.MeshSampleIndicesAndSectionsBuffer);
+	CommonParameters.OutSamplePositionsBuffer 		= OutBuffer.UAV;
+
+	const bool bUseRDGPositionBuffer = CommonParameters.RDGMeshPositionBuffer0 != nullptr;
+
+	// Create buffer with section data for look up in shader
+	struct FSectionData
 	{
-		FRDGImportedBuffer OutBuffer;
-		{
-			FHairStrandsDeformedRootResource::FLOD& DeformedLODData = DeformedResources->LODs[LODIndex];
-			check(DeformedLODData.LODIndex == LODIndex);
+		uint32 bIsSwapped;
+	};
+	TArray<FSectionData> SectionDatas;
+	SectionDatas.SetNum(SectionCount);
+	for (uint32 SectionIt = 0; SectionIt < SectionCount; ++SectionIt)
+	{
+		SectionDatas[SectionIt].bIsSwapped = 
+			bUseRDGPositionBuffer ?
+			(MeshData.Sections[SectionIt].RDGPositionBuffer != CommonParameters.RDGMeshPositionBuffer0 ? 1u : 0u):
+			(MeshData.Sections[SectionIt].PositionBuffer    != CommonParameters.MeshPositionBuffer0    ? 1u : 0u);
 
-			OutBuffer = Register(GraphBuilder, DeformedLODData.GetDeformedSamplePositionsBuffer(FHairStrandsDeformedRootResource::FLOD::Current), ERDGImportedBufferFlags::CreateUAV);
-		}
-		for (uint32 PassIt = 0; PassIt < PassCount; ++PassIt)
-		{
-			FHairInitMeshSamplesCS::FParameters Parameters;
-
-			for (uint32 DataIndex = 0; DataIndex < FHairInitMeshSamplesCS::SectionArrayCount; ++DataIndex)
-			{
-				SetSectionBufferIndex(Parameters, DataIndex, 0);
-				SetSectionVertexOffset(Parameters, DataIndex, 0);
-				SetSectionVertexCount(Parameters, DataIndex, 0);
-			}
-
-			const int32 PassSectionStart = PassIt * FHairInitMeshSamplesCS::SectionArrayCount;
-			const int32 PassSectionCount = FMath::Min(ValidSectionCount - PassSectionStart, FHairInitMeshSamplesCS::SectionArrayCount);
-			Parameters.PassSectionCount = PassSectionCount;
-
-			struct FMeshSectionBuffers
-			{
-				uint32 SectionBufferIndex = 0;
-				FRDGBufferSRVRef RDGPositionBuffer = nullptr;
-				FRDGBufferSRVRef RDGPreviousPositionBuffer = nullptr;
-				FRHIShaderResourceView* PositionBuffer = nullptr;
-				FRHIShaderResourceView* PreviousPositionBuffer = nullptr;
-			};
-			TMap<FRHIShaderResourceView*, FMeshSectionBuffers>	UniqueMeshSectionBuffers;
-			TMap<FRDGBufferSRVRef, FMeshSectionBuffers>		UniqueMeshSectionBuffersRDG;
-
-			#define SETPARAMETERS(OutParameters, InMeshSectionData, Index) \
-				OutParameters.RDGVertexPositionsBuffer##Index = InMeshSectionData.RDGPositionBuffer; \
-				OutParameters.VertexPositionsBuffer##Index = InMeshSectionData.PositionBuffer;
-
-			auto SetMeshSectionBuffers = [&Parameters](uint32 UniqueIndex, const FHairStrandsProjectionMeshData::Section& MeshSectionData)
-			{
-				switch (UniqueIndex)
-				{
-				case 0: SETPARAMETERS(Parameters, MeshSectionData, 0); break;
-				case 1: SETPARAMETERS(Parameters, MeshSectionData, 1); break;
-				case 2: SETPARAMETERS(Parameters, MeshSectionData, 2); break;
-				case 3: SETPARAMETERS(Parameters, MeshSectionData, 3); break;
-				case 4: SETPARAMETERS(Parameters, MeshSectionData, 4); break;
-				case 5: SETPARAMETERS(Parameters, MeshSectionData, 5); break;
-				case 6: SETPARAMETERS(Parameters, MeshSectionData, 6); break;
-				case 7: SETPARAMETERS(Parameters, MeshSectionData, 7); break;
-				}
-			};
-
-			#undef SETPARAMETERS
-
-			uint32 UniqueMeshSectionBufferIndex = 0;
-			bool bUseRDGPositionBuffer = false;
-			for (int32 SectionStartIt = PassSectionStart, SectionItEnd = PassSectionStart + PassSectionCount; SectionStartIt < SectionItEnd; ++SectionStartIt)
-			{
-				const int32 SectionIt = SectionStartIt - PassSectionStart;
-				const int32 SectionIndex = ValidSectionIndices[SectionStartIt];
-
-				const FHairStrandsProjectionMeshData::Section& MeshSectionData = MeshData.Sections[SectionIndex];
-
-				SetSectionVertexOffset(Parameters, SectionIt, MeshSectionData.VertexBaseIndex);
-				SetSectionVertexCount(Parameters, SectionIt, MeshSectionData.NumVertices);
-
-				const FMeshSectionBuffers* Buffers = nullptr;
-				if (MeshSectionData.PositionBuffer)
-				{
-					Buffers = UniqueMeshSectionBuffers.Find(MeshSectionData.PositionBuffer);
-				}
-				else if (MeshSectionData.RDGPositionBuffer)
-				{
-					Buffers = UniqueMeshSectionBuffersRDG.Find(MeshSectionData.RDGPositionBuffer);
-				}
-				else
-				{
-					check(false); // Should never happen
-					continue;
-				}
-				if (Buffers != nullptr)
-				{
-					SetSectionBufferIndex(Parameters, SectionIt, Buffers->SectionBufferIndex);
-				}
-				else
-				{
-					// Only support 8 unique different buffer at the moment
-					check(UniqueMeshSectionBufferIndex < 8);
-					SetMeshSectionBuffers(UniqueMeshSectionBufferIndex, MeshSectionData);
-
-					FMeshSectionBuffers Entry;
-					Entry.SectionBufferIndex = UniqueMeshSectionBufferIndex;
-					Entry.RDGPositionBuffer = MeshSectionData.RDGPositionBuffer;
-					Entry.PositionBuffer = MeshSectionData.PositionBuffer;
-
-					if (MeshSectionData.PositionBuffer)
-					{
-						UniqueMeshSectionBuffers.Add(MeshSectionData.PositionBuffer, Entry);
-					}
-					else if (MeshSectionData.RDGPositionBuffer)
-					{
-						UniqueMeshSectionBuffersRDG.Add(MeshSectionData.RDGPositionBuffer, Entry);
-					}
-
-					SetSectionBufferIndex(Parameters, SectionIt, UniqueMeshSectionBufferIndex);
-					++UniqueMeshSectionBufferIndex;
-				}
-
-				// Sanity check
-				// If one of the input is using RDG position, we expect all mesh sections to use RDG input
-				if (bUseRDGPositionBuffer)
-				{
-					check(MeshSectionData.RDGPositionBuffer != nullptr);
-				}
-				else if (MeshSectionData.RDGPositionBuffer != nullptr)
-				{
-					bUseRDGPositionBuffer = true;
-				}
-			}
-
-			if (MeshData.Sections.Num() > 0)
-			{
-				for (uint32 Index = UniqueMeshSectionBufferIndex; Index < 8; ++Index)
-				{
-					SetMeshSectionBuffers(Index, MeshData.Sections[0]);
-				}
-			}
-
-			if (UniqueMeshSectionBufferIndex == 0)
-			{
-				return;
-			}
-			Parameters.MaxSectionCount = MeshData.Sections.Num();
-			Parameters.MaxVertexCount = MeshData.Sections[0].TotalVertexCount;
-			Parameters.MaxSampleCount = RestLODData.SampleCount;
-			Parameters.SampleIndicesAndSectionsBuffer = RegisterAsSRV(GraphBuilder, RestLODData.MeshSampleIndicesAndSectionsBuffer);
-			Parameters.OutSamplePositionsBuffer = OutBuffer.UAV;
-
-			FHairInitMeshSamplesCS::FParameters* PassParameters = GraphBuilder.AllocParameters<FHairInitMeshSamplesCS::FParameters>();
-			*PassParameters = Parameters;
-
-			FHairInitMeshSamplesCS::FPermutationDomain PermutationVector;
-			PermutationVector.Set<FHairInitMeshSamplesCS::FPositionType>(bUseRDGPositionBuffer ? 1 : 0);
-
-			const FIntVector DispatchGroupCount = FComputeShaderUtils::GetGroupCount(RestLODData.SampleCount, 128);
-			check(DispatchGroupCount.X <= GRHIMaxDispatchThreadGroupsPerDimension.X);
-			TShaderMapRef<FHairInitMeshSamplesCS> ComputeShader(ShaderMap, PermutationVector);
-
-			FComputeShaderUtils::AddPass(
-				GraphBuilder,
-				RDG_EVENT_NAME("HairStrandsInitMeshSamples"),
-				ComputeShader,
-				PassParameters,
-				DispatchGroupCount);
-		}
-		GraphBuilder.SetBufferAccessFinal(OutBuffer.Buffer, ERHIAccess::SRVMask);
+		// Sanity check
+		check(MeshData.Sections[SectionIt].SectionIndex == SectionIt);
+		check(CommonParameters.RDGMeshPositionBuffer0 == MeshData.Sections[SectionIt].RDGPositionBuffer || CommonParameters.RDGMeshPositionBuffer0 == MeshData.Sections[SectionIt].RDGPreviousPositionBuffer);
+		check(CommonParameters.MeshPositionBuffer0    == MeshData.Sections[SectionIt].PositionBuffer    || CommonParameters.MeshPositionBuffer0    == MeshData.Sections[SectionIt].PreviousPositionBuffer);
+	}
+	CommonParameters.MeshSectionBuffer = GraphBuilder.CreateSRV(CreateStructuredBuffer(GraphBuilder, TEXT("Hair.SkelMeshSectionBuffer"), sizeof(FSectionData),  SectionDatas.Num(), SectionDatas.GetData(), sizeof(FSectionData) * SectionDatas.Num()));
+
+	// If no previous position buffer available, reusing the current position buffers
+	if (CommonParameters.MeshPositionBuffer1 == nullptr)
+	{
+		CommonParameters.MeshPositionBuffer1 = CommonParameters.MeshPositionBuffer0;
 	}
+	if (CommonParameters.RDGMeshPositionBuffer1 == nullptr)
+	{
+		CommonParameters.RDGMeshPositionBuffer1 = CommonParameters.RDGMeshPositionBuffer0;
+	}
+
+	FHairInitMeshSamplesCS::FParameters* PassParameters = GraphBuilder.AllocParameters<FHairInitMeshSamplesCS::FParameters>();
+	*PassParameters = CommonParameters;
+
+	FHairInitMeshSamplesCS::FPermutationDomain PermutationVector;
+	PermutationVector.Set<FHairInitMeshSamplesCS::FPositionType>(bUseRDGPositionBuffer ? 1 : 0);
+
+	const FIntVector DispatchGroupCount = FComputeShaderUtils::GetGroupCount(CommonParameters.MaxSampleCount, FHairInitMeshSamplesCS::GetGroupSize());
+	check(DispatchGroupCount.X <= GRHIMaxDispatchThreadGroupsPerDimension.X);
+	TShaderMapRef<FHairInitMeshSamplesCS> ComputeShader(ShaderMap, PermutationVector);
+	FComputeShaderUtils::AddPass(
+		GraphBuilder,
+		RDG_EVENT_NAME("HairStrands::TriangleMeshUpdate(%s)", bUseRDGPositionBuffer ? TEXT("RDGBuffer") : TEXT("SkinCacheBuffer")),
+		ComputeShader,
+		PassParameters,
+		DispatchGroupCount);
 }
 
 ///////////////////////////////////////////////////////////////////////////////////////////////////
diff --git a/Engine/Shaders/Private/HairStrands/HairStrandsMesh.usf b/Engine/Shaders/Private/HairStrands/HairStrandsMesh.usf
--- a/Engine/Shaders/Private/HairStrands/HairStrandsMesh.usf
+++ b/Engine/Shaders/Private/HairStrands/HairStrandsMesh.usf
@@ -37,55 +37,24 @@ void MainHairMeshesCS(uint3 DispatchThreadId : SV_DispatchThreadID)
 
 #if SHADER_SAMPLE_INIT
 
-#ifndef MAX_SECTION_COUNT
-#error MAX_SECTION_COUNT needs to be defined
-#endif
-
-uint PassSectionCount;
-
-// Manual packing for mesh section scalars.  MUST MATCH WITH HairStrandsMeshProjection.cpp
-// PackedSectionScalars: SectionVertexOffset, SectionVertexCount, SectionBufferIndex, *free*
-uint4 PackedSectionScalars[MAX_SECTION_COUNT];
-#define SectionVertexOffset(SectionIndex) PackedSectionScalars[SectionIndex].x
-#define SectionVertexCount(SectionIndex)  PackedSectionScalars[SectionIndex].y
-#define SectionBufferIndex(SectionIndex)  PackedSectionScalars[SectionIndex].z
-
 #if PERMUTATION_POSITION_TYPE == 0
-Buffer<float> VertexPositionsBuffer0;
-Buffer<float> VertexPositionsBuffer1;
-Buffer<float> VertexPositionsBuffer2;
-Buffer<float> VertexPositionsBuffer3;
-Buffer<float> VertexPositionsBuffer4;
-Buffer<float> VertexPositionsBuffer5;
-Buffer<float> VertexPositionsBuffer6;
-Buffer<float> VertexPositionsBuffer7;
+Buffer<float> MeshPositionBuffer0;
+Buffer<float> MeshPositionBuffer1;
 #elif PERMUTATION_POSITION_TYPE == 1
-Buffer<float> RDGVertexPositionsBuffer0;
-Buffer<float> RDGVertexPositionsBuffer1;
-Buffer<float> RDGVertexPositionsBuffer2;
-Buffer<float> RDGVertexPositionsBuffer3;
-Buffer<float> RDGVertexPositionsBuffer4;
-Buffer<float> RDGVertexPositionsBuffer5;
-Buffer<float> RDGVertexPositionsBuffer6;
-Buffer<float> RDGVertexPositionsBuffer7;
-
-#define VertexPositionsBuffer0 RDGVertexPositionsBuffer0
-#define VertexPositionsBuffer1 RDGVertexPositionsBuffer1
-#define VertexPositionsBuffer2 RDGVertexPositionsBuffer2
-#define VertexPositionsBuffer3 RDGVertexPositionsBuffer3
-#define VertexPositionsBuffer4 RDGVertexPositionsBuffer4
-#define VertexPositionsBuffer5 RDGVertexPositionsBuffer5
-#define VertexPositionsBuffer6 RDGVertexPositionsBuffer6
-#define VertexPositionsBuffer7 RDGVertexPositionsBuffer7
+Buffer<float> RDGMeshPositionBuffer0;
+Buffer<float> RDGMeshPositionBuffer1;
+#define MeshPositionBuffer0 RDGMeshPositionBuffer0
+#define MeshPositionBuffer1 RDGMeshPositionBuffer1
 #endif
 
 uint MaxSectionCount;
 uint MaxVertexCount;
 uint MaxSampleCount;
 Buffer<uint> SampleIndicesAndSectionsBuffer;
+StructuredBuffer<uint> MeshSectionBuffer;
 RWStructuredBuffer<float4> OutSamplePositionsBuffer;
 
-[numthreads(128, 1, 1)]
+[numthreads(GROUP_SIZE, 1, 1)]
 void MainCS(uint3 DispatchThreadId : SV_DispatchThreadID)
 {
 	const uint SampleIndex = DispatchThreadId.x;
@@ -98,58 +67,21 @@ void MainCS(uint3 DispatchThreadId : SV_DispatchThreadID)
 	if (VertexIndex >= MaxVertexCount || SectionIndex >= MaxSectionCount) 
 		return;
 	
+	const bool bIsSwapped = MeshSectionBuffer[SectionIndex] > 0;
+
 	float3 SamplePosition = float3(0,0,0);
-	const uint BufferIndex = SectionBufferIndex(SectionIndex);
-	const int LocalIndex = VertexIndex;
-	if (BufferIndex == 0)
-	{
-		SamplePosition.x = VertexPositionsBuffer0.Load(LocalIndex * 3 + 0);
-		SamplePosition.y = VertexPositionsBuffer0.Load(LocalIndex * 3 + 1);
-		SamplePosition.z = VertexPositionsBuffer0.Load(LocalIndex * 3 + 2);
-	}
-	else if (BufferIndex == 1)
-	{
-		SamplePosition.x = VertexPositionsBuffer1.Load(LocalIndex * 3 + 0);
-		SamplePosition.y = VertexPositionsBuffer1.Load(LocalIndex * 3 + 1);
-		SamplePosition.z = VertexPositionsBuffer1.Load(LocalIndex * 3 + 2);
-	}
-	else if (BufferIndex == 2)
+	if (bIsSwapped)
 	{
-		SamplePosition.x = VertexPositionsBuffer2.Load(LocalIndex * 3 + 0);
-		SamplePosition.y = VertexPositionsBuffer2.Load(LocalIndex * 3 + 1);
-		SamplePosition.z = VertexPositionsBuffer2.Load(LocalIndex * 3 + 2);
+		SamplePosition.x = MeshPositionBuffer0.Load(VertexIndex * 3 + 0);
+		SamplePosition.y = MeshPositionBuffer0.Load(VertexIndex * 3 + 1);
+		SamplePosition.z = MeshPositionBuffer0.Load(VertexIndex * 3 + 2);
 	}
-	else if (BufferIndex == 3)
-	{
-		SamplePosition.x = VertexPositionsBuffer3.Load(LocalIndex * 3 + 0);
-		SamplePosition.y = VertexPositionsBuffer3.Load(LocalIndex * 3 + 1);
-		SamplePosition.z = VertexPositionsBuffer3.Load(LocalIndex * 3 + 2);
-	}
-	else if (BufferIndex == 4)
-	{
-		SamplePosition.x = VertexPositionsBuffer4.Load(LocalIndex * 3 + 0);
-		SamplePosition.y = VertexPositionsBuffer4.Load(LocalIndex * 3 + 1);
-		SamplePosition.z = VertexPositionsBuffer4.Load(LocalIndex * 3 + 2);
-	}
-	else if (BufferIndex == 5)
-	{
-		SamplePosition.x = VertexPositionsBuffer5.Load(LocalIndex * 3 + 0);
-		SamplePosition.y = VertexPositionsBuffer5.Load(LocalIndex * 3 + 1);
-		SamplePosition.z = VertexPositionsBuffer5.Load(LocalIndex * 3 + 2);
-	}
-	else if (BufferIndex == 6)
-	{
-		SamplePosition.x = VertexPositionsBuffer6.Load(LocalIndex * 3 + 0);
-		SamplePosition.y = VertexPositionsBuffer6.Load(LocalIndex * 3 + 1);
-		SamplePosition.z = VertexPositionsBuffer6.Load(LocalIndex * 3 + 2);
-	}
-	else if (BufferIndex == 7)
+	else
 	{
-		SamplePosition.x = VertexPositionsBuffer7.Load(LocalIndex * 3 + 0);
-		SamplePosition.y = VertexPositionsBuffer7.Load(LocalIndex * 3 + 1);
-		SamplePosition.z = VertexPositionsBuffer7.Load(LocalIndex * 3 + 2);
+		SamplePosition.x = MeshPositionBuffer1.Load(VertexIndex * 3 + 0);
+		SamplePosition.y = MeshPositionBuffer1.Load(VertexIndex * 3 + 1);
+		SamplePosition.z = MeshPositionBuffer1.Load(VertexIndex * 3 + 2);
 	}
-
 	OutSamplePositionsBuffer[SampleIndex] = float4(SamplePosition,1);
 }
 
