Change hair manual skin cache to run section update within a render pass to avoid overhead and avoid UAV barrier.

#prefllight 6569dba5cd3c9e408dab1666

diff --git a/Engine/Plugins/Runtime/HairStrands/Source/HairStrandsCore/Private/GroomDebug.cpp b/Engine/Plugins/Runtime/HairStrands/Source/HairStrandsCore/Private/GroomDebug.cpp
--- a/Engine/Plugins/Runtime/HairStrands/Source/HairStrandsCore/Private/GroomDebug.cpp
+++ b/Engine/Plugins/Runtime/HairStrands/Source/HairStrandsCore/Private/GroomDebug.cpp
@@ -82,7 +82,7 @@ static void GetGroomInterpolationData(
 	FSceneInterface* Scene,
 	const FHairStrandsInstances& Instances,
 	const EHairStrandsProjectionMeshType MeshType,
-	FHairStrandsProjectionMeshData::LOD& OutGeometries)
+	FHairStrandsProjectionMeshData::FLOD& OutGeometries)
 {
 	for (FHairStrandsInstance* AbstractInstance : Instances)
 	{
@@ -99,7 +99,7 @@ static void GetGroomInterpolationData(
 		{
 			for (int32 SectionIndex = 0; SectionIndex < CachedGeometry.Sections.Num(); ++SectionIndex)
 			{
-				FHairStrandsProjectionMeshData::Section OutSection = ConvertMeshSection(CachedGeometry, SectionIndex);
+				FHairStrandsProjectionMeshData::FSection OutSection = ConvertMeshSection(CachedGeometry, SectionIndex);
 				if (MeshType == EHairStrandsProjectionMeshType::RestMesh)
 				{					
 					// If the mesh has some mesh-tranferred data, we display that otherwise we use the rest data
@@ -183,7 +183,7 @@ static void AddDebugProjectionMeshPass(
 	const FIntRect Viewport,
 	const TUniformBufferRef<FViewUniformShaderParameters>& ViewUniformBuffer,
 	const EHairStrandsProjectionMeshType MeshType,
-	FHairStrandsProjectionMeshData::Section& MeshSectionData)
+	FHairStrandsProjectionMeshData::FSection& MeshSectionData)
 {
 	if (!View || !ShaderPrint::IsSupported(View->GetShaderPlatform()))
 	{
@@ -1404,9 +1404,9 @@ void RunHairStrandsDebug(
 			{
 				auto RenderMeshProjection = [ShaderMap, LocalView, Scene, ShaderPrintData, Viewport, &ViewUniformBuffer, Instances, &GraphBuilder](FRDGBuilder& LocalGraphBuilder, EHairStrandsProjectionMeshType MeshType)
 				{
-					FHairStrandsProjectionMeshData::LOD MeshProjectionLODData;
+					FHairStrandsProjectionMeshData::FLOD MeshProjectionLODData;
 					GetGroomInterpolationData(GraphBuilder, ShaderMap, Scene, Instances, MeshType, MeshProjectionLODData);
-					for (FHairStrandsProjectionMeshData::Section& Section : MeshProjectionLODData.Sections)
+					for (FHairStrandsProjectionMeshData::FSection& Section : MeshProjectionLODData.Sections)
 					{
 						AddDebugProjectionMeshPass(LocalGraphBuilder, LocalView, ShaderMap, ShaderPrintData, Viewport, ViewUniformBuffer, MeshType, Section);
 					}
diff --git a/Engine/Plugins/Runtime/HairStrands/Source/HairStrandsCore/Private/GroomGeometryCache.cpp b/Engine/Plugins/Runtime/HairStrands/Source/HairStrandsCore/Private/GroomGeometryCache.cpp
--- a/Engine/Plugins/Runtime/HairStrands/Source/HairStrandsCore/Private/GroomGeometryCache.cpp
+++ b/Engine/Plugins/Runtime/HairStrands/Source/HairStrandsCore/Private/GroomGeometryCache.cpp
@@ -19,10 +19,10 @@
 
 ///////////////////////////////////////////////////////////////////////////////////////////////////
 
-FHairStrandsProjectionMeshData::Section ConvertMeshSection(FCachedGeometry const& InCachedGeometry, int32 InSectionIndex)
+FHairStrandsProjectionMeshData::FSection ConvertMeshSection(FCachedGeometry const& InCachedGeometry, int32 InSectionIndex)
 {
 	FCachedGeometry::Section const& In = InCachedGeometry.Sections[InSectionIndex];
-	FHairStrandsProjectionMeshData::Section Out;
+	FHairStrandsProjectionMeshData::FSection Out;
 	Out.IndexBuffer = In.IndexBuffer;
 	Out.RDGPositionBuffer = In.RDGPositionBuffer;
 	Out.RDGPreviousPositionBuffer = In.RDGPreviousPositionBuffer;
@@ -71,59 +71,65 @@ ENGINE_API void UpdatePreviousRefToLocalMatrices(TArray<FMatrix44f>& ReferenceTo
 	const int32 LODIndex = SkeletalMeshObject->GetLOD();
 	Out.LODIndex = LODIndex;
 
-	if (!bOutputTriangleData)
+	if (!bOutputTriangleData || !SkeletalMeshObject->HaveValidDynamicData())
 	{
 		return;
 	}
 
 	FSkeletalMeshLODRenderData& LODData = SkeletalMeshObject->GetSkeletalMeshRenderData().LODRenderData[LODIndex];
-	if (LODData.RenderSections.Num() == 0)
+	const uint32 SectionCount = LODData.RenderSections.Num();
+	if (SectionCount == 0)
 	{
 		return;
 	}
 
-
 	const bool bNeedPreviousPosition = IsHairStrandContinuousDecimationReorderingEnabled();
 
 	// Create deformed position buffer (output)
-	FRDGBufferRef DeformedPositionsBuffer			= GraphBuilder.CreateBuffer(FRDGBufferDesc::CreateBufferDesc(sizeof(float), LODData.StaticVertexBuffers.PositionVertexBuffer.GetNumVertices() * 3), TEXT("Hair.SkinnedDeformedPositions"));
-	FRDGBufferRef DeformedPreviousPositionsBuffer	= bNeedPreviousPosition ? GraphBuilder.CreateBuffer(FRDGBufferDesc::CreateBufferDesc(sizeof(float), LODData.StaticVertexBuffers.PositionVertexBuffer.GetNumVertices() * 3), TEXT("Hair.SkinnedDeformedPreviousPositions")) : nullptr;
-
-	Out.DeformedPositionBuffer = DeformedPositionsBuffer;
-	Out.DeformedPreviousPositionBuffer = DeformedPreviousPositionsBuffer;
-	FRDGBufferSRVRef DeformedPositionSRV = GraphBuilder.CreateSRV(DeformedPositionsBuffer, PF_R32_FLOAT);
-	FRDGBufferSRVRef DeformedPreviousPositionSRV = bNeedPreviousPosition ? GraphBuilder.CreateSRV(DeformedPreviousPositionsBuffer, PF_R32_FLOAT) : nullptr;
-
-	uint32 BonesOffset = 0;
-	for (int32 SectionIdx = 0; SectionIdx < LODData.RenderSections.Num(); ++SectionIdx)
+	Out.DeformedPositionBuffer					= GraphBuilder.CreateBuffer(FRDGBufferDesc::CreateBufferDesc(sizeof(float), LODData.StaticVertexBuffers.PositionVertexBuffer.GetNumVertices() * 3), TEXT("Hair.SkinnedDeformedPositions"));
+	Out.DeformedPreviousPositionBuffer			= bNeedPreviousPosition ?  GraphBuilder.CreateBuffer(FRDGBufferDesc::CreateBufferDesc(sizeof(float), LODData.StaticVertexBuffers.PositionVertexBuffer.GetNumVertices() * 3), TEXT("Hair.SkinnedDeformedPreviousPositions")) : nullptr;
+	FRDGBufferSRVRef DeformedPositionSRV		= GraphBuilder.CreateSRV(Out.DeformedPositionBuffer, PF_R32_FLOAT);
+	FRDGBufferSRVRef DeformedPreviousPositionSRV= bNeedPreviousPosition ? GraphBuilder.CreateSRV(Out.DeformedPreviousPositionBuffer, PF_R32_FLOAT) : nullptr;
+
+	// Fill in result
+	TArray<FSkinUpdateSection> Sections;
+	Sections.SetNum(SectionCount);
+	for (uint32 SectionIt = 0; SectionIt < SectionCount; ++SectionIt)
 	{
-		const FSkelMeshRenderSection& Section = LODData.RenderSections[SectionIdx];
-		if (SkeletalMeshObject->HaveValidDynamicData())
-		{
-			FRHIShaderResourceView * BoneBuffer		= FSkeletalMeshDeformerHelpers::GetBoneBufferForReading(SkeletalMeshObject, LODIndex, SectionIdx, false);
-			FRHIShaderResourceView * BonePrevBuffer	= FSkeletalMeshDeformerHelpers::GetBoneBufferForReading(SkeletalMeshObject, LODIndex, SectionIdx, true);
-			AddSkinUpdatePass(GraphBuilder, ShaderMap, SectionIdx, BonesOffset, LODData, BoneBuffer, BonePrevBuffer, DeformedPositionsBuffer, DeformedPreviousPositionsBuffer);
-
-			FCachedGeometry::Section& OutSection = Out.Sections.AddDefaulted_GetRef();
-			OutSection.RDGPositionBuffer = DeformedPositionSRV;
-			OutSection.RDGPreviousPositionBuffer = DeformedPreviousPositionSRV;
-			OutSection.PositionBuffer = nullptr; // Do not use the SRV slot, but instead use the RDG buffer created above (DeformedPositionSRV)
-			OutSection.PreviousPositionBuffer = nullptr; // Do not use the SRV slot, but instead use the RDG buffer created above (DeformedPositionSRV)
-			OutSection.UVsBuffer = LODData.StaticVertexBuffers.StaticMeshVertexBuffer.GetTexCoordsSRV();
-			OutSection.TotalVertexCount = LODData.StaticVertexBuffers.PositionVertexBuffer.GetNumVertices();
-			OutSection.IndexBuffer = LODData.MultiSizeIndexContainer.GetIndexBuffer()->GetSRV();
-			OutSection.TotalIndexCount = LODData.MultiSizeIndexContainer.GetIndexBuffer()->Num();
-			OutSection.UVsChannelCount = LODData.StaticVertexBuffers.StaticMeshVertexBuffer.GetNumTexCoords();
-			OutSection.NumPrimitives = Section.NumTriangles;
-			OutSection.NumVertices = Section.NumVertices;
-			OutSection.IndexBaseIndex = Section.BaseIndex;
-			OutSection.VertexBaseIndex = Section.BaseVertexIndex;
-			OutSection.SectionIndex = SectionIdx;
-			OutSection.LODIndex = LODIndex;
-			OutSection.UVsChannelOffset = 0; // Assume that we needs to pair meshes based on UVs 0
-		}
-		BonesOffset += Section.BoneMap.Num();
+		const FSkelMeshRenderSection& Section = LODData.RenderSections[SectionIt];
+
+		Sections[SectionIt].SectionIndex 			= SectionIt;
+		Sections[SectionIt].NumVertexToProcess 		= Section.NumVertices;
+		Sections[SectionIt].SectionVertexBaseIndex 	= Section.BaseVertexIndex;
+		Sections[SectionIt].BoneBuffer 				= FSkeletalMeshDeformerHelpers::GetBoneBufferForReading(SkeletalMeshObject, LODIndex, SectionIt, false);
+		Sections[SectionIt].BonePrevBuffer 			= FSkeletalMeshDeformerHelpers::GetBoneBufferForReading(SkeletalMeshObject, LODIndex, SectionIt, true);
+
+		FCachedGeometry::Section& OutSection= Out.Sections.AddDefaulted_GetRef();
+		OutSection.RDGPositionBuffer 		= DeformedPositionSRV;
+		OutSection.RDGPreviousPositionBuffer= DeformedPreviousPositionSRV;
+		OutSection.PositionBuffer 			= nullptr; // Do not use the SRV slot, but instead use the RDG buffer created above (DeformedPositionSRV)
+		OutSection.PreviousPositionBuffer 	= nullptr; // Do not use the SRV slot, but instead use the RDG buffer created above (DeformedPositionSRV)
+		OutSection.UVsBuffer 				= LODData.StaticVertexBuffers.StaticMeshVertexBuffer.GetTexCoordsSRV();
+		OutSection.TotalVertexCount 		= LODData.StaticVertexBuffers.PositionVertexBuffer.GetNumVertices();
+		OutSection.IndexBuffer 				= LODData.MultiSizeIndexContainer.GetIndexBuffer()->GetSRV();
+		OutSection.TotalIndexCount 			= LODData.MultiSizeIndexContainer.GetIndexBuffer()->Num();
+		OutSection.UVsChannelCount 			= LODData.StaticVertexBuffers.StaticMeshVertexBuffer.GetNumTexCoords();
+		OutSection.NumPrimitives 			= Section.NumTriangles;
+		OutSection.NumVertices 				= Section.NumVertices;
+		OutSection.IndexBaseIndex 			= Section.BaseIndex;
+		OutSection.VertexBaseIndex 			= Section.BaseVertexIndex;
+		OutSection.SectionIndex 			= SectionIt;
+		OutSection.LODIndex 				= LODIndex;
+		OutSection.UVsChannelOffset 		= 0; // Assume that we needs to pair meshes based on UVs 0
 	}
+
+	AddSkinUpdatePass(
+		GraphBuilder, 
+		ShaderMap, 
+		LODData, 
+		Sections, 
+		Out.DeformedPositionBuffer, 
+		Out.DeformedPreviousPositionBuffer);
 }
 
  void BuildCacheGeometry(
diff --git a/Engine/Plugins/Runtime/HairStrands/Source/HairStrandsCore/Private/GroomGeometryCache.h b/Engine/Plugins/Runtime/HairStrands/Source/HairStrandsCore/Private/GroomGeometryCache.h
--- a/Engine/Plugins/Runtime/HairStrands/Source/HairStrandsCore/Private/GroomGeometryCache.h
+++ b/Engine/Plugins/Runtime/HairStrands/Source/HairStrandsCore/Private/GroomGeometryCache.h
@@ -11,7 +11,7 @@ class FGlobalShaderMap;
 class FRDGBuilder;
 class FGeometryCacheSceneProxy;
 
-FHairStrandsProjectionMeshData::Section ConvertMeshSection(FCachedGeometry const& InCachedGeometry, int32 InSectionIndex);
+FHairStrandsProjectionMeshData::FSection ConvertMeshSection(FCachedGeometry const& InCachedGeometry, int32 InSectionIndex);
 
 void BuildCacheGeometry(
 	FRDGBuilder& GraphBuilder,
diff --git a/Engine/Plugins/Runtime/HairStrands/Source/HairStrandsCore/Private/GroomManager.cpp b/Engine/Plugins/Runtime/HairStrands/Source/HairStrandsCore/Private/GroomManager.cpp
--- a/Engine/Plugins/Runtime/HairStrands/Source/HairStrandsCore/Private/GroomManager.cpp
+++ b/Engine/Plugins/Runtime/HairStrands/Source/HairStrandsCore/Private/GroomManager.cpp
@@ -232,21 +232,6 @@ FCachedGeometry GetCacheGeometryForHair(
 	return Out;
 }
 
-// Return the LOD index of the cached geometry on which the hair are bound to. 
-// Return -1 if the hair are not bound of if the underlying geometry is invalid
-static int32 GetCacheGeometryLODIndex(const FCachedGeometry& In)
-{
-	int32 MeshLODIndex = In.LODIndex;
-	for (const FCachedGeometry::Section& Section : In.Sections)
-	{
-		// Ensure all mesh's sections have the same LOD index
-		if (MeshLODIndex < 0) MeshLODIndex = Section.LODIndex;
-		check(MeshLODIndex == Section.LODIndex);
-	}
-
-	return MeshLODIndex;
-}
-
 enum class EHairInterpolationPassType
 {
 	Strands,
@@ -286,7 +271,7 @@ static void RunInternalHairInterpolation(
 	
 		check(Instance->HairGroupPublicData);
 
-		FHairStrandsProjectionMeshData::LOD MeshDataLOD;
+		FHairStrandsProjectionMeshData::FLOD MeshDataLOD;
 		const FCachedGeometry CachedGeometry = GetCacheGeometryForHair(GraphBuilder, Scene, Instance, ShaderMap, true);
 		for (int32 SectionIndex = 0; SectionIndex < CachedGeometry.Sections.Num(); ++SectionIndex)
 		{
@@ -420,13 +405,13 @@ static void RunHairStrandsInterpolation_Guide(
 		bool bDeformationEnable = false;
 		bool bGlobalDeformationEnable = false;
 
-		FHairStrandsProjectionMeshData::LOD MeshDataLOD;
+		FHairStrandsProjectionMeshData::FLOD MeshDataLOD;
 		FHairGroupInstance* Instance = nullptr;
 
-		bool NeedsMeshUpdate() const
+		bool NeedsMeshUpdate(bool bCheckMeshLODIndex=true) const
 		{
 			return 
-				MeshLODIndex >= 0 && 
+				(!bCheckMeshLODIndex || (bCheckMeshLODIndex && MeshLODIndex >= 0)) &&
 				(BindingType == EHairBindingType::Skinning) && 
 				(bGlobalDeformationEnable || bSimulationEnable || bDeformationEnable);
 		}
@@ -455,22 +440,26 @@ static void RunHairStrandsInterpolation_Guide(
 		InstanceData.bGlobalDeformationEnable 	= Instance->HairGroupPublicData->IsGlobalInterpolationEnable(InstanceData.HairLODIndex);
 		InstanceData.MeshLODIndex				= -1;
 
-		// Extract MeshDataLOD and compute MeshLODIndex
-		const FCachedGeometry CachedGeometry = GetCacheGeometryForHair(GraphBuilder, Scene, Instance, ShaderMap, true);
-		for (int32 SectionIndex = 0; SectionIndex < CachedGeometry.Sections.Num(); ++SectionIndex)
+		if (InstanceData.NeedsMeshUpdate(false/*bCheckMeshLODIndex*/))
 		{
-			// Ensure all mesh's sections have the same LOD index
-			const int32 SectionLodIndex = CachedGeometry.Sections[SectionIndex].LODIndex;
-			if (InstanceData.MeshLODIndex < 0) InstanceData.MeshLODIndex = SectionLodIndex;
-			check(InstanceData.MeshLODIndex == SectionLodIndex);
+			// Extract MeshDataLOD and MeshLODIndex
+			const FCachedGeometry CachedGeometry = GetCacheGeometryForHair(GraphBuilder, Scene, Instance, ShaderMap, true /*bOutputTriangleData*/);
+			for (int32 SectionIndex = 0; SectionIndex < CachedGeometry.Sections.Num(); ++SectionIndex)
+			{
+				InstanceData.MeshDataLOD.Sections.Add(ConvertMeshSection(CachedGeometry, SectionIndex));
+			}
+			InstanceData.MeshLODIndex    = CachedGeometry.LODIndex;
+			Instance->Debug.MeshLODIndex = CachedGeometry.LODIndex;
 
-			InstanceData.MeshDataLOD.Sections.Add(ConvertMeshSection(CachedGeometry, SectionIndex));
+			// Sanity check
+			check(InstanceData.Instance->Guides.IsValid());
 		}
-		Instance->Debug.MeshLODIndex = InstanceData.MeshLODIndex;
-
-		if (InstanceData.NeedsMeshUpdate())
+		else
 		{
-			check(InstanceData.Instance->Guides.IsValid());
+			// Extract MeshLODIndex
+			const FCachedGeometry CachedGeometry = GetCacheGeometryForHair(GraphBuilder, Scene, Instance, ShaderMap, false /*bOutputTriangleData*/);
+			InstanceData.MeshLODIndex    = CachedGeometry.LODIndex;
+			Instance->Debug.MeshLODIndex = CachedGeometry.LODIndex;
 		}
 
 		if (InstanceData.ActiveGroomCacheType == EGroomCacheType::Guides)  { bHasAnySimCacheInstances = true; }
@@ -763,15 +752,15 @@ static void RunHairStrandsInterpolation_Strands(
 
 		FInstanceRDGResources RDGResources;
 
-		FHairStrandsProjectionMeshData::LOD MeshDataLOD;
+		FHairStrandsProjectionMeshData::FLOD MeshDataLOD;
 		FHairGroupInstance* Instance = nullptr;
 		FRDGHairStrandsCullingData CullingData;
 
-		bool NeedsMeshUpdate() const
+		bool NeedsMeshUpdate(bool bCheckMeshLODIndex = true) const
 		{
-			return 
-				MeshLODIndex >= 0 && 
-				(BindingType == EHairBindingType::Skinning) && 
+			return
+				(!bCheckMeshLODIndex || (bCheckMeshLODIndex && MeshLODIndex >= 0)) &&
+				(BindingType == EHairBindingType::Skinning) &&
 				(bGlobalDeformationEnable || bSimulationEnable || bDeformationEnable);
 		}
 	};
@@ -852,14 +841,10 @@ static void RunHairStrandsInterpolation_Strands(
 		const FCachedGeometry CachedGeometry = GetCacheGeometryForHair(GraphBuilder, Scene, Instance, ShaderMap, true);
 		for (int32 SectionIndex = 0; SectionIndex < CachedGeometry.Sections.Num(); ++SectionIndex)
 		{
-			// Ensure all mesh's sections have the same LOD index
-			const int32 SectionLodIndex = CachedGeometry.Sections[SectionIndex].LODIndex;
-			if (InstanceData.MeshLODIndex < 0) InstanceData.MeshLODIndex = SectionLodIndex;
-			check(InstanceData.MeshLODIndex == SectionLodIndex);
-
 			InstanceData.MeshDataLOD.Sections.Add(ConvertMeshSection(CachedGeometry, SectionIndex));
 		}
-		Instance->Debug.MeshLODIndex = InstanceData.MeshLODIndex;
+		InstanceData.MeshLODIndex    = CachedGeometry.LODIndex;
+		Instance->Debug.MeshLODIndex = CachedGeometry.LODIndex;
 		Instance->HairGroupPublicData->VFInput.Strands	= FHairGroupPublicData::FVertexFactoryInput::FStrands();
 
 		if (InstanceData.MeshLODIndex >= 0 && (InstanceData.BindingType == EHairBindingType::Skinning || InstanceData.bGlobalDeformationEnable))
@@ -1714,7 +1699,7 @@ static void RunHairLODSelection(
 		check(Instance);
 		check(Instance->HairGroupPublicData);
 		const FCachedGeometry CachedGeometry = GetCacheGeometryForHair(GraphBuilder, Scene, Instance, ShaderMap, false);
-		const int32 MeshLODIndex = GetCacheGeometryLODIndex(CachedGeometry);
+		const int32 MeshLODIndex = CachedGeometry.LODIndex;
 
 		// Perform LOD selection based on all the views	
 		// CPU LOD selection. 
diff --git a/Engine/Plugins/Runtime/HairStrands/Source/HairStrandsCore/Private/HairStrandsMeshProjection.cpp b/Engine/Plugins/Runtime/HairStrands/Source/HairStrandsCore/Private/HairStrandsMeshProjection.cpp
--- a/Engine/Plugins/Runtime/HairStrands/Source/HairStrandsCore/Private/HairStrandsMeshProjection.cpp
+++ b/Engine/Plugins/Runtime/HairStrands/Source/HairStrandsCore/Private/HairStrandsMeshProjection.cpp
@@ -33,44 +33,6 @@ uint32 GetHairStrandsMaxTriangleCount()
 	return (1 << MAX_HAIRSTRANDS_SECTION_BITOFFSET) - 1;
 }
 
-///////////////////////////////////////////////////////////////////////////////////////////////////
-
-FHairStrandsProjectionMeshData ExtractMeshData(FSkeletalMeshRenderData* RenderData)
-{
-	FHairStrandsProjectionMeshData MeshData;
-	uint32 LODIndex = 0;
-	for (FSkeletalMeshLODRenderData& LODRenderData : RenderData->LODRenderData)
-	{
-		FHairStrandsProjectionMeshData::LOD& LOD = MeshData.LODs.AddDefaulted_GetRef();
-		uint32 SectionIndex = 0;
-		for (FSkelMeshRenderSection& InSection : LODRenderData.RenderSections)
-		{
-			// Pick between float and halt
-			const uint32 UVSizeInByte = (LODRenderData.StaticVertexBuffers.StaticMeshVertexBuffer.GetUseFullPrecisionUVs() ? 4 : 2) * 2;
-
-			FHairStrandsProjectionMeshData::Section& OutSection = LOD.Sections.AddDefaulted_GetRef();
-			OutSection.UVsChannelOffset = 0; // Assume that we needs to pair meshes based on UVs 0
-			OutSection.UVsChannelCount = LODRenderData.StaticVertexBuffers.StaticMeshVertexBuffer.GetNumTexCoords();
-			OutSection.UVsBuffer = LODRenderData.StaticVertexBuffers.StaticMeshVertexBuffer.GetTexCoordsSRV();
-			OutSection.PositionBuffer = LODRenderData.StaticVertexBuffers.PositionVertexBuffer.GetSRV();
-			OutSection.IndexBuffer = LODRenderData.MultiSizeIndexContainer.GetIndexBuffer()->GetSRV();
-			OutSection.TotalVertexCount = LODRenderData.StaticVertexBuffers.PositionVertexBuffer.GetNumVertices();
-			OutSection.TotalIndexCount = LODRenderData.MultiSizeIndexContainer.GetIndexBuffer()->Num();
-			OutSection.NumPrimitives = InSection.NumTriangles;
-			OutSection.NumVertices = InSection.NumVertices;
-			OutSection.VertexBaseIndex = InSection.BaseVertexIndex;
-			OutSection.IndexBaseIndex = InSection.BaseIndex;
-			OutSection.SectionIndex = SectionIndex;
-			OutSection.LODIndex = LODIndex;
-
-			++SectionIndex;
-		}
-		++LODIndex;
-	}
-
-	return MeshData;
-}
-
 ///////////////////////////////////////////////////////////////////////////////////////////////////
 class FSkinUpdateCS : public FGlobalShader
 {
@@ -91,7 +53,6 @@ class FSkinUpdateCS : public FGlobalShader
 		SHADER_PARAMETER(uint32, SectionVertexBaseIndex)
 		SHADER_PARAMETER(uint32, WeightIndexSize)
 		SHADER_PARAMETER(uint32, WeightStride)
-		SHADER_PARAMETER(uint32, BonesOffset)
 		SHADER_PARAMETER_SRV(Buffer<uint>, WeightLookup)
 		SHADER_PARAMETER_SRV(Buffer<float4>, BoneMatrices)
 		SHADER_PARAMETER_SRV(Buffer<float4>, PrevBoneMatrices)
@@ -102,7 +63,14 @@ class FSkinUpdateCS : public FGlobalShader
 	END_SHADER_PARAMETER_STRUCT()
 
 public:
+	static uint32 GetGroupSize() { return 64; }
 	static bool ShouldCompilePermutation(const FGlobalShaderPermutationParameters& Parameters) { return IsHairStrandsSupported(EHairStrandsShaderType::All, Parameters.Platform); }
+	static void ModifyCompilationEnvironment(const FGlobalShaderPermutationParameters& Parameters, FShaderCompilerEnvironment& OutEnvironment)
+	{
+		FGlobalShader::ModifyCompilationEnvironment(Parameters, OutEnvironment);
+		OutEnvironment.SetDefine(TEXT("SHADER_SKIN_UPDATE"), 1);
+		OutEnvironment.SetDefine(TEXT("GROUP_SIZE"), GetGroupSize());
+	}
 };
 
 IMPLEMENT_GLOBAL_SHADER(FSkinUpdateCS, "/Engine/Private/HairStrands/HairStrandsSkinUpdate.usf", "UpdateSkinPositionCS", SF_Compute);
@@ -110,41 +78,30 @@ IMPLEMENT_GLOBAL_SHADER(FSkinUpdateCS, "/Engine/Private/HairStrands/HairStrandsS
 void AddSkinUpdatePass(
 	FRDGBuilder& GraphBuilder,
 	FGlobalShaderMap* ShaderMap,
-	uint32 SectionIndex,
-	uint32 BonesOffset,	
 	FSkeletalMeshLODRenderData& RenderData,
-	FRHIShaderResourceView* BoneMatrices,
-	FRHIShaderResourceView* PrevBoneMatrices,
-	FRDGBufferRef OutDeformedPosition,
-	FRDGBufferRef OutPrevDeformedPosition)
+	const TArray<FSkinUpdateSection>& Sections,
+	FRDGBufferRef OutDeformedPositionBuffer,
+	FRDGBufferRef OutPrevDeformedPositionBuffer)
 {
-	check(BoneMatrices);
-
-	FSkinWeightVertexBuffer* SkinWeight = &RenderData.SkinWeightVertexBuffer;
-	const FSkelMeshRenderSection& Section = RenderData.RenderSections[SectionIndex];
-	const uint32 NumVertexToProcess = Section.NumVertices;
-	const uint32 SectionVertexBaseIndex = Section.BaseVertexIndex;
-	const uint32 NumTotalVertices = RenderData.StaticVertexBuffers.PositionVertexBuffer.GetNumVertices();
-
-	const bool bPrevPosition = OutPrevDeformedPosition != nullptr && PrevBoneMatrices != nullptr;
+	check(Sections.Num() > 0);
+	const FSkinWeightVertexBuffer* SkinWeight = &RenderData.SkinWeightVertexBuffer;
+	const bool bPrevPosition = OutPrevDeformedPositionBuffer != nullptr && Sections[0].BonePrevBuffer;
 	
 	FSkinUpdateCS::FParameters* Parameters = GraphBuilder.AllocParameters<FSkinUpdateCS::FParameters>();
-	Parameters->WeightIndexSize = SkinWeight->GetBoneIndexByteSize() | (SkinWeight->GetBoneWeightByteSize() << 8);
-	Parameters->NumVertexToProcess = NumVertexToProcess;
-	Parameters->NumTotalVertices = NumTotalVertices;
-	Parameters->SectionVertexBaseIndex = SectionVertexBaseIndex;
-	Parameters->WeightStride = SkinWeight->GetConstantInfluencesVertexStride();
-	Parameters->WeightLookup = SkinWeight->GetLookupVertexBuffer()->GetSRV();
-	Parameters->BonesOffset = BonesOffset;
-	Parameters->BoneMatrices = BoneMatrices;
-	Parameters->VertexWeights = SkinWeight->GetDataVertexBuffer()->GetSRV();
-	Parameters->RestPositions = RenderData.StaticVertexBuffers.PositionVertexBuffer.GetSRV();
-	Parameters->DeformedPositions = GraphBuilder.CreateUAV(OutDeformedPosition, PF_R32_FLOAT);
+	Parameters->WeightIndexSize 		= SkinWeight->GetBoneIndexByteSize() | (SkinWeight->GetBoneWeightByteSize() << 8);
+	Parameters->NumVertexToProcess 		= 0;
+	Parameters->NumTotalVertices		= RenderData.StaticVertexBuffers.PositionVertexBuffer.GetNumVertices();
+	Parameters->SectionVertexBaseIndex 	= 0;
+	Parameters->WeightStride 			= SkinWeight->GetConstantInfluencesVertexStride();
+	Parameters->WeightLookup 			= SkinWeight->GetLookupVertexBuffer()->GetSRV();
+	Parameters->BoneMatrices 			= Sections[0].BoneBuffer;
+	Parameters->PrevBoneMatrices 		= Sections[0].BonePrevBuffer;
+	Parameters->VertexWeights 			= SkinWeight->GetDataVertexBuffer()->GetSRV();
+	Parameters->RestPositions 			= RenderData.StaticVertexBuffers.PositionVertexBuffer.GetSRV();
+	Parameters->DeformedPositions 		= GraphBuilder.CreateUAV(OutDeformedPositionBuffer, PF_R32_FLOAT, ERDGUnorderedAccessViewFlags::SkipBarrier);
 	if (bPrevPosition)
 	{
-		check(PrevBoneMatrices);
-		Parameters->PrevBoneMatrices = BoneMatrices;
-		Parameters->PrevDeformedPositions = GraphBuilder.CreateUAV(OutPrevDeformedPosition, PF_R32_FLOAT);
+		Parameters->PrevDeformedPositions = GraphBuilder.CreateUAV(OutPrevDeformedPositionBuffer, PF_R32_FLOAT, ERDGUnorderedAccessViewFlags::SkipBarrier);
 	}
 
 	FSkinUpdateCS::FPermutationDomain PermutationVector;
@@ -153,16 +110,34 @@ void AddSkinUpdatePass(
 	PermutationVector.Set<FSkinUpdateCS::FBoneIndexUint16 >(SkinWeight->Use16BitBoneIndex());
 	PermutationVector.Set<FSkinUpdateCS::FBoneIndexUint16 >(SkinWeight->Use16BitBoneWeight());
 	PermutationVector.Set<FSkinUpdateCS::FPrevious>(bPrevPosition);
-
-	const FIntVector DispatchGroupCount = FComputeShaderUtils::GetGroupCount(NumVertexToProcess, 64);
-	check(DispatchGroupCount.X <= GRHIMaxDispatchThreadGroupsPerDimension.X);
 	TShaderMapRef<FSkinUpdateCS> ComputeShader(ShaderMap, PermutationVector);
-	FComputeShaderUtils::AddPass(
-		GraphBuilder,
-		RDG_EVENT_NAME("HairStrands::UpdateSkinPosition(%s,Section:%d)", bPrevPosition ? TEXT("Curr,Prev") : TEXT("Curr"), SectionIndex),
-		ComputeShader,
+
+	const FShaderParametersMetadata* ParametersMetadata = FSkinUpdateCS::FParameters::FTypeInfo::GetStructMetadata();
+	ClearUnusedGraphResources(ComputeShader, ParametersMetadata, Parameters);
+	GraphBuilder.AddPass(
+		RDG_EVENT_NAME("HairStrands::UpdateSkinPosition(%s)", bPrevPosition ? TEXT("Curr,Prev") : TEXT("Curr")),
+		ParametersMetadata,
 		Parameters,
-		DispatchGroupCount);
+		ERDGPassFlags::Compute,
+		[ParametersMetadata, Parameters, ComputeShader, Sections, bPrevPosition](FRHIComputeCommandList& RHICmdList)
+		{
+			for (uint32 SectionIt = 0, SectionCount = Sections.Num(); SectionIt < SectionCount; ++SectionIt)
+			{
+				Parameters->NumVertexToProcess 		= Sections[SectionIt].NumVertexToProcess;
+				Parameters->SectionVertexBaseIndex 	= Sections[SectionIt].SectionVertexBaseIndex;
+				Parameters->BoneMatrices 			= Sections[SectionIt].BoneBuffer;
+				check(Parameters->BoneMatrices);
+				if (bPrevPosition)
+				{
+					Parameters->PrevBoneMatrices = Sections[SectionIt].BonePrevBuffer;
+					check(Parameters->PrevBoneMatrices);
+				}
+
+				const FIntVector DispatchGroupCount = FComputeShaderUtils::GetGroupCount(Parameters->NumVertexToProcess, FSkinUpdateCS::GetGroupSize());
+				check(DispatchGroupCount.X <= GRHIMaxDispatchThreadGroupsPerDimension.X);
+				FComputeShaderUtils::Dispatch(RHICmdList, ComputeShader, ParametersMetadata, *Parameters, DispatchGroupCount);
+			}
+		});
 }
 
 ///////////////////////////////////////////////////////////////////////////////////////////////////
@@ -210,7 +185,7 @@ bool AddHairStrandUpdateMeshTrianglesPass(
 	FRDGBuilder& GraphBuilder,
 	FGlobalShaderMap* ShaderMap,
 	const uint32 UniqueTriangleCount, 
-	const FHairStrandsProjectionMeshData::LOD& MeshData,
+	const FHairStrandsProjectionMeshData::FLOD& MeshData,
 	FRDGBufferSRVRef UniqueTriangleIndexSRV,
 	FRDGBufferUAVRef OutputCurrUAV,
 	FRDGBufferUAVRef OutputPrevUAV)
@@ -254,7 +229,7 @@ bool AddHairStrandUpdateMeshTrianglesPass(
 	SectionDatas.SetNum(SectionCount);
 	for (uint32 SectionIt = 0; SectionIt < SectionCount; ++SectionIt)
 	{
-		const FHairStrandsProjectionMeshData::Section& MeshSectionData = MeshData.Sections[SectionIt];
+		const FHairStrandsProjectionMeshData::FSection& MeshSectionData = MeshData.Sections[SectionIt];
 		SectionDatas[SectionIt].TotalIndexCount	= MeshSectionData.TotalIndexCount;
 		SectionDatas[SectionIt].TotalVertexCount= MeshSectionData.TotalVertexCount;
 		SectionDatas[SectionIt].IndexBaseIndex	= MeshSectionData.IndexBaseIndex;
@@ -315,7 +290,7 @@ void AddHairStrandUpdateMeshTrianglesPass(
 	FRDGBuilder& GraphBuilder,
 	FGlobalShaderMap* ShaderMap,
 	const int32 MeshLODIndex,
-	const FHairStrandsProjectionMeshData::LOD& MeshData,
+	const FHairStrandsProjectionMeshData::FLOD& MeshData,
 	FHairStrandsRestRootResource* RestResources,
 	FHairStrandsDeformedRootResource* DeformedResources)
 {	
@@ -527,7 +502,7 @@ void AddHairStrandInitMeshSamplesPass(
 	FRDGBuilder& GraphBuilder,
 	FGlobalShaderMap* ShaderMap,
 	const int32 LODIndex,
-	const FHairStrandsProjectionMeshData::LOD& MeshData,
+	const FHairStrandsProjectionMeshData::FLOD& MeshData,
 	FHairStrandsRestRootResource* RestResources,
 	FHairStrandsDeformedRootResource* DeformedResources)
 {
@@ -642,7 +617,7 @@ void AddHairStrandUpdateMeshSamplesPass(
 	FRDGBuilder& GraphBuilder,
 	FGlobalShaderMap* ShaderMap,
 	const int32 LODIndex,
-	const FHairStrandsProjectionMeshData::LOD& MeshData,
+	const FHairStrandsProjectionMeshData::FLOD& MeshData,
 	FHairStrandsRestRootResource* RestResources,
 	FHairStrandsDeformedRootResource* DeformedResources)
 {
diff --git a/Engine/Plugins/Runtime/HairStrands/Source/HairStrandsCore/Public/HairStrandsMeshProjection.h b/Engine/Plugins/Runtime/HairStrands/Source/HairStrandsCore/Public/HairStrandsMeshProjection.h
--- a/Engine/Plugins/Runtime/HairStrands/Source/HairStrandsCore/Public/HairStrandsMeshProjection.h
+++ b/Engine/Plugins/Runtime/HairStrands/Source/HairStrandsCore/Public/HairStrandsMeshProjection.h
@@ -10,6 +10,7 @@
 
 
 class FSkeletalMeshRenderData;
+class FSkeletalMeshLODRenderData;
 struct FHairStrandsRestRootResource;
 struct FHairStrandsDeformedRootResource;
 struct FHairMeshesRestResource;
@@ -20,7 +21,7 @@ struct FHairStrandsDeformedResource;
 
 struct FHairStrandsProjectionMeshData
 {
-	struct Section
+	struct FSection
 	{
 		FTransform LocalToWorld;
 		FRDGBufferSRVRef RDGPositionBuffer = nullptr;
@@ -41,11 +42,11 @@ struct FHairStrandsProjectionMeshData
 		int32 LODIndex = 0;
 	};
 
-	struct LOD
+	struct FLOD
 	{
-		TArray<Section> Sections;
+		TArray<FSection> Sections;
 	};
-	TArray<LOD> LODs;
+	TArray<FLOD> LODs;
 };
 
 // Return the max number of section/triangle a skeletal mesh can have. After this count, binding will be disabled
@@ -57,7 +58,7 @@ void AddHairStrandUpdateMeshTrianglesPass(
 	FRDGBuilder& GraphBuilder,
 	FGlobalShaderMap* ShaderMap,
 	const int32 LODIndex,
-	const FHairStrandsProjectionMeshData::LOD& MeshData,
+	const FHairStrandsProjectionMeshData::FLOD& MeshData,
 	FHairStrandsRestRootResource* RestResources,
 	FHairStrandsDeformedRootResource* DeformedResources);
 
@@ -66,7 +67,7 @@ void AddHairStrandInitMeshSamplesPass(
 	FRDGBuilder& GraphBuilder,
 	FGlobalShaderMap* ShaderMap,
 	const int32 LODIndex,
-	const FHairStrandsProjectionMeshData::LOD& MeshData,
+	const FHairStrandsProjectionMeshData::FLOD& MeshData,
 	FHairStrandsRestRootResource* RestResources,
 	FHairStrandsDeformedRootResource* DeformedResources);
 
@@ -75,7 +76,7 @@ void AddHairStrandUpdateMeshSamplesPass(
 	FRDGBuilder& GraphBuilder,
 	FGlobalShaderMap* ShaderMap,
 	const int32 LODIndex,
-	const FHairStrandsProjectionMeshData::LOD& ProjectionMeshData,
+	const FHairStrandsProjectionMeshData::FLOD& ProjectionMeshData,
 	FHairStrandsRestRootResource* RestResources,
 	FHairStrandsDeformedRootResource* DeformedResources);
 
@@ -107,16 +108,22 @@ void AddComputeMipsPass(
 	FGlobalShaderMap* ShaderMap,
 	FRDGTextureRef& OutTexture);
 
+struct FSkinUpdateSection
+{
+	uint32 SectionIndex = 0;
+	uint32 NumVertexToProcess = 0;
+	uint32 SectionVertexBaseIndex = 0;
+	FRHIShaderResourceView* BoneBuffer = nullptr;
+	FRHIShaderResourceView* BonePrevBuffer = nullptr;
+};
+
 void AddSkinUpdatePass(
 	FRDGBuilder& GraphBuilder,
 	FGlobalShaderMap* ShaderMap,
-	uint32 SectionIndex,
-	uint32 BonesOffset, 
-	class FSkeletalMeshLODRenderData& RenderData,
-	FRHIShaderResourceView* BoneMatrices,
-	FRHIShaderResourceView* PrevBoneMatrices,
-	FRDGBufferRef OutDeformedosition,
-	FRDGBufferRef OutPreviousDeformedosition);
+	FSkeletalMeshLODRenderData& RenderData,
+	const TArray<FSkinUpdateSection>& Sections,
+	FRDGBufferRef OutDeformedPositionBuffer,
+	FRDGBufferRef OutPrevDeformedPositionBuffer);
 
 void AddHairMeshesRBFInterpolationPass(
 	FRDGBuilder& GraphBuilder,
@@ -151,6 +158,4 @@ void AddHairStrandUpdatePositionOffsetPass(
 	const int32 HairLODIndex,
 	const int32 MeshLODIndex,
 	FHairStrandsDeformedRootResource* DeformedRootResources,
-	FHairStrandsDeformedResource* DeformedResources);
-
-FHairStrandsProjectionMeshData ExtractMeshData(FSkeletalMeshRenderData* RenderData);
+	FHairStrandsDeformedResource* DeformedResources);
\ No newline at end of file
diff --git a/Engine/Shaders/Private/HairStrands/HairStrandsSkinUpdate.usf b/Engine/Shaders/Private/HairStrands/HairStrandsSkinUpdate.usf
--- a/Engine/Shaders/Private/HairStrands/HairStrandsSkinUpdate.usf
+++ b/Engine/Shaders/Private/HairStrands/HairStrandsSkinUpdate.usf
@@ -5,6 +5,7 @@
 #include "/Engine/Public/Platform.ush"
 #include "../GpuSkinCommon.ush"
 
+#if SHADER_SKIN_UPDATE
 //----------------------------------------------------------------------------------
 
 uint WeightIndexSize; // bone index size: 1 - 8bit or 2 - 16bit
@@ -12,7 +13,6 @@ uint SectionVertexBaseIndex;
 uint NumTotalVertices;
 uint NumVertexToProcess;
 uint WeightStride;
-uint BonesOffset;
 
 Buffer<uint> WeightLookup;
 Buffer<uint> VertexWeights;
@@ -30,7 +30,7 @@ RWBuffer<float> PrevDeformedPositions;
 
 //----------------------------------------------------------------------------------
 
-[numthreads(64,1,1)]
+[numthreads(GROUP_SIZE,1,1)]
 void UpdateSkinPositionCS(uint3 GroupID : SV_GroupID,
  uint3 DispatchThreadID : SV_DispatchThreadID,
  uint3 GroupThreadID : SV_GroupThreadID)
@@ -82,3 +82,4 @@ void UpdateSkinPositionCS(uint3 GroupID : SV_GroupID,
 	#endif
 	}
 }
+#endif // SHADER_SKIN_UPDATE
\ No newline at end of file
