Add missing file.

diff --git a/Engine/Plugins/Runtime/HairStrands/Source/HairStrandsCore/Private/GroomManager.h b/Engine/Plugins/Runtime/HairStrands/Source/HairStrandsCore/Private/GroomManager.h
--- a/Engine/Plugins/Runtime/HairStrands/Source/HairStrandsCore/Private/GroomManager.h
+++ b/Engine/Plugins/Runtime/HairStrands/Source/HairStrandsCore/Private/GroomManager.h
@@ -54,7 +54,5 @@ EHairBufferSwapType GetHairSwapBufferType();
 void RegisterHairStrands(struct FHairGroupInstance* Instance);
 void UnregisterHairStrands(uint32 ComponentId);
 
-FHairStrandsProjectionMeshData ExtractMeshData(class FSkeletalMeshRenderData* RenderData);
-
 ////////////////////////////////////////////////////////////////////////////////////////////////////////////////////////
 
