Split hair mesh/cards interpolation into batch work.

#preflitght 656ad130fd601f4340a92fa1

diff --git a/Engine/Plugins/Runtime/HairStrands/Source/HairStrandsCore/Private/GroomManager.cpp b/Engine/Plugins/Runtime/HairStrands/Source/HairStrandsCore/Private/GroomManager.cpp
--- a/Engine/Plugins/Runtime/HairStrands/Source/HairStrandsCore/Private/GroomManager.cpp
+++ b/Engine/Plugins/Runtime/HairStrands/Source/HairStrandsCore/Private/GroomManager.cpp
@@ -153,8 +153,6 @@ static bool IsInstanceFrustumCullingEnable()
 	return GHairStrands_InterpolationFrustumCullingEnable > 0;
 }
 
-bool NeedsUpdateCardsMeshTriangles();
-
 static bool IsSkeletalMeshEvaluationEnabled()
 {
 	// When deferred skel. mesh update is enabled, hair strands skeletal mesh deformation is not allowed, as skin-cached update happen after 
@@ -232,153 +230,6 @@ FCachedGeometry GetCacheGeometryForHair(
 	return Out;
 }
 
-enum class EHairInterpolationPassType
-{
-	Strands,
-	CardsAndMeshes,
-	Guides
-};
-
-static void RunInternalHairInterpolation(
-	FRDGBuilder& GraphBuilder,
-	FSceneInterface* Scene,
-	const FSceneView* View,
-	const uint32 ViewUniqueID,
-	const FHairStrandsInstances& Instances,
-	const FShaderPrintData* ShaderPrintData,
-	FGlobalShaderMap* ShaderMap, 
-	EHairInterpolationPassType PassType,
-	FHairStrandClusterData* ClusterData)
-{
-	check(IsInRenderingThread());
-
-	#if RHI_RAYTRACING
-	const uint32 ViewRayTracingMask = View->Family->EngineShowFlags.PathTracing ? EHairViewRayTracingMask::PathTracing : EHairViewRayTracingMask::RayTracing;
-	#else
-	const uint32 ViewRayTracingMask = 0u;
-	#endif
-
-	const EGroomViewMode ViewMode = GetGroomViewMode(*View);
-
-	// Update dynamic mesh triangles
-	for (FHairStrandsInstance* AbstractInstance : Instances)
-	{
-		FHairGroupInstance* Instance = static_cast<FHairGroupInstance*>(AbstractInstance);
-
-		int32 MeshLODIndex = -1;
-		if (Instance->GeometryType == EHairGeometryType::NoneGeometry)
-			continue;
-	
-		check(Instance->HairGroupPublicData);
-
-		FHairStrandsProjectionMeshData::FLOD MeshDataLOD;
-		const FCachedGeometry CachedGeometry = GetCacheGeometryForHair(GraphBuilder, Scene, Instance, ShaderMap, true);
-		for (int32 SectionIndex = 0; SectionIndex < CachedGeometry.Sections.Num(); ++SectionIndex)
-		{
-			// Ensure all mesh's sections have the same LOD index
-			const int32 SectionLodIndex = CachedGeometry.Sections[SectionIndex].LODIndex;
-			if (MeshLODIndex < 0) MeshLODIndex = SectionLodIndex;
-			check(MeshLODIndex == SectionLodIndex);
-
-			MeshDataLOD.Sections.Add(ConvertMeshSection(CachedGeometry, SectionIndex));
-		}
-
-		Instance->Debug.MeshLODIndex = MeshLODIndex;
-		if (0 <= MeshLODIndex)
-		{
-			const EHairGeometryType InstanceGeometryType = Instance->GeometryType;
-			const EHairBindingType InstanceBindingType = Instance->BindingType;
-
-			const uint32 HairLODIndex = Instance->HairGroupPublicData->LODIndex;
-			const EHairBindingType BindingType = Instance->HairGroupPublicData->GetBindingType(HairLODIndex);
-			const bool bSimulationEnable = Instance->HairGroupPublicData->IsSimulationEnable(HairLODIndex);
-			const bool bDeformationEnable = Instance->HairGroupPublicData->bIsDeformationEnable;
-			const bool bGlobalDeformationEnable = Instance->HairGroupPublicData->IsGlobalInterpolationEnable(HairLODIndex);
-			check(InstanceBindingType == BindingType);
-
-			if (EHairInterpolationPassType::CardsAndMeshes == PassType)
-			{
-				if (InstanceGeometryType == EHairGeometryType::Cards)
-				{
-					if (Instance->Cards.IsValid(HairLODIndex))
-					{
-						FHairGroupInstance::FCards::FLOD& CardsInstance = Instance->Cards.LODs[HairLODIndex];
-						if (BindingType == EHairBindingType::Skinning || bGlobalDeformationEnable)
-						{
-							check(CardsInstance.Guides.IsValid());
-							check(CardsInstance.Guides.HasValidRootData());
-							check(CardsInstance.Guides.DeformedRootResource->IsValid(MeshLODIndex));
-
-							AddHairStrandUpdateMeshTrianglesPass(
-								GraphBuilder,
-								ShaderMap,
-								MeshLODIndex,
-								MeshDataLOD,
-								CardsInstance.Guides.RestRootResource,
-								CardsInstance.Guides.DeformedRootResource);
-
-							AddHairStrandUpdatePositionOffsetPass(
-								GraphBuilder,
-								ShaderMap,
-								EHairPositionUpdateType::Cards,
-								Instance->RegisteredIndex,
-								HairLODIndex,
-								MeshLODIndex,
-								CardsInstance.Guides.DeformedRootResource,
-								CardsInstance.Guides.DeformedResource);
-						}
-						else if (bSimulationEnable || bDeformationEnable)
-						{
-							check(CardsInstance.Guides.IsValid());
-							AddHairStrandUpdatePositionOffsetPass(
-								GraphBuilder,
-								ShaderMap,
-								EHairPositionUpdateType::Cards,
-								Instance->RegisteredIndex,
-								HairLODIndex,
-								MeshLODIndex,
-								nullptr,
-								CardsInstance.Guides.DeformedResource);
-						}
-					}
-				}
-				else if (InstanceGeometryType == EHairGeometryType::Meshes)
-				{
-					// Nothing to do
-				}
-			}
-		}
-	}
-
-	// Hair interpolation
-	if (EHairInterpolationPassType::CardsAndMeshes == PassType)
-	{
-		check(View);
-		const FVector& TranslatedWorldOffset = View->ViewMatrices.GetPreViewTranslation();
-		for (FHairStrandsInstance* AbstractInstance : Instances)
-		{
-			FHairGroupInstance* Instance = static_cast<FHairGroupInstance*>(AbstractInstance);
-
-			const bool bGeometryCompatibleWithPassType =
-				(EHairInterpolationPassType::CardsAndMeshes == PassType && (Instance->GeometryType == EHairGeometryType::Cards || Instance->GeometryType == EHairGeometryType::Meshes));
-			if (bGeometryCompatibleWithPassType)
-			{
-				ComputeHairStrandsInterpolation(
-					GraphBuilder, 
-					ShaderMap,
-					ViewUniqueID,
-					ViewRayTracingMask,
-					ViewMode,
-					TranslatedWorldOffset,
-					ShaderPrintData,
-					Instance,
-					Instance->Debug.MeshLODIndex,
-					ClusterData);
-			}
-		}
-	}
-}
-
 static void RunHairStrandsInterpolation_Guide(
 	FRDGBuilder& GraphBuilder,
 	FSceneInterface* Scene,
@@ -857,7 +708,7 @@ static void RunHairStrandsInterpolation_Strands(
 		{
 			check(Instance->Strands.HasValidData());
 		}
-	}	
+	}
 
 	FRDGExternalAccessQueue ExternalAccessQueue;
 	const EGroomViewMode ViewMode = GetGroomViewMode(*View);
@@ -1350,7 +1201,7 @@ static void RunHairStrandsInterpolation_Strands(
 					InstanceData.Instance->Guides.RestResource->GetPointCount(),
 					InstanceData.Instance->HairGroupPublicData,
 					RegisterAsSRV(GraphBuilder, InstanceData.Instance->Guides.DeformedResource->GetBuffer(FHairStrandsDeformedResource::Current)),
-					RegisterAsUAV(GraphBuilder, InstanceData.Instance->Guides.DeformedResource->TangentBuffer));					
+					RegisterAsUAV(GraphBuilder, InstanceData.Instance->Guides.DeformedResource->TangentBuffer));
 			}
 			InstanceData.Instance->HairGroupPublicData->VFInput = InternalComputeHairStrandsVertexInputData(&GraphBuilder, InstanceData.Instance, ViewMode);
 		}
@@ -1369,6 +1220,15 @@ static void RunHairStrandsInterpolation_Strands(
 	ExternalAccessQueue.Submit(GraphBuilder);
 }
 
+enum class EHairCardsSimulationType
+{
+	None,
+	Guide,
+	RBF
+};
+
+uint32 GetHairCardsInterpolationType();
+
 static void RunHairStrandsInterpolation_Cards(
 	FRDGBuilder& GraphBuilder,
 	FSceneInterface* Scene,
@@ -1379,24 +1239,373 @@ static void RunHairStrandsInterpolation_Cards(
 	const FShaderPrintData* ShaderPrintData,
 	FGlobalShaderMap* ShaderMap)
 {
-	if (Instances.IsEmpty()) { return; }
+	const bool bCardSupported = IsHairStrandsEnabled(EHairStrandsShaderType::Cards, Scene->GetShaderPlatform());
+	const bool bMeshSupported = IsHairStrandsEnabled(EHairStrandsShaderType::Meshes, Scene->GetShaderPlatform());
 
+	if (Instances.IsEmpty() || (!bCardSupported && !bMeshSupported)) { return; }
 	check(IsInRenderingThread());
+	check(View);
 
 	DECLARE_GPU_STAT(HairCardsInterpolation);
 	RDG_EVENT_SCOPE(GraphBuilder, "HairCardsInterpolation");
 	RDG_GPU_STAT_SCOPE(GraphBuilder, HairCardsInterpolation);
 
-	RunInternalHairInterpolation(
-		GraphBuilder,
-		Scene,
-		View,
-		ViewUniqueID,
-		Instances,
-		ShaderPrintData,
-		ShaderMap,
-		EHairInterpolationPassType::CardsAndMeshes,
-		nullptr);
+	struct FInstanceData
+	{
+		uint32 RegisteredIndex = ~0;
+		int32 HairLODIndex = -1;
+		int32 MeshLODIndex = -1;
+		EHairBindingType BindingType = EHairBindingType::NoneBinding;
+		bool bSimulationEnable = false;
+		bool bDeformationEnable = false;
+		bool bGlobalDeformationEnable = false;
+		bool bValidGuide = false;
+		bool bHasSkinning = false;
+		bool bNeedDeformation = false;
+		EHairCardsSimulationType CardsSimulationType = EHairCardsSimulationType::None;
+		FHairGroupInstance::FCards::FLOD*  CardInstance = nullptr;
+		FHairGroupInstance::FMeshes::FLOD* MeshInstance = nullptr;
+		FHairStrandsProjectionMeshData::FLOD MeshDataLOD;
+		FHairGroupInstance* Instance = nullptr;
+	};
+
+	// Gather all strands instances
+	TArray<FInstanceData> InstanceDatas;
+	TArray<uint32> MeshInstances;
+	TArray<uint32> CardInstances;
+	MeshInstances.Reserve(Instances.Num());
+	CardInstances.Reserve(Instances.Num());
+	for (FHairStrandsInstance* AbstractInstance : Instances)
+	{
+		FHairGroupInstance* Instance = static_cast<FHairGroupInstance*>(AbstractInstance);
+		check(Instance->HairGroupPublicData);
+		
+		if (Instance->GeometryType == EHairGeometryType::Cards || Instance->GeometryType == EHairGeometryType::Meshes)
+		{
+			const uint32 HairLODIndex = Instance->HairGroupPublicData->LODIndex;
+			if (Instance->GeometryType == EHairGeometryType::Cards  && (!bCardSupported || !Instance->Cards.IsValid(HairLODIndex)))  { continue; }
+			if (Instance->GeometryType == EHairGeometryType::Meshes && (!bMeshSupported || !Instance->Meshes.IsValid(HairLODIndex))) { continue; }
+
+			if (Instance->GeometryType == EHairGeometryType::Cards)
+			{	
+				CardInstances.Add(InstanceDatas.Num());
+			}
+			else if (Instance->GeometryType == EHairGeometryType::Meshes)
+			{	
+				MeshInstances.Add(InstanceDatas.Num());
+			}
+
+			// Common instance parameters
+			FInstanceData& InstanceData = InstanceDatas.AddDefaulted_GetRef();
+			InstanceData.RegisteredIndex			= Instance->RegisteredIndex;
+			InstanceData.Instance 					= Instance;
+			InstanceData.BindingType 				= Instance->BindingType;
+			InstanceData.HairLODIndex 				= HairLODIndex;
+			InstanceData.bSimulationEnable 			= Instance->HairGroupPublicData->IsSimulationEnable(HairLODIndex);
+			InstanceData.bDeformationEnable 		= Instance->HairGroupPublicData->bIsDeformationEnable;
+			InstanceData.bGlobalDeformationEnable 	= Instance->HairGroupPublicData->IsGlobalInterpolationEnable(HairLODIndex);
+			InstanceData.MeshLODIndex				= -1;
+
+			// Binding surface parameters (skel.mesh/geom. cache)
+			FHairStrandsProjectionMeshData::FLOD MeshDataLOD;
+			const FCachedGeometry CachedGeometry = GetCacheGeometryForHair(GraphBuilder, Scene, Instance, ShaderMap, true);
+			for (int32 SectionIndex = 0; SectionIndex < CachedGeometry.Sections.Num(); ++SectionIndex)
+			{	
+				InstanceData.MeshDataLOD.Sections.Add(ConvertMeshSection(CachedGeometry, SectionIndex));
+			}
+			InstanceData.MeshLODIndex    = CachedGeometry.LODIndex;
+			Instance->Debug.MeshLODIndex = CachedGeometry.LODIndex;
+
+			// Card or Mesh specific parameters
+			if (Instance->GeometryType == EHairGeometryType::Cards)
+			{
+				InstanceData.CardInstance 			= &Instance->Cards.LODs[InstanceData.HairLODIndex];
+				InstanceData.bValidGuide			= InstanceData.Instance->Guides.bIsSimulationEnable || InstanceData.Instance->Guides.bHasGlobalInterpolation || InstanceData.Instance->Guides.bIsDeformationEnable || InstanceData.Instance->Guides.bIsSimulationCacheEnable;
+				InstanceData.bHasSkinning			= InstanceData.BindingType == EHairBindingType::Skinning && InstanceData.MeshLODIndex >= 0;
+				InstanceData.bNeedDeformation		= InstanceData.bValidGuide || InstanceData.bHasSkinning;
+				InstanceData.CardsSimulationType 	= EHairCardsSimulationType::None;
+				if (InstanceData.bHasSkinning || InstanceData.bValidGuide)
+				{
+					switch (GetHairCardsInterpolationType())
+					{
+						case 0 : InstanceData.CardsSimulationType = EHairCardsSimulationType::None; break;
+						case 1 : InstanceData.CardsSimulationType = EHairCardsSimulationType::Guide; break;
+						case 2 : InstanceData.CardsSimulationType = EHairCardsSimulationType::RBF;  break;
+					}
+				}
+
+				if (InstanceData.bNeedDeformation)
+				{
+					check(InstanceData.CardInstance->Guides.Data);
+				}
+				if (InstanceData.bSimulationEnable || InstanceData.bDeformationEnable)
+				{
+					check(InstanceData.CardInstance->Guides.IsValid());
+				}
+				if (InstanceData.BindingType == EHairBindingType::Skinning || InstanceData.bGlobalDeformationEnable)
+				{
+					check(InstanceData.CardInstance->Guides.IsValid());
+					check(InstanceData.CardInstance->Guides.HasValidRootData());
+					check(InstanceData.CardInstance->Guides.DeformedRootResource->IsValid(InstanceData.MeshLODIndex));
+				}
+			}
+			else if (Instance->GeometryType == EHairGeometryType::Meshes)
+			{
+				check(Instance->Meshes.IsValid(InstanceData.HairLODIndex));
+				InstanceData.bNeedDeformation 	= Instance->Meshes.LODs[InstanceData.HairLODIndex].DeformedResource != nullptr;
+				InstanceData.MeshInstance 		= &Instance->Meshes.LODs[InstanceData.HairLODIndex];
+				if (InstanceData.bNeedDeformation)
+				{
+					check(Instance->BindingType == EHairBindingType::Skinning);
+					check(Instance->Guides.IsValid());
+					check(Instance->Guides.HasValidRootData());
+					check(Instance->Guides.DeformedRootResource);
+					// MeshLODIndex -1 indicates that skin cache is disabled and this is a workaround to prevent the editor from crashing 
+					// An editor setting guildeline popup exists to inform the user that the skin cache should be enabled. 
+					check(InstanceData.MeshLODIndex == -1 || Instance->Guides.DeformedRootResource->IsValid(InstanceData.MeshLODIndex)); 
+				}
+			}
+
+			// Prepare VF input
+			InstanceData.Instance->HairGroupPublicData->VFInput.Cards  					= FHairGroupPublicData::FVertexFactoryInput::FCards();
+			InstanceData.Instance->HairGroupPublicData->VFInput.Meshes 					= FHairGroupPublicData::FVertexFactoryInput::FMeshes();
+			InstanceData.Instance->HairGroupPublicData->VFInput.GeometryType 			= InstanceData.Instance->GeometryType;
+			InstanceData.Instance->HairGroupPublicData->VFInput.LocalToWorldTransform 	= Instance->GetCurrentLocalToWorld();
+			InstanceData.Instance->HairGroupPublicData->bSupportVoxelization 			= false;
+		}
+	}
+
+	// Cards only - Update dynamic mesh triangles
+	for (uint32 InstanceIndex : CardInstances)
+	{
+		FInstanceData& InstanceData = InstanceDatas[InstanceIndex];
+		if (InstanceData.MeshLODIndex >= 0 && (InstanceData.BindingType == EHairBindingType::Skinning || InstanceData.bGlobalDeformationEnable))
+		{
+			AddHairStrandUpdateMeshTrianglesPass(
+				GraphBuilder,
+				ShaderMap,
+				InstanceData.MeshLODIndex,
+				InstanceData.MeshDataLOD,
+				InstanceData.CardInstance->Guides.RestRootResource,
+				InstanceData.CardInstance->Guides.DeformedRootResource);
+		}
+	}
+
+	// Cards only - Update position offset (GPU)
+	for (uint32 InstanceIndex : CardInstances)
+	{
+		FInstanceData& InstanceData = InstanceDatas[InstanceIndex];
+		if (InstanceData.MeshLODIndex >= 0 && (InstanceData.BindingType == EHairBindingType::Skinning || InstanceData.bGlobalDeformationEnable))
+		{
+			AddHairStrandUpdatePositionOffsetPass(
+				GraphBuilder,
+				ShaderMap,
+				EHairPositionUpdateType::Cards,
+				InstanceData.RegisteredIndex,
+				InstanceData.HairLODIndex,
+				InstanceData.MeshLODIndex,
+				InstanceData.CardInstance->Guides.DeformedRootResource,
+				InstanceData.CardInstance->Guides.DeformedResource);
+		}
+	}
+
+	// Cards only - Update position offset (CPU)
+	for (uint32 InstanceIndex : CardInstances)
+	{
+		FInstanceData& InstanceData = InstanceDatas[InstanceIndex];
+		if (InstanceData.MeshLODIndex >= 0 && InstanceData.BindingType != EHairBindingType::Skinning && !InstanceData.bGlobalDeformationEnable && (InstanceData.bSimulationEnable || InstanceData.bDeformationEnable))
+		{
+			AddHairStrandUpdatePositionOffsetPass(
+				GraphBuilder,
+				ShaderMap,
+				EHairPositionUpdateType::Cards,
+				InstanceData.RegisteredIndex,
+				InstanceData.HairLODIndex,
+				InstanceData.MeshLODIndex,
+				nullptr,
+				InstanceData.CardInstance->Guides.DeformedResource);
+		}
+	}
+
+	FRDGExternalAccessQueue ExternalAccessQueue;
+
+	// Cards only - Deform guide with skinning and/or simulation
+	for (uint32 InstanceIndex : CardInstances)
+	{
+		FInstanceData& InstanceData = InstanceDatas[InstanceIndex];
+		FHairGroupInstance::FCards::FLOD& LOD = *InstanceData.CardInstance;
+
+		if (InstanceData.bNeedDeformation)
+		{
+			// 1. Cards are deformed based on guides motion (simulation or RBF applied on guides)
+			if (InstanceData.CardsSimulationType == EHairCardsSimulationType::Guide)
+			{
+				FRDGBufferUAVRef Guides_DeformedPositionUAV = RegisterAsUAV(GraphBuilder, LOD.Guides.DeformedResource->GetBuffer(FHairStrandsDeformedResource::Current));
+
+				const bool bUseSingleGuide = LOD.Guides.InterpolationResource->UseSingleGuide();
+
+				FRDGHairStrandsCullingData CullingData;
+				AddHairStrandsInterpolationPass(
+					GraphBuilder,
+					ShaderMap,
+					ShaderPrintData,
+					InstanceData.Instance,
+					LOD.Guides.RestResource->GetPointCount(),
+					InstanceData.MeshLODIndex,
+					1.0f,
+					LOD.Guides.HairInterpolationType,
+					EHairGeometryType::Cards,
+					CullingData,
+					LOD.Guides.RestResource->GetPositionOffset(),
+					InstanceData.bValidGuide ? InstanceData.Instance->Guides.RestResource->GetPositionOffset() : FVector::ZeroVector,
+					RegisterAsSRV(GraphBuilder, LOD.Guides.DeformedResource->GetPositionOffsetBuffer(FHairStrandsDeformedResource::Current)),
+					InstanceData.bValidGuide ? RegisterAsSRV(GraphBuilder, InstanceData.Instance->Guides.DeformedResource->GetPositionOffsetBuffer(FHairStrandsDeformedResource::Current)) : nullptr,
+					InstanceData.bHasSkinning ? LOD.Guides.RestRootResource : nullptr ,
+					InstanceData.bHasSkinning && InstanceData.bValidGuide ? InstanceData.Instance->Guides.RestRootResource : nullptr,
+					InstanceData.bHasSkinning ? LOD.Guides.DeformedRootResource : nullptr,
+					InstanceData.bHasSkinning && InstanceData.bValidGuide ? InstanceData.Instance->Guides.DeformedRootResource : nullptr,
+					RegisterAsSRV(GraphBuilder, LOD.Guides.RestResource->PositionBuffer),
+					RegisterAsSRV(GraphBuilder, LOD.Guides.RestResource->PointToCurveBuffer),
+					bUseSingleGuide,
+					InstanceData.bValidGuide ? RegisterAsSRV(GraphBuilder, LOD.Guides.InterpolationResource->InterpolationBuffer) : nullptr,
+					InstanceData.bValidGuide ? RegisterAsSRV(GraphBuilder, InstanceData.Instance->Guides.RestResource->PositionBuffer) : nullptr,
+					InstanceData.bValidGuide ? RegisterAsSRV(GraphBuilder, InstanceData.Instance->Guides.DeformedResource->GetBuffer(FHairStrandsDeformedResource::Current)) : nullptr,
+					RegisterAsSRV(GraphBuilder, LOD.Guides.InterpolationResource->SimRootPointIndexBuffer),
+					InstanceData.bValidGuide ? RegisterAsSRV(GraphBuilder, InstanceData.Instance->Guides.RestResource->PointToCurveBuffer) : nullptr,
+					nullptr,
+					Guides_DeformedPositionUAV,
+					FHairStrandsDeformedRootResource::FLOD::Current); // <- this should be optional
+			}
+		}
+	}
+
+	// Cards only - Deform final cards geometry (using guides)
+	for (uint32 InstanceIndex : CardInstances)
+	{
+		FInstanceData& InstanceData = InstanceDatas[InstanceIndex];
+		FHairGroupInstance::FCards::FLOD& LOD = *InstanceData.CardInstance;
+
+		if (InstanceData.bNeedDeformation)
+		{
+			// 1. Cards are deformed based on guides motion (simulation or RBF applied on guides)
+			if (InstanceData.CardsSimulationType == EHairCardsSimulationType::Guide)
+			{
+				AddHairCardsDeformationPass(
+					GraphBuilder,
+					ShaderMap,
+					ShaderPrintData,
+					InstanceData.Instance,
+					InstanceData.MeshLODIndex);
+			}
+		}
+	}
+
+	// Cards only - Deform final cards geometry (using RBF)
+	for (uint32 InstanceIndex : CardInstances)
+	{
+		FInstanceData& InstanceData = InstanceDatas[InstanceIndex];
+		FHairGroupInstance::FCards::FLOD& LOD = *InstanceData.CardInstance;
+
+		if (InstanceData.bNeedDeformation)
+		{
+			// 2. Cards are deformed only based on skel. mesh RBF data)
+			if (InstanceData.CardsSimulationType == EHairCardsSimulationType::RBF)
+			{
+				AddHairCardsRBFInterpolationPass(
+					GraphBuilder,
+					ShaderMap,
+					InstanceData.MeshLODIndex,
+					LOD.RestResource,
+					LOD.DeformedResource,
+					InstanceData.Instance->Guides.RestRootResource,
+					InstanceData.Instance->Guides.DeformedRootResource);
+			}
+		}
+	}
+
+	for (uint32 InstanceIndex : CardInstances)
+	{
+		FInstanceData& InstanceData = InstanceDatas[InstanceIndex];
+		FHairGroupInstance::FCards::FLOD& LOD = *InstanceData.CardInstance;
+
+		if (InstanceData.bNeedDeformation)
+		{		
+			// 2. Cards are deformed only based on skel. mesh RBF data)
+			if (InstanceData.CardsSimulationType == EHairCardsSimulationType::RBF)
+			{
+				AddHairCardsRBFInterpolationPass(
+					GraphBuilder,
+					ShaderMap,
+					InstanceData.MeshLODIndex,
+					LOD.RestResource,
+					LOD.DeformedResource,
+					InstanceData.Instance->Guides.RestRootResource,
+					InstanceData.Instance->Guides.DeformedRootResource);
+			}
+
+			if (LOD.DeformedResource)
+			{
+				ExternalAccessQueue.Add(Register(GraphBuilder, LOD.DeformedResource->GetBuffer(FHairCardsDeformedResource::Current), ERDGImportedBufferFlags::None).Buffer, ERHIAccess::SRVMask);
+				ExternalAccessQueue.Add(Register(GraphBuilder, LOD.DeformedResource->GetBuffer(FHairCardsDeformedResource::Previous), ERDGImportedBufferFlags::None).Buffer, ERHIAccess::SRVMask);
+			}
+		}
+	}
+
+	// Cards only - Build/update RT geometry
+	#if RHI_RAYTRACING
+	if (IsRayTracingEnabled())
+	{
+		for (uint32 InstanceIndex : CardInstances)
+		{
+			FInstanceData& InstanceData = InstanceDatas[InstanceIndex];
+			if (InstanceData.CardInstance->RaytracingResource)
+			{
+				AddBuildHairCardAccelerationStructurePass(
+					GraphBuilder,
+					InstanceData.Instance,
+					InstanceData.HairLODIndex,
+					InstanceData.bNeedDeformation);
+			}
+		}
+	}
+	#endif
+
+	// Meshes only - Deform final mesh geometry (using RBF)
+	for (uint32 InstanceIndex : MeshInstances)
+	{
+		FInstanceData& InstanceData = InstanceDatas[InstanceIndex];
+		AddHairMeshesRBFInterpolationPass(
+			GraphBuilder,
+			ShaderMap,
+			InstanceData.MeshLODIndex,
+			InstanceData.MeshInstance->RestResource,
+			InstanceData.MeshInstance->DeformedResource,
+			InstanceData.Instance->Guides.RestRootResource,
+			InstanceData.Instance->Guides.DeformedRootResource);
+
+		ExternalAccessQueue.Add(Register(GraphBuilder, InstanceData.MeshInstance->DeformedResource->GetBuffer(FHairMeshesDeformedResource::Current), ERDGImportedBufferFlags::None).Buffer, ERHIAccess::SRVMask);
+	}
+
+	// Meshes only - Build/update RT geometry
+	#if RHI_RAYTRACING
+	if (IsRayTracingEnabled())
+	{
+		for (uint32 InstanceIndex : MeshInstances)
+		{
+			FInstanceData& InstanceData = InstanceDatas[InstanceIndex];
+			if (InstanceData.MeshInstance->RaytracingResource)
+			{
+				AddBuildHairMeshAccelerationStructurePass(
+					GraphBuilder,
+					InstanceData.Instance,
+					InstanceData.HairLODIndex,
+					InstanceData.bNeedDeformation);
+			}
+		}
+	}
+	#endif
+
+	ExternalAccessQueue.Submit(GraphBuilder);
 }
 
 // Return the LOD which should be used for a given screen size and LOD bias value
diff --git a/Engine/Plugins/Runtime/HairStrands/Source/HairStrandsCore/Private/HairStrandsInterpolation.cpp b/Engine/Plugins/Runtime/HairStrands/Source/HairStrandsCore/Private/HairStrandsInterpolation.cpp
--- a/Engine/Plugins/Runtime/HairStrands/Source/HairStrandsCore/Private/HairStrandsInterpolation.cpp
+++ b/Engine/Plugins/Runtime/HairStrands/Source/HairStrandsCore/Private/HairStrandsInterpolation.cpp
@@ -53,12 +53,10 @@ static FAutoConsoleVariableRef CVarHairRTGeomForceRebuild(TEXT("r.HairStrands.St
 
 ////////////////////////////////////////////////////////////////////////////////////////////////////////////////////////
 
-enum class EHairCardsSimulationType
+uint32 GetHairCardsInterpolationType()
 {
-	None,
-	Guide,
-	RBF
-};
+	return FMath::Clamp(GHairCardsInterpolationType, 0, 2);
+}
 
 bool IsHairStrandsTransferPositionOnLODChange()
 {
@@ -70,18 +68,6 @@ bool IsHairStrandsForceRebuildBVH()
 	return GHairStrands_Raytracing_ForceRebuildBVH > 0;
 }
 
-EHairCardsSimulationType GetHairCardsSimulationType()
-{
-	return GHairCardsInterpolationType >= 2 ? 
-		EHairCardsSimulationType::RBF : 
-		(GHairCardsInterpolationType >= 1 ? EHairCardsSimulationType::Guide : EHairCardsSimulationType::None);
-}
-
-bool NeedsUpdateCardsMeshTriangles()
-{
-	return GetHairCardsSimulationType() == EHairCardsSimulationType::Guide;
-}
-
 uint32 GetHairRaytracingProceduralSplits()
 {
 	return FMath::Clamp(GHairRaytracingProceduralSplits, 1, STRANDS_PROCEDURAL_INTERSECTOR_MAX_SPLITS);
@@ -941,7 +927,7 @@ public:
 
 IMPLEMENT_GLOBAL_SHADER(FHairCardsDeformationCS, "/Engine/Private/HairStrands/HairCardsDeformation.usf", "MainCS", SF_Compute);
 
-static void AddHairCardsDeformationPass(
+void AddHairCardsDeformationPass(
 	FRDGBuilder& GraphBuilder,
 	FGlobalShaderMap* ShaderMap,
 	const FShaderPrintData* ShaderPrintData,
@@ -1502,6 +1488,89 @@ void AddBuildStrandsAccelerationStructurePass(
 		});
 }
 
+void AddBuildHairCardAccelerationStructurePass(
+	FRDGBuilder& GraphBuilder,
+	FHairGroupInstance* Instance,
+	int32 HairLODIndex,
+	bool bNeedUpdate)
+{
+	check(Instance);
+	check(Instance->Cards.LODs.IsValidIndex(HairLODIndex));
+
+	FHairGroupInstance::FCards::FLOD& LOD = Instance->Cards.LODs[HairLODIndex];
+	FCardsOrMeshesResourceBLASParameters* Parameters = GraphBuilder.AllocParameters<FCardsOrMeshesResourceBLASParameters>();
+	Parameters->PositionBuffer = LOD.DeformedResource ? Register(GraphBuilder, LOD.DeformedResource->GetBuffer(FHairCardsDeformedResource::Current), ERDGImportedBufferFlags::None).Buffer : nullptr;
+
+	// * When cards are dynamic, RT geometry is built once and then update.
+	// * When cards are static (i.e., no sim, not attached to  skinning()), in which case RT geometry needs only to be built once. This static geometry is shared between all the instances, 
+	//   and owned by the asset, rathter than the component. In this case the RT geometry will be built only once for all the instances
+	const bool bNeedBuild = !LOD.RaytracingResource->bIsRTGeometryInitialized;
+	if (bNeedBuild || bNeedUpdate)
+	{
+		GraphBuilder.AddPass(
+			RDG_EVENT_NAME("HairStrands::UpdateBLAS(Cards)"),
+			Parameters,
+			ERDGPassFlags::NeverCull | ERDGPassFlags::Compute,
+			[Instance, HairLODIndex, bNeedUpdate](FRHICommandListImmediate& RHICmdList)
+		{
+			SCOPED_GPU_MASK(RHICmdList, FRHIGPUMask::All());
+
+			FHairGroupInstance::FCards::FLOD& LocalLOD = Instance->Cards.LODs[HairLODIndex];
+
+			const bool bLocalNeedBuild = !LocalLOD.RaytracingResource->bIsRTGeometryInitialized;
+			if (bLocalNeedBuild)
+			{
+				BuildHairAccelerationStructure_Cards(RHICmdList, LocalLOD.RestResource, LocalLOD.DeformedResource, &LocalLOD.RaytracingResource->RayTracingGeometry, Instance->Debug.GroomAssetName, Instance->Debug.MeshComponentName, HairLODIndex);
+				LocalLOD.RaytracingResource->bIsRTGeometryInitialized = true;
+			}
+			else if (bNeedUpdate)
+			{
+				// TODO: evaluate perf tradeoff of rebuild vs refit
+				UpdateHairAccelerationStructure(RHICmdList, &LocalLOD.RaytracingResource->RayTracingGeometry, EAccelerationStructureBuildMode::Update);
+			}
+		});
+	}
+}
+
+void AddBuildHairMeshAccelerationStructurePass(
+	FRDGBuilder& GraphBuilder,
+	FHairGroupInstance* Instance,
+	int32 HairLODIndex,
+	bool bNeedUpdate)
+{
+	check(Instance);
+	check(Instance->Meshes.LODs.IsValidIndex(HairLODIndex));
+
+	FHairGroupInstance::FMeshes::FLOD& LOD = Instance->Meshes.LODs[HairLODIndex];
+	FCardsOrMeshesResourceBLASParameters* Parameters = GraphBuilder.AllocParameters<FCardsOrMeshesResourceBLASParameters>();
+	Parameters->PositionBuffer = LOD.DeformedResource ? Register(GraphBuilder, LOD.DeformedResource->GetBuffer(FHairMeshesDeformedResource::Current), ERDGImportedBufferFlags::None).Buffer : nullptr;
+
+	const bool bNeedBuid = !LOD.RaytracingResource->bIsRTGeometryInitialized;
+	if (bNeedBuid || bNeedUpdate)
+	{
+		GraphBuilder.AddPass(
+			RDG_EVENT_NAME("HairStrands::UpdateBLAS(Meshes)"),
+			Parameters,
+			ERDGPassFlags::NeverCull | ERDGPassFlags::Compute,
+			[Instance, HairLODIndex, bNeedUpdate](FRHICommandListImmediate& RHICmdList)
+			{
+				SCOPED_GPU_MASK(RHICmdList, FRHIGPUMask::All());
+
+				FHairGroupInstance::FMeshes::FLOD& LocalLOD = Instance->Meshes.LODs[HairLODIndex];	
+				const bool bLocalNeedBuild = !LocalLOD.RaytracingResource->bIsRTGeometryInitialized;
+				if (bLocalNeedBuild)
+				{
+					BuildHairAccelerationStructure_Meshes(RHICmdList, LocalLOD.RestResource, LocalLOD.DeformedResource,  &LocalLOD.RaytracingResource->RayTracingGeometry, Instance->Debug.GroomAssetName, HairLODIndex);
+					LocalLOD.RaytracingResource->bIsRTGeometryInitialized = true;
+				}
+				else if (bNeedUpdate)
+				{
+					// TODO: evaluate perf tradeoff of rebuild vs refit
+					UpdateHairAccelerationStructure(RHICmdList, &LocalLOD.RaytracingResource->RayTracingGeometry, EAccelerationStructureBuildMode::Update);
+				}
+			});
+	}
+}
 #endif // RHI_RAYTRACING
 
 static void ConvertHairStrandsVFParameters(
@@ -1673,237 +1742,4 @@ EGroomCacheType GetHairInstanceCacheType(const FHairGroupInstance* Instance)
 		Instance->Debug.GroomCacheBuffers->GetCurrentFrameBuffer().GroupsData.IsValidIndex(Instance->Debug.GroupIndex) && 
 		Instance->Debug.GroomCacheBuffers->GetNextFrameBuffer().GroupsData.IsValidIndex(Instance->Debug.GroupIndex);
 	return bHasValidGroomCacheBuffers ? Instance->Debug.GroomCacheType : EGroomCacheType::None;
-}
-
-void ComputeHairStrandsInterpolation(
-	FRDGBuilder& GraphBuilder,
-	FGlobalShaderMap* ShaderMap,
-	const uint32 ViewUniqueID,
-	const uint32 ViewRayTracingMask,
-	const EGroomViewMode ViewMode,
-	const FVector& TranslatedWorldOffset,
-	const FShaderPrintData* ShaderPrintData,
-	FHairGroupInstance* Instance,
-	int32 MeshLODIndex,
-	FHairStrandClusterData* InClusterData)
-{	
-	if (!Instance)
-	{
-		return;
-	}
-
-	// Reset
-	Instance->HairGroupPublicData->VFInput.Strands	= FHairGroupPublicData::FVertexFactoryInput::FStrands();
-	Instance->HairGroupPublicData->VFInput.Cards	= FHairGroupPublicData::FVertexFactoryInput::FCards();
-	Instance->HairGroupPublicData->VFInput.Meshes	= FHairGroupPublicData::FVertexFactoryInput::FMeshes();
-	const EHairGeometryType InstanceGeometryType	= Instance->GeometryType;
-
-	FRDGExternalAccessQueue ExternalAccessQueue;
-
-	if (InstanceGeometryType == EHairGeometryType::Cards)
-	{	
-		DECLARE_GPU_STAT(HairCardsInterpolation);
-		RDG_EVENT_SCOPE(GraphBuilder, "HairInterpolation(Cards)");
-		RDG_GPU_STAT_SCOPE(GraphBuilder, HairCardsInterpolation);
-
-		const uint32 HairLODIndex = Instance->HairGroupPublicData->GetIntLODIndex();
-		const bool bIsCardsValid = Instance->Cards.IsValid(HairLODIndex);
-		if (bIsCardsValid)
-		{
-			const bool bValidGuide = Instance->Guides.bIsSimulationEnable || Instance->Guides.bHasGlobalInterpolation || Instance->Guides.bIsDeformationEnable || Instance->Guides.bIsSimulationCacheEnable;
-			const bool bHasSkinning = Instance->BindingType == EHairBindingType::Skinning && MeshLODIndex >= 0;
-			const bool bNeedDeformation = bValidGuide || bHasSkinning;
-
-			FHairGroupInstance::FCards::FLOD& LOD = Instance->Cards.LODs[HairLODIndex];
-
-			if (bNeedDeformation)
-			{
-				check(LOD.Guides.Data);
-
-				const EHairCardsSimulationType CardsSimulationType = 
-					bHasSkinning || bValidGuide ?
-					GetHairCardsSimulationType() :
-					EHairCardsSimulationType::None;
-
-				// 1. Cards are deformed based on guides motion (simulation or RBF applied on guides)
-				if (CardsSimulationType == EHairCardsSimulationType::Guide)
-				{
-					FRDGBufferUAVRef Guides_DeformedPositionUAV = RegisterAsUAV(GraphBuilder, LOD.Guides.DeformedResource->GetBuffer(FHairStrandsDeformedResource::Current));
-
-					const bool bUseSingleGuide = LOD.Guides.InterpolationResource->UseSingleGuide();
-
-					FRDGHairStrandsCullingData CullingData;
-					AddHairStrandsInterpolationPass(
-						GraphBuilder,
-						ShaderMap,
-						ShaderPrintData,
-						Instance,
-						LOD.Guides.RestResource->GetPointCount(),
-						MeshLODIndex,
-						1.0f,
-						LOD.Guides.HairInterpolationType,
-						InstanceGeometryType,
-						CullingData,
-						LOD.Guides.RestResource->GetPositionOffset(),
-						bValidGuide ? Instance->Guides.RestResource->GetPositionOffset() : FVector::ZeroVector,
-						RegisterAsSRV(GraphBuilder, LOD.Guides.DeformedResource->GetPositionOffsetBuffer(FHairStrandsDeformedResource::Current)),
-						bValidGuide ? RegisterAsSRV(GraphBuilder, Instance->Guides.DeformedResource->GetPositionOffsetBuffer(FHairStrandsDeformedResource::Current)) : nullptr,
-						bHasSkinning ? LOD.Guides.RestRootResource : nullptr ,
-						bHasSkinning && bValidGuide ? Instance->Guides.RestRootResource : nullptr,
-						bHasSkinning ? LOD.Guides.DeformedRootResource : nullptr,
-						bHasSkinning && bValidGuide ? Instance->Guides.DeformedRootResource : nullptr,
-						RegisterAsSRV(GraphBuilder, LOD.Guides.RestResource->PositionBuffer),
-						RegisterAsSRV(GraphBuilder, LOD.Guides.RestResource->PointToCurveBuffer),
-						bUseSingleGuide,
-						bValidGuide ? RegisterAsSRV(GraphBuilder, LOD.Guides.InterpolationResource->InterpolationBuffer) : nullptr,
-						bValidGuide ? RegisterAsSRV(GraphBuilder, Instance->Guides.RestResource->PositionBuffer) : nullptr,
-						bValidGuide ? RegisterAsSRV(GraphBuilder, Instance->Guides.DeformedResource->GetBuffer(FHairStrandsDeformedResource::Current)) : nullptr,
-						RegisterAsSRV(GraphBuilder, LOD.Guides.InterpolationResource->SimRootPointIndexBuffer),
-						bValidGuide ? RegisterAsSRV(GraphBuilder, Instance->Guides.RestResource->PointToCurveBuffer) : nullptr,
-						nullptr,
-						Guides_DeformedPositionUAV,
-						FHairStrandsDeformedRootResource::FLOD::Current); // <- this should be optional
-
-					AddHairCardsDeformationPass(
-						GraphBuilder,
-						ShaderMap,
-						ShaderPrintData,
-						Instance,
-						MeshLODIndex);
-				}
-				// 2. Cards are deformed only based on skel. mesh RBF data)
-				else if (CardsSimulationType == EHairCardsSimulationType::RBF)
-				{
-					AddHairCardsRBFInterpolationPass(
-						GraphBuilder,
-						ShaderMap,
-						MeshLODIndex,
-						LOD.RestResource,
-						LOD.DeformedResource,
-						Instance->Guides.RestRootResource,
-						Instance->Guides.DeformedRootResource);
-				}
-
-				if (LOD.DeformedResource)
-				{
-					ExternalAccessQueue.Add(Register(GraphBuilder, LOD.DeformedResource->GetBuffer(FHairCardsDeformedResource::Current), ERDGImportedBufferFlags::None).Buffer, ERHIAccess::SRVMask);
-					ExternalAccessQueue.Add(Register(GraphBuilder, LOD.DeformedResource->GetBuffer(FHairCardsDeformedResource::Previous), ERDGImportedBufferFlags::None).Buffer, ERHIAccess::SRVMask);
-				}
-			}
-
-			#if RHI_RAYTRACING
-			if (LOD.RaytracingResource && IsRayTracingEnabled())
-			{
-				FCardsOrMeshesResourceBLASParameters* Parameters = GraphBuilder.AllocParameters<FCardsOrMeshesResourceBLASParameters>();
-				Parameters->PositionBuffer = LOD.DeformedResource ? Register(GraphBuilder, LOD.DeformedResource->GetBuffer(FHairCardsDeformedResource::Current), ERDGImportedBufferFlags::None).Buffer : nullptr;
-
-				// * When cards are dynamic, RT geometry is built once and then update.
-				// * When cards are static (i.e., no sim, not attached to  skinning()), in which case RT geometry needs only to be built once. This static geometry is shared between all the instances, 
-				//   and owned by the asset, rathter than the component. In this case the RT geometry will be built only once for all the instances
-				const bool bNeedUpdate = bNeedDeformation;
-				const bool bNeedBuild = !LOD.RaytracingResource->bIsRTGeometryInitialized;
-				if (bNeedBuild || bNeedUpdate)
-				{
-					GraphBuilder.AddPass(
-						RDG_EVENT_NAME("HairStrands::UpdateBLAS(Cards)"),
-						Parameters,
-						ERDGPassFlags::NeverCull | ERDGPassFlags::Compute,
-						[Instance, HairLODIndex, bNeedUpdate](FRHICommandListImmediate& RHICmdList)
-					{
-						SCOPED_GPU_MASK(RHICmdList, FRHIGPUMask::All());
-
-						FHairGroupInstance::FCards::FLOD& LocalLOD = Instance->Cards.LODs[HairLODIndex];
-
-						const bool bLocalNeedBuild = !LocalLOD.RaytracingResource->bIsRTGeometryInitialized;
-						if (bLocalNeedBuild)
-						{
-							BuildHairAccelerationStructure_Cards(RHICmdList, LocalLOD.RestResource, LocalLOD.DeformedResource, &LocalLOD.RaytracingResource->RayTracingGeometry, Instance->Debug.GroomAssetName, Instance->Debug.MeshComponentName, HairLODIndex);
-							LocalLOD.RaytracingResource->bIsRTGeometryInitialized = true;
-						}
-						else if (bNeedUpdate)
-						{
-							// TODO: evaluate perf tradeoff of rebuild vs refit
-							UpdateHairAccelerationStructure(RHICmdList, &LocalLOD.RaytracingResource->RayTracingGeometry, EAccelerationStructureBuildMode::Update);
-						}
-					});
-				}
-			}
-			#endif
-		}
-	}
-	else if (InstanceGeometryType == EHairGeometryType::Meshes)
-	{
-		DECLARE_GPU_STAT(HairMeshesInterpolation);
-		RDG_EVENT_SCOPE(GraphBuilder, "HairInterpolation(Meshes)");
-		RDG_GPU_STAT_SCOPE(GraphBuilder, HairMeshesInterpolation);
-
-		const uint32 HairLODIndex = Instance->HairGroupPublicData->GetIntLODIndex();
-		if (Instance->Meshes.IsValid(HairLODIndex))
-		{
-			const bool bNeedDeformation = Instance->Meshes.LODs[HairLODIndex].DeformedResource != nullptr;
-			if (bNeedDeformation)
-			{
-				FHairGroupInstance::FMeshes::FLOD& MeshesInstance = Instance->Meshes.LODs[HairLODIndex];
-
-				check(Instance->BindingType == EHairBindingType::Skinning)
-				check(Instance->Guides.IsValid());
-				check(Instance->Guides.HasValidRootData());
-				check(Instance->Guides.DeformedRootResource);
-				check(MeshLODIndex == -1 || Instance->Guides.DeformedRootResource->IsValid(MeshLODIndex)); //MeshLODIndex -1 indicates that skin cache is disabled and this is a workaround to prevent the editor from crashing - an editor setting guildeline popup exists to inform the user that the skin cache should be enabled. 
-
-				AddHairMeshesRBFInterpolationPass(
-					GraphBuilder,
-					ShaderMap,
-					MeshLODIndex,
-					MeshesInstance.RestResource,
-					MeshesInstance.DeformedResource,
-					Instance->Guides.RestRootResource,
-					Instance->Guides.DeformedRootResource);
-
-				ExternalAccessQueue.Add(Register(GraphBuilder, MeshesInstance.DeformedResource->GetBuffer(FHairMeshesDeformedResource::Current), ERDGImportedBufferFlags::None).Buffer, ERHIAccess::SRVMask);
-			}
-
-			#if RHI_RAYTRACING
-			FHairGroupInstance::FMeshes::FLOD& LOD = Instance->Meshes.LODs[HairLODIndex];
-			if (LOD.RaytracingResource && IsRayTracingEnabled())
-			{
-				FCardsOrMeshesResourceBLASParameters* Parameters = GraphBuilder.AllocParameters<FCardsOrMeshesResourceBLASParameters>();
-				Parameters->PositionBuffer = LOD.DeformedResource ? Register(GraphBuilder, LOD.DeformedResource->GetBuffer(FHairMeshesDeformedResource::Current), ERDGImportedBufferFlags::None).Buffer : nullptr;
-
-				const bool bNeedUpdate = bNeedDeformation;
-				const bool bNeedBuid = !LOD.RaytracingResource->bIsRTGeometryInitialized;
-				if (bNeedBuid || bNeedUpdate)
-				{
-					GraphBuilder.AddPass(
-						RDG_EVENT_NAME("HairStrands::UpdateBLAS(Meshes)"),
-						Parameters,
-						ERDGPassFlags::NeverCull | ERDGPassFlags::Compute,
-						[Instance, HairLODIndex, bNeedUpdate](FRHICommandListImmediate& RHICmdList)
-					{
-						SCOPED_GPU_MASK(RHICmdList, FRHIGPUMask::All());
-
-						FHairGroupInstance::FMeshes::FLOD& LocalLOD = Instance->Meshes.LODs[HairLODIndex];				
-						const bool bLocalNeedBuild = !LocalLOD.RaytracingResource->bIsRTGeometryInitialized;
-						if (bLocalNeedBuild)
-						{
-							BuildHairAccelerationStructure_Meshes(RHICmdList, LocalLOD.RestResource, LocalLOD.DeformedResource,  &LocalLOD.RaytracingResource->RayTracingGeometry, Instance->Debug.GroomAssetName, HairLODIndex);
-							LocalLOD.RaytracingResource->bIsRTGeometryInitialized = true;
-						}
-						else if (bNeedUpdate)
-						{
-							// TODO: evaluate perf tradeoff of rebuild vs refit
-							UpdateHairAccelerationStructure(RHICmdList, &LocalLOD.RaytracingResource->RayTracingGeometry, EAccelerationStructureBuildMode::Update);
-						}
-					});
-				}
-			}
-			#endif
-		}
-	}
-
-	Instance->HairGroupPublicData->VFInput.GeometryType = InstanceGeometryType;
-	Instance->HairGroupPublicData->VFInput.LocalToWorldTransform = Instance->GetCurrentLocalToWorld();
-	Instance->HairGroupPublicData->bSupportVoxelization = Instance->Strands.Modifier.bSupportVoxelization && Instance->bCastShadow;
-
-	ExternalAccessQueue.Submit(GraphBuilder);
-}
+}
\ No newline at end of file
diff --git a/Engine/Plugins/Runtime/HairStrands/Source/HairStrandsCore/Public/HairStrandsInterpolation.h b/Engine/Plugins/Runtime/HairStrands/Source/HairStrandsCore/Public/HairStrandsInterpolation.h
--- a/Engine/Plugins/Runtime/HairStrands/Source/HairStrandsCore/Public/HairStrandsInterpolation.h
+++ b/Engine/Plugins/Runtime/HairStrands/Source/HairStrandsCore/Public/HairStrandsInterpolation.h
@@ -208,4 +208,23 @@ void AddBuildStrandsAccelerationStructurePass(
 	bool bNeedUpdate,
 	FRDGBufferRef Raytracing_PositionBuffer,
 	FRDGBufferRef Raytracing_IndexBuffer);
-#endif
\ No newline at end of file
+
+void AddBuildHairCardAccelerationStructurePass(
+	FRDGBuilder& GraphBuilder,
+	FHairGroupInstance* Instance,
+	int32 HairLODIndex,
+	bool bNeedUpdate);
+
+void AddBuildHairMeshAccelerationStructurePass(
+	FRDGBuilder& GraphBuilder,
+	FHairGroupInstance* Instance,
+	int32 HairLODIndex,
+	bool bNeedUpdate);
+#endif
+
+void AddHairCardsDeformationPass(
+	FRDGBuilder& GraphBuilder,
+	FGlobalShaderMap* ShaderMap,
+	const FShaderPrintData* ShaderPrintData,
+	FHairGroupInstance* Instance,
+	const int32 MeshLODIndex);
\ No newline at end of file
