When running with manual skin cache only update the needed/required sections.

Save about 20% of update cost.

diff --git a/Engine/Plugins/Runtime/HairStrands/Source/HairStrandsCore/Private/GroomBindingBuilder.cpp b/Engine/Plugins/Runtime/HairStrands/Source/HairStrandsCore/Private/GroomBindingBuilder.cpp
--- a/Engine/Plugins/Runtime/HairStrands/Source/HairStrandsCore/Private/GroomBindingBuilder.cpp
+++ b/Engine/Plugins/Runtime/HairStrands/Source/HairStrandsCore/Private/GroomBindingBuilder.cpp
@@ -60,7 +60,7 @@ static FAutoConsoleVariableRef CVarHairStrandsBindingBuilderWarningEnable(TEXT("
 FString FGroomBindingBuilder::GetVersion()
 {
 	// Important to update the version when groom building changes
-	return TEXT("3p_14");
+	return TEXT("3p_15");
 }
 
 ///////////////////////////////////////////////////////////////////////////////////////////////////
@@ -2044,6 +2044,11 @@ static void BuildRootBulkData(
 				const uint32 Index = In.MeshProjectionLODs[MeshLODIt].MeshSampleIndicesBuffer[SampleIt];
 				const uint32 SectionIndex = In.MeshProjectionLODs[MeshLODIt].MeshSampleSectionsBuffer[SampleIt];
 				MeshSampleIndicesAndSectionBuffer[SampleIt] = FHairStrandsRootUtils::PackTriangleIndex(Index, SectionIndex);
+
+				// Update the unique section indices with section containing RBF samples
+				// This allows faster update at runtime when not using skin cache
+				// This is done only for guides, which is the only root data containing RBF data
+				Out.Header.LODs[MeshLODIt].UniqueSectionIndices.AddUnique(SectionIndex);
 			}
 
 			CopyToBulkData<FHairStrandsWeightFormat>(Out.Data.LODs[MeshLODIt].MeshInterpolationWeightsBuffer, In.MeshProjectionLODs[MeshLODIt].MeshInterpolationWeightsBuffer);
diff --git a/Engine/Plugins/Runtime/HairStrands/Source/HairStrandsCore/Private/GroomDebug.cpp b/Engine/Plugins/Runtime/HairStrands/Source/HairStrandsCore/Private/GroomDebug.cpp
--- a/Engine/Plugins/Runtime/HairStrands/Source/HairStrandsCore/Private/GroomDebug.cpp
+++ b/Engine/Plugins/Runtime/HairStrands/Source/HairStrandsCore/Private/GroomDebug.cpp
@@ -71,9 +71,10 @@ enum class HairStrandsTriangleType
 
 FCachedGeometry GetCacheGeometryForHair(
 	FRDGBuilder& GraphBuilder,
+	FGlobalShaderMap* ShaderMap,
 	FSceneInterface* Scene,
 	FHairGroupInstance* Instance,
-	FGlobalShaderMap* ShaderMap,
+	const FHairStrandsRootBulkData* RootBulkData,
 	const bool bOutputTriangleData);
 
 static void GetGroomInterpolationData(
@@ -89,21 +90,42 @@ static void GetGroomInterpolationData(
 		FHairGroupInstance* Instance = static_cast<FHairGroupInstance*>(AbstractInstance);
 
 		if (!Instance)
+		{
 			continue;
+		}
+
+		const FHairStrandsRestRootResource* RestRootResource = nullptr;
+		if (Instance->GetHairGeometry() == EHairGeometryType::Strands)
+		{
+			RestRootResource = Instance->Strands.RestRootResource;
+		}
+		else if (Instance->GetHairGeometry() == EHairGeometryType::Cards)
+		{
+			const int32 LODIndex = Instance->HairGroupPublicData->GetIntLODIndex();
+			if (Instance->Cards.IsValid(LODIndex))
+			{
+				RestRootResource = Instance->Cards.LODs[LODIndex].Guides.RestRootResource;
+			}
+		}
+
+		if (!RestRootResource)
+		{
+			continue;
+		}
 
-		const FCachedGeometry CachedGeometry = GetCacheGeometryForHair(GraphBuilder, Scene, Instance, ShaderMap, true);
+		const FCachedGeometry CachedGeometry = GetCacheGeometryForHair(GraphBuilder, ShaderMap, Scene, Instance, &RestRootResource->BulkData, true);
 		if (CachedGeometry.Sections.Num() == 0)
 			continue;
 
 		if (MeshType == EHairStrandsProjectionMeshType::DeformedMesh || MeshType == EHairStrandsProjectionMeshType::RestMesh)
 		{
-			for (int32 SectionIndex = 0; SectionIndex < CachedGeometry.Sections.Num(); ++SectionIndex)
+			for (const FCachedGeometry::Section& CachedGeometrySection : CachedGeometry.Sections)
 			{
-				FHairStrandsProjectionMeshData::FSection OutSection = ConvertMeshSection(CachedGeometry, SectionIndex);
+				FHairStrandsProjectionMeshData::FSection OutSection = ConvertMeshSection(CachedGeometry, CachedGeometrySection);
 				if (MeshType == EHairStrandsProjectionMeshType::RestMesh)
 				{					
 					// If the mesh has some mesh-tranferred data, we display that otherwise we use the rest data
-					const int32 SectionLodIndex = CachedGeometry.Sections[SectionIndex].LODIndex;
+					const int32 SectionLodIndex = CachedGeometrySection .LODIndex;
 					const bool bHasTransferData = SectionLodIndex < Instance->Debug.TransferredPositions.Num();
 					if (bHasTransferData)
 					{
diff --git a/Engine/Plugins/Runtime/HairStrands/Source/HairStrandsCore/Private/GroomGeometryCache.cpp b/Engine/Plugins/Runtime/HairStrands/Source/HairStrandsCore/Private/GroomGeometryCache.cpp
--- a/Engine/Plugins/Runtime/HairStrands/Source/HairStrandsCore/Private/GroomGeometryCache.cpp
+++ b/Engine/Plugins/Runtime/HairStrands/Source/HairStrandsCore/Private/GroomGeometryCache.cpp
@@ -16,12 +16,12 @@
 #include "Rendering/SkeletalMeshLODRenderData.h"
 #include "RenderGraphUtils.h"
 #include "SkeletalMeshDeformerHelpers.h"
+#include "HairStrandsDatas.h"
 
 ///////////////////////////////////////////////////////////////////////////////////////////////////
 
-FHairStrandsProjectionMeshData::FSection ConvertMeshSection(FCachedGeometry const& InCachedGeometry, int32 InSectionIndex)
+FHairStrandsProjectionMeshData::FSection ConvertMeshSection(const FCachedGeometry& InCachedGeometry, const FCachedGeometry::Section & In)
 {
-	FCachedGeometry::Section const& In = InCachedGeometry.Sections[InSectionIndex];
 	FHairStrandsProjectionMeshData::FSection Out;
 	Out.IndexBuffer = In.IndexBuffer;
 	Out.RDGPositionBuffer = In.RDGPositionBuffer;
@@ -52,6 +52,7 @@ ENGINE_API void UpdatePreviousRefToLocalMatrices(TArray<FMatrix44f>& ReferenceTo
 	FRDGBuilder& GraphBuilder,
 	FGlobalShaderMap* ShaderMap, 
 	const FSkeletalMeshSceneProxy* Proxy,
+	const FHairStrandsRootBulkData* RootBulkData,
 	const bool bOutputTriangleData,
 	FCachedGeometry& Out)
 {
@@ -91,18 +92,24 @@ ENGINE_API void UpdatePreviousRefToLocalMatrices(TArray<FMatrix44f>& ReferenceTo
 	FRDGBufferSRVRef DeformedPositionSRV		= GraphBuilder.CreateSRV(Out.DeformedPositionBuffer, PF_R32_FLOAT);
 	FRDGBufferSRVRef DeformedPreviousPositionSRV= bNeedPreviousPosition ? GraphBuilder.CreateSRV(Out.DeformedPreviousPositionBuffer, PF_R32_FLOAT) : nullptr;
 
+	check(RootBulkData);
+	check(RootBulkData->Header.LODs.IsValidIndex(LODIndex));
+	const TArray<uint32>& UniqueSectionCount = RootBulkData->Header.LODs[LODIndex].UniqueSectionIndices;
+	const uint32 EffectiveSectionCount = UniqueSectionCount.Num();
+
 	// Fill in result
-	TArray<FSkinUpdateSection> Sections;
-	Sections.SetNum(SectionCount);
-	for (uint32 SectionIt = 0; SectionIt < SectionCount; ++SectionIt)
+	TArray<FSkinUpdateSection> UpdateSections;
+	UpdateSections.Reserve(EffectiveSectionCount);
+	for (uint32 SectionIndex : UniqueSectionCount)
 	{
-		const FSkelMeshRenderSection& Section = LODData.RenderSections[SectionIt];
+		const FSkelMeshRenderSection& Section = LODData.RenderSections[SectionIndex];
 
-		Sections[SectionIt].SectionIndex 			= SectionIt;
-		Sections[SectionIt].NumVertexToProcess 		= Section.NumVertices;
-		Sections[SectionIt].SectionVertexBaseIndex 	= Section.BaseVertexIndex;
-		Sections[SectionIt].BoneBuffer 				= FSkeletalMeshDeformerHelpers::GetBoneBufferForReading(SkeletalMeshObject, LODIndex, SectionIt, false);
-		Sections[SectionIt].BonePrevBuffer 			= FSkeletalMeshDeformerHelpers::GetBoneBufferForReading(SkeletalMeshObject, LODIndex, SectionIt, true);
+		FSkinUpdateSection& UpdateSection 	= UpdateSections.AddDefaulted_GetRef();
+		UpdateSection.SectionIndex 			= SectionIndex;
+		UpdateSection.NumVertexToProcess 	= Section.NumVertices;
+		UpdateSection.SectionVertexBaseIndex= Section.BaseVertexIndex;
+		UpdateSection.BoneBuffer 			= FSkeletalMeshDeformerHelpers::GetBoneBufferForReading(SkeletalMeshObject, LODIndex, SectionIndex, false);
+		UpdateSection.BonePrevBuffer 		= FSkeletalMeshDeformerHelpers::GetBoneBufferForReading(SkeletalMeshObject, LODIndex, SectionIndex, true);
 
 		FCachedGeometry::Section& OutSection= Out.Sections.AddDefaulted_GetRef();
 		OutSection.RDGPositionBuffer 		= DeformedPositionSRV;
@@ -118,7 +125,7 @@ ENGINE_API void UpdatePreviousRefToLocalMatrices(TArray<FMatrix44f>& ReferenceTo
 		OutSection.NumVertices 				= Section.NumVertices;
 		OutSection.IndexBaseIndex 			= Section.BaseIndex;
 		OutSection.VertexBaseIndex 			= Section.BaseVertexIndex;
-		OutSection.SectionIndex 			= SectionIt;
+		OutSection.SectionIndex 			= SectionIndex;
 		OutSection.LODIndex 				= LODIndex;
 		OutSection.UVsChannelOffset 		= 0; // Assume that we needs to pair meshes based on UVs 0
 	}
@@ -127,7 +134,7 @@ ENGINE_API void UpdatePreviousRefToLocalMatrices(TArray<FMatrix44f>& ReferenceTo
 		GraphBuilder, 
 		ShaderMap, 
 		LODData, 
-		Sections, 
+		UpdateSections, 
 		Out.DeformedPositionBuffer, 
 		Out.DeformedPreviousPositionBuffer);
 }
diff --git a/Engine/Plugins/Runtime/HairStrands/Source/HairStrandsCore/Private/GroomGeometryCache.h b/Engine/Plugins/Runtime/HairStrands/Source/HairStrandsCore/Private/GroomGeometryCache.h
--- a/Engine/Plugins/Runtime/HairStrands/Source/HairStrandsCore/Private/GroomGeometryCache.h
+++ b/Engine/Plugins/Runtime/HairStrands/Source/HairStrandsCore/Private/GroomGeometryCache.h
@@ -4,19 +4,22 @@
 
 #include "CoreMinimal.h"
 #include "HairStrandsMeshProjection.h"
+#include "CachedGeometry.h"
 
 struct FCachedGeometry;
 class FSkeletalMeshSceneProxy;
 class FGlobalShaderMap;
 class FRDGBuilder;
 class FGeometryCacheSceneProxy;
+struct FHairStrandsRootBulkData;
 
-FHairStrandsProjectionMeshData::FSection ConvertMeshSection(FCachedGeometry const& InCachedGeometry, int32 InSectionIndex);
+FHairStrandsProjectionMeshData::FSection ConvertMeshSection(const FCachedGeometry& InCachedGeometry, const FCachedGeometry::Section& In);
 
 void BuildCacheGeometry(
 	FRDGBuilder& GraphBuilder,
 	FGlobalShaderMap* ShaderMap, 
 	const FSkeletalMeshSceneProxy* Proxy,
+	const FHairStrandsRootBulkData* RootBulkData,
 	const bool bOutputTriangleData,
 	FCachedGeometry& OutCachedGeometry);
 
diff --git a/Engine/Plugins/Runtime/HairStrands/Source/HairStrandsCore/Private/GroomManager.cpp b/Engine/Plugins/Runtime/HairStrands/Source/HairStrandsCore/Private/GroomManager.cpp
--- a/Engine/Plugins/Runtime/HairStrands/Source/HairStrandsCore/Private/GroomManager.cpp
+++ b/Engine/Plugins/Runtime/HairStrands/Source/HairStrandsCore/Private/GroomManager.cpp
@@ -191,9 +191,10 @@ static const FPrimitiveSceneInfo* GetMeshSceneInfo(FSceneInterface* Scene, FHair
 // Returns the cached geometry of the underlying geometry on which a hair instance is attached to
 FCachedGeometry GetCacheGeometryForHair(
 	FRDGBuilder& GraphBuilder, 
+	FGlobalShaderMap* ShaderMap,
 	FSceneInterface* Scene,
 	FHairGroupInstance* Instance, 
-	FGlobalShaderMap* ShaderMap,
+	const FHairStrandsRootBulkData* RootBulkData,
 	const bool bOutputTriangleData)
 {
 	FCachedGeometry Out;
@@ -211,7 +212,7 @@ FCachedGeometry GetCacheGeometryForHair(
 					{
 						//#hair_todo: Need to have a (frame) cache to insure that we don't recompute the same projection several time
 						// Actual populate the cache with only the needed part based on the groom projection data. At the moment it recompute everything ...
-						BuildCacheGeometry(GraphBuilder, ShaderMap, SceneProxy, bOutputTriangleData, Out);
+						BuildCacheGeometry(GraphBuilder, ShaderMap, SceneProxy, RootBulkData, bOutputTriangleData, Out);
 					}
 				}
 			}
@@ -294,10 +295,11 @@ static void RunHairStrandsInterpolation_Guide(
 		if (InstanceData.NeedsMeshUpdate(false/*bCheckMeshLODIndex*/))
 		{
 			// Extract MeshDataLOD and MeshLODIndex
-			const FCachedGeometry CachedGeometry = GetCacheGeometryForHair(GraphBuilder, Scene, Instance, ShaderMap, true /*bOutputTriangleData*/);
-			for (int32 SectionIndex = 0; SectionIndex < CachedGeometry.Sections.Num(); ++SectionIndex)
+			check(Instance->Guides.RestRootResource);
+			const FCachedGeometry CachedGeometry = GetCacheGeometryForHair(GraphBuilder, ShaderMap, Scene, Instance, &Instance->Guides.RestRootResource->BulkData, true /*bOutputTriangleData*/);
+			for (const FCachedGeometry::Section& CachedGeometrySection : CachedGeometry.Sections)
 			{
-				InstanceData.MeshDataLOD.Sections.Add(ConvertMeshSection(CachedGeometry, SectionIndex));
+				InstanceData.MeshDataLOD.Sections.Add(ConvertMeshSection(CachedGeometry, CachedGeometrySection));
 			}
 			InstanceData.MeshLODIndex    = CachedGeometry.LODIndex;
 			Instance->Debug.MeshLODIndex = CachedGeometry.LODIndex;
@@ -308,7 +310,7 @@ static void RunHairStrandsInterpolation_Guide(
 		else
 		{
 			// Extract MeshLODIndex
-			const FCachedGeometry CachedGeometry = GetCacheGeometryForHair(GraphBuilder, Scene, Instance, ShaderMap, false /*bOutputTriangleData*/);
+			const FCachedGeometry CachedGeometry = GetCacheGeometryForHair(GraphBuilder, ShaderMap, Scene, Instance, nullptr, false /*bOutputTriangleData*/);
 			InstanceData.MeshLODIndex    = CachedGeometry.LODIndex;
 			Instance->Debug.MeshLODIndex = CachedGeometry.LODIndex;
 		}
@@ -689,10 +691,10 @@ static void RunHairStrandsInterpolation_Strands(
 		if (InstanceData.ActiveGroomCacheType == EGroomCacheType::Strands) { bHasAnyRenCacheInstances = true; }
 
 		// Extract MeshDataLOD and compute MeshLODIndex
-		const FCachedGeometry CachedGeometry = GetCacheGeometryForHair(GraphBuilder, Scene, Instance, ShaderMap, true);
-		for (int32 SectionIndex = 0; SectionIndex < CachedGeometry.Sections.Num(); ++SectionIndex)
+		const FCachedGeometry CachedGeometry = GetCacheGeometryForHair(GraphBuilder, ShaderMap, Scene, Instance, &Instance->Strands.RestRootResource->BulkData, true);
+		for (const FCachedGeometry::Section& CachedGeometrySection : CachedGeometry.Sections)
 		{
-			InstanceData.MeshDataLOD.Sections.Add(ConvertMeshSection(CachedGeometry, SectionIndex));
+			InstanceData.MeshDataLOD.Sections.Add(ConvertMeshSection(CachedGeometry, CachedGeometrySection));
 		}
 		InstanceData.MeshLODIndex    = CachedGeometry.LODIndex;
 		Instance->Debug.MeshLODIndex = CachedGeometry.LODIndex;
@@ -1307,18 +1309,35 @@ static void RunHairStrandsInterpolation_Cards(
 			InstanceData.MeshLODIndex				= -1;
 
 			// Binding surface parameters (skel.mesh/geom. cache)
-			FHairStrandsProjectionMeshData::FLOD MeshDataLOD;
-			const FCachedGeometry CachedGeometry = GetCacheGeometryForHair(GraphBuilder, Scene, Instance, ShaderMap, true);
-			for (int32 SectionIndex = 0; SectionIndex < CachedGeometry.Sections.Num(); ++SectionIndex)
-			{	
-				InstanceData.MeshDataLOD.Sections.Add(ConvertMeshSection(CachedGeometry, SectionIndex));
+			// This is only needed for card geometry. Mesh geometry only uses RBF deformation, which are initalized by the guide pass.
+			if (Instance->GeometryType == EHairGeometryType::Cards)
+			{				
+				const bool bNeedSurfaceUpdate = InstanceData.BindingType == EHairBindingType::Skinning || InstanceData.bGlobalDeformationEnable || InstanceData.bDeformationEnable;
+				if (bNeedSurfaceUpdate)
+				{	
+					// Extract MeshDataLOD and MeshLODIndex
+					check(Instance->Cards.LODs.IsValidIndex(InstanceData.HairLODIndex));
+					const FCachedGeometry CachedGeometry = GetCacheGeometryForHair(GraphBuilder, ShaderMap, Scene, Instance, &Instance->Cards.LODs[InstanceData.HairLODIndex].Guides.RestRootResource->BulkData, true /*bOutputTriangleData*/);
+					for (const FCachedGeometry::Section& CachedGeometrySection : CachedGeometry.Sections)
+					{
+						InstanceData.MeshDataLOD.Sections.Add(ConvertMeshSection(CachedGeometry, CachedGeometrySection));
+					}
+					InstanceData.MeshLODIndex    = CachedGeometry.LODIndex;
+					Instance->Debug.MeshLODIndex = CachedGeometry.LODIndex;
+				}
+				else
+				{
+					// Extract MeshLODIndex
+					const FCachedGeometry CachedGeometry = GetCacheGeometryForHair(GraphBuilder, ShaderMap, Scene, Instance, nullptr, false /*bOutputTriangleData*/);
+					InstanceData.MeshLODIndex    = CachedGeometry.LODIndex;
+					Instance->Debug.MeshLODIndex = CachedGeometry.LODIndex;
+				}
 			}
-			InstanceData.MeshLODIndex    = CachedGeometry.LODIndex;
-			Instance->Debug.MeshLODIndex = CachedGeometry.LODIndex;
 
 			// Card or Mesh specific parameters
 			if (Instance->GeometryType == EHairGeometryType::Cards)
 			{
+				check(Instance->Cards.LODs.IsValidIndex(InstanceData.HairLODIndex));
 				InstanceData.CardInstance 			= &Instance->Cards.LODs[InstanceData.HairLODIndex];
 				InstanceData.bValidGuide			= InstanceData.Instance->Guides.bIsSimulationEnable || InstanceData.Instance->Guides.bHasGlobalInterpolation || InstanceData.Instance->Guides.bIsDeformationEnable || InstanceData.Instance->Guides.bIsSimulationCacheEnable;
 				InstanceData.bHasSkinning			= InstanceData.BindingType == EHairBindingType::Skinning && InstanceData.MeshLODIndex >= 0;
@@ -1354,6 +1373,7 @@ static void RunHairStrandsInterpolation_Cards(
 				check(Instance->Meshes.IsValid(InstanceData.HairLODIndex));
 				InstanceData.bNeedDeformation 	= Instance->Meshes.LODs[InstanceData.HairLODIndex].DeformedResource != nullptr;
 				InstanceData.MeshInstance 		= &Instance->Meshes.LODs[InstanceData.HairLODIndex];
+				InstanceData.MeshLODIndex 		= Instance->Debug.MeshLODIndex; // Debug.MeshLODIndex is initialized during the guide pass (if guides are needed updated)
 				if (InstanceData.bNeedDeformation)
 				{
 					check(Instance->BindingType == EHairBindingType::Skinning);
@@ -1362,7 +1382,7 @@ static void RunHairStrandsInterpolation_Cards(
 					check(Instance->Guides.DeformedRootResource);
 					// MeshLODIndex -1 indicates that skin cache is disabled and this is a workaround to prevent the editor from crashing 
 					// An editor setting guildeline popup exists to inform the user that the skin cache should be enabled. 
-					check(InstanceData.MeshLODIndex == -1 || Instance->Guides.DeformedRootResource->IsValid(InstanceData.MeshLODIndex)); 
+					check(InstanceData.MeshLODIndex == -1 || Instance->Guides.DeformedRootResource->IsValid(InstanceData.MeshLODIndex));
 				}
 			}
 
@@ -1573,17 +1593,19 @@ static void RunHairStrandsInterpolation_Cards(
 	// Meshes only - Deform final mesh geometry (using RBF)
 	for (uint32 InstanceIndex : MeshInstances)
 	{
-		FInstanceData& InstanceData = InstanceDatas[InstanceIndex];
-		AddHairMeshesRBFInterpolationPass(
-			GraphBuilder,
-			ShaderMap,
-			InstanceData.MeshLODIndex,
-			InstanceData.MeshInstance->RestResource,
-			InstanceData.MeshInstance->DeformedResource,
-			InstanceData.Instance->Guides.RestRootResource,
-			InstanceData.Instance->Guides.DeformedRootResource);
-
-		ExternalAccessQueue.Add(Register(GraphBuilder, InstanceData.MeshInstance->DeformedResource->GetBuffer(FHairMeshesDeformedResource::Current), ERDGImportedBufferFlags::None).Buffer, ERHIAccess::SRVMask);
+		{
+			FInstanceData& InstanceData = InstanceDatas[InstanceIndex];
+			AddHairMeshesRBFInterpolationPass(
+				GraphBuilder,
+				ShaderMap,
+				InstanceData.Instance->Debug.MeshLODIndex,
+				InstanceData.MeshInstance->RestResource,
+				InstanceData.MeshInstance->DeformedResource,
+				InstanceData.Instance->Guides.RestRootResource,
+				InstanceData.Instance->Guides.DeformedRootResource);
+	
+			ExternalAccessQueue.Add(Register(GraphBuilder, InstanceData.MeshInstance->DeformedResource->GetBuffer(FHairMeshesDeformedResource::Current), ERDGImportedBufferFlags::None).Buffer, ERHIAccess::SRVMask);
+		}
 	}
 
 	// Meshes only - Build/update RT geometry
@@ -1907,7 +1929,7 @@ static void RunHairLODSelection(
 
 		check(Instance);
 		check(Instance->HairGroupPublicData);
-		const FCachedGeometry CachedGeometry = GetCacheGeometryForHair(GraphBuilder, Scene, Instance, ShaderMap, false);
+		const FCachedGeometry CachedGeometry = GetCacheGeometryForHair(GraphBuilder, ShaderMap, Scene, Instance, nullptr, false);
 		const int32 MeshLODIndex = CachedGeometry.LODIndex;
 
 		// Perform LOD selection based on all the views	
diff --git a/Engine/Plugins/Runtime/HairStrands/Source/HairStrandsCore/Private/HairStrandsMeshProjection.cpp b/Engine/Plugins/Runtime/HairStrands/Source/HairStrandsCore/Private/HairStrandsMeshProjection.cpp
--- a/Engine/Plugins/Runtime/HairStrands/Source/HairStrandsCore/Private/HairStrandsMeshProjection.cpp
+++ b/Engine/Plugins/Runtime/HairStrands/Source/HairStrandsCore/Private/HairStrandsMeshProjection.cpp
@@ -121,15 +121,15 @@ void AddSkinUpdatePass(
 		ERDGPassFlags::Compute,
 		[ParametersMetadata, Parameters, ComputeShader, Sections, bPrevPosition](FRHIComputeCommandList& RHICmdList)
 		{
-			for (uint32 SectionIt = 0, SectionCount = Sections.Num(); SectionIt < SectionCount; ++SectionIt)
+			for (const FSkinUpdateSection& Section : Sections)
 			{
-				Parameters->NumVertexToProcess 		= Sections[SectionIt].NumVertexToProcess;
-				Parameters->SectionVertexBaseIndex 	= Sections[SectionIt].SectionVertexBaseIndex;
-				Parameters->BoneMatrices 			= Sections[SectionIt].BoneBuffer;
+				Parameters->NumVertexToProcess 		= Section.NumVertexToProcess;
+				Parameters->SectionVertexBaseIndex 	= Section.SectionVertexBaseIndex;
+				Parameters->BoneMatrices 			= Section.BoneBuffer;
 				check(Parameters->BoneMatrices);
 				if (bPrevPosition)
 				{
-					Parameters->PrevBoneMatrices = Sections[SectionIt].BonePrevBuffer;
+					Parameters->PrevBoneMatrices = Section.BonePrevBuffer;
 					check(Parameters->PrevBoneMatrices);
 				}
 
@@ -185,15 +185,16 @@ bool AddHairStrandUpdateMeshTrianglesPass(
 	FRDGBuilder& GraphBuilder,
 	FGlobalShaderMap* ShaderMap,
 	const uint32 UniqueTriangleCount, 
+	const uint32 TotalSectionCount, // Total section count of the underlying mesh
 	const FHairStrandsProjectionMeshData::FLOD& MeshData,
 	FRDGBufferSRVRef UniqueTriangleIndexSRV,
 	FRDGBufferUAVRef OutputCurrUAV,
 	FRDGBufferUAVRef OutputPrevUAV)
 {
-	const uint32 SectionCount = MeshData.Sections.Num();
+	const uint32 EffectiveSectionCount 	= MeshData.Sections.Num(); // Section count which are valid & used by groom data
 
 	// If a skel. mesh is not streaming yet, its SRV will be null
-	const bool bValid = SectionCount > 0 && SectionCount < GetHairStrandsMaxSectionCount();
+	const bool bValid = EffectiveSectionCount > 0 && TotalSectionCount < GetHairStrandsMaxSectionCount();
 	const bool bReady = MeshData.Sections[0].IndexBuffer != nullptr && MeshData.Sections[0].UVsBuffer != nullptr && (MeshData.Sections[0].RDGPositionBuffer != nullptr || MeshData.Sections[0].PositionBuffer != nullptr);
 	if (!bReady || !bValid)
 	{
@@ -202,7 +203,7 @@ bool AddHairStrandUpdateMeshTrianglesPass(
 
 	FHairUpdateMeshTriangleCS::FParameters CommonParameters;
 	CommonParameters.MaxUniqueTriangleCount 		= UniqueTriangleCount;
-	CommonParameters.MaxSectionCount 				= SectionCount;
+	CommonParameters.MaxSectionCount 				= TotalSectionCount;
 	CommonParameters.RDGMeshPositionBuffer			= MeshData.Sections[0].RDGPositionBuffer;
 	CommonParameters.RDGMeshPreviousPositionBuffer 	= MeshData.Sections[0].RDGPreviousPositionBuffer;
 	CommonParameters.MeshPositionBuffer				= MeshData.Sections[0].PositionBuffer;
@@ -225,29 +226,39 @@ bool AddHairStrandUpdateMeshTrianglesPass(
 		uint32 bIsSwapped : 8;
 		uint32 Pad : 8;
 	};
+
+	FSectionData Default;
+	Default.TotalIndexCount = 0;
+	Default.TotalVertexCount = 0;
+	Default.IndexBaseIndex = 0;
+	Default.UVsChannelOffset  = 0;
+	Default.UVsChannelCount  = 0;
+	Default.bIsSwapped = 0;
+	Default.Pad = 0;
+
+	// Allocate data for *all* sections, but only fill in the used/valid sections
 	TArray<FSectionData> SectionDatas;
-	SectionDatas.SetNum(SectionCount);
-	for (uint32 SectionIt = 0; SectionIt < SectionCount; ++SectionIt)
+	SectionDatas.Init(Default, TotalSectionCount); 
+	for (const FHairStrandsProjectionMeshData::FSection& MeshSectionData : MeshData.Sections)
 	{
-		const FHairStrandsProjectionMeshData::FSection& MeshSectionData = MeshData.Sections[SectionIt];
-		SectionDatas[SectionIt].TotalIndexCount	= MeshSectionData.TotalIndexCount;
-		SectionDatas[SectionIt].TotalVertexCount= MeshSectionData.TotalVertexCount;
-		SectionDatas[SectionIt].IndexBaseIndex	= MeshSectionData.IndexBaseIndex;
-		SectionDatas[SectionIt].UVsChannelOffset= MeshSectionData.UVsChannelOffset;
-		SectionDatas[SectionIt].UVsChannelCount	= MeshSectionData.UVsChannelCount;
-		SectionDatas[SectionIt].Pad				= 0u;
-		SectionDatas[SectionIt].bIsSwapped		= bUseRDGPositionBuffer ? 
-			(MeshData.Sections[SectionIt].RDGPositionBuffer != CommonParameters.RDGMeshPositionBuffer ? 1u : 0u) : 
-			(MeshData.Sections[SectionIt].PositionBuffer != CommonParameters.MeshPositionBuffer ? 1u : 0u);
+		const uint32 SectionIndex = MeshSectionData.SectionIndex;
+		SectionDatas[SectionIndex].TotalIndexCount	= MeshSectionData.TotalIndexCount;
+		SectionDatas[SectionIndex].TotalVertexCount	= MeshSectionData.TotalVertexCount;
+		SectionDatas[SectionIndex].IndexBaseIndex	= MeshSectionData.IndexBaseIndex;
+		SectionDatas[SectionIndex].UVsChannelOffset	= MeshSectionData.UVsChannelOffset;
+		SectionDatas[SectionIndex].UVsChannelCount	= MeshSectionData.UVsChannelCount;
+		SectionDatas[SectionIndex].Pad				= 0u;
+		SectionDatas[SectionIndex].bIsSwapped		= bUseRDGPositionBuffer ? 
+			(MeshSectionData.RDGPositionBuffer != CommonParameters.RDGMeshPositionBuffer ? 1u : 0u) :
+			(MeshSectionData.PositionBuffer    != CommonParameters.MeshPositionBuffer    ? 1u : 0u);
 
 		// Sanity check
-		check(MeshSectionData.SectionIndex == SectionIt);
 		check(MeshSectionData.UVsChannelOffset < 255);
 		check(MeshSectionData.UVsChannelCount < 255);
-		check(CommonParameters.RDGMeshPositionBuffer == MeshData.Sections[SectionIt].RDGPositionBuffer || CommonParameters.RDGMeshPositionBuffer == MeshData.Sections[SectionIt].RDGPreviousPositionBuffer);
-		check(CommonParameters.MeshPositionBuffer    == MeshData.Sections[SectionIt].PositionBuffer    || CommonParameters.MeshPositionBuffer    == MeshData.Sections[SectionIt].PreviousPositionBuffer);
-		check(CommonParameters.MeshIndexBuffer		 == MeshData.Sections[SectionIt].IndexBuffer);
-		check(CommonParameters.MeshUVsBuffer		 == MeshData.Sections[SectionIt].UVsBuffer);
+		check(CommonParameters.RDGMeshPositionBuffer == MeshSectionData.RDGPositionBuffer || CommonParameters.RDGMeshPositionBuffer == MeshSectionData.RDGPreviousPositionBuffer);
+		check(CommonParameters.MeshPositionBuffer    == MeshSectionData.PositionBuffer    || CommonParameters.MeshPositionBuffer    == MeshSectionData.PreviousPositionBuffer);
+		check(CommonParameters.MeshIndexBuffer		 == MeshSectionData.IndexBuffer);
+		check(CommonParameters.MeshUVsBuffer		 == MeshSectionData.UVsBuffer);
 	}
 
 	// If no previous position buffer available, reusing the current position buffers
@@ -303,6 +314,7 @@ void AddHairStrandUpdateMeshTrianglesPass(
 	check(RestLODData.LODIndex == MeshLODIndex);
 	check(RestResources->BulkData.Header.LODs.IsValidIndex(MeshLODIndex));
 	const uint32 UniqueTriangleCount = RestResources->BulkData.Header.LODs[MeshLODIndex].UniqueTriangleCount;
+	const uint32 TotalMeshSectionCount = RestResources->BulkData.Header.LODs[MeshLODIndex].MeshSectionCount;
 	if (UniqueTriangleCount == 0)
 	{
 		return;
@@ -317,13 +329,14 @@ void AddHairStrandUpdateMeshTrianglesPass(
 	}
 
 	if (AddHairStrandUpdateMeshTrianglesPass(
-		GraphBuilder,
-		ShaderMap,
-		UniqueTriangleCount,
-		MeshData,
-		RegisterAsSRV(GraphBuilder, RestLODData.UniqueTriangleIndexBuffer),
-		OutputCurrBuffer.UAV,
-		OutputPrevBuffer.UAV))
+			GraphBuilder,
+			ShaderMap,
+			UniqueTriangleCount,
+			TotalMeshSectionCount,
+			MeshData,
+			RegisterAsSRV(GraphBuilder, RestLODData.UniqueTriangleIndexBuffer),
+			OutputCurrBuffer.UAV,
+			OutputPrevBuffer.UAV))
 	{
 
 		GraphBuilder.SetBufferAccessFinal(OutputCurrBuffer.Buffer, ERHIAccess::SRVMask);
@@ -514,8 +527,9 @@ void AddHairStrandInitMeshSamplesPass(
 	FHairStrandsRestRootResource::FLOD& RestLODData = RestResources->LODs[LODIndex];
 	check(RestLODData.LODIndex == LODIndex);
 
-	const uint32 SectionCount = MeshData.Sections.Num();
-	if (SectionCount == 0 || RestLODData.SampleCount == 0)
+	const uint32 TotalSectionCount = RestResources->BulkData.Header.LODs[LODIndex].MeshSectionCount;
+	const uint32 EffectiveSectionCount = MeshData.Sections.Num();
+	if (EffectiveSectionCount == 0 || RestLODData.SampleCount == 0)
 	{
 		return;
 	}
@@ -524,7 +538,7 @@ void AddHairStrandInitMeshSamplesPass(
 	
 	FHairInitMeshSamplesCS::FParameters CommonParameters;
 	CommonParameters.MaxSampleCount 				= RestLODData.SampleCount;
-	CommonParameters.MaxSectionCount 				= MeshData.Sections.Num();
+	CommonParameters.MaxSectionCount 				= TotalSectionCount;
 	CommonParameters.MaxVertexCount 				= MeshData.Sections[0].TotalVertexCount;
 	CommonParameters.RDGMeshPositionBuffer0			= MeshData.Sections[0].RDGPositionBuffer;
 	CommonParameters.RDGMeshPositionBuffer1		 	= MeshData.Sections[0].RDGPreviousPositionBuffer;
@@ -538,21 +552,21 @@ void AddHairStrandInitMeshSamplesPass(
 	// Create buffer with section data for look up in shader
 	struct FSectionData
 	{
-		uint32 bIsSwapped;
+		uint32 bIsSwapped = 0;
 	};
 	TArray<FSectionData> SectionDatas;
-	SectionDatas.SetNum(SectionCount);
-	for (uint32 SectionIt = 0; SectionIt < SectionCount; ++SectionIt)
+	SectionDatas.Init(FSectionData(), TotalSectionCount);
+	for (const FHairStrandsProjectionMeshData::FSection& Section : MeshData.Sections)
 	{
-		SectionDatas[SectionIt].bIsSwapped = 
+		const uint32 SectionIndex = Section.SectionIndex;
+		SectionDatas[SectionIndex].bIsSwapped = 
 			bUseRDGPositionBuffer ?
-			(MeshData.Sections[SectionIt].RDGPositionBuffer != CommonParameters.RDGMeshPositionBuffer0 ? 1u : 0u):
-			(MeshData.Sections[SectionIt].PositionBuffer    != CommonParameters.MeshPositionBuffer0    ? 1u : 0u);
+			(Section.RDGPositionBuffer != CommonParameters.RDGMeshPositionBuffer0 ? 1u : 0u) :
+			(Section.PositionBuffer    != CommonParameters.MeshPositionBuffer0    ? 1u : 0u);
 
 		// Sanity check
-		check(MeshData.Sections[SectionIt].SectionIndex == SectionIt);
-		check(CommonParameters.RDGMeshPositionBuffer0 == MeshData.Sections[SectionIt].RDGPositionBuffer || CommonParameters.RDGMeshPositionBuffer0 == MeshData.Sections[SectionIt].RDGPreviousPositionBuffer);
-		check(CommonParameters.MeshPositionBuffer0    == MeshData.Sections[SectionIt].PositionBuffer    || CommonParameters.MeshPositionBuffer0    == MeshData.Sections[SectionIt].PreviousPositionBuffer);
+		check(CommonParameters.RDGMeshPositionBuffer0 == Section.RDGPositionBuffer || CommonParameters.RDGMeshPositionBuffer0 == Section.RDGPreviousPositionBuffer);
+		check(CommonParameters.MeshPositionBuffer0    == Section.PositionBuffer    || CommonParameters.MeshPositionBuffer0    == Section.PreviousPositionBuffer);
 	}
 	CommonParameters.MeshSectionBuffer = GraphBuilder.CreateSRV(CreateStructuredBuffer(GraphBuilder, TEXT("Hair.SkelMeshSectionBuffer"), sizeof(FSectionData),  SectionDatas.Num(), SectionDatas.GetData(), sizeof(FSectionData) * SectionDatas.Num()));
 
