Fix crash with hair strands not bound to skin geometry.

diff --git a/Engine/Plugins/Runtime/HairStrands/Source/HairStrandsCore/Private/GroomManager.cpp b/Engine/Plugins/Runtime/HairStrands/Source/HairStrandsCore/Private/GroomManager.cpp
--- a/Engine/Plugins/Runtime/HairStrands/Source/HairStrandsCore/Private/GroomManager.cpp
+++ b/Engine/Plugins/Runtime/HairStrands/Source/HairStrandsCore/Private/GroomManager.cpp
@@ -691,7 +691,8 @@ static void RunHairStrandsInterpolation_Strands(
 		if (InstanceData.ActiveGroomCacheType == EGroomCacheType::Strands) { bHasAnyRenCacheInstances = true; }
 
 		// Extract MeshDataLOD and compute MeshLODIndex
-		const FCachedGeometry CachedGeometry = GetCacheGeometryForHair(GraphBuilder, ShaderMap, Scene, Instance, &Instance->Strands.RestRootResource->BulkData, true);
+		const bool bNeedOutputTriangleData = Instance->Strands.RestRootResource != nullptr;
+		const FCachedGeometry CachedGeometry = GetCacheGeometryForHair(GraphBuilder, ShaderMap, Scene, Instance, bNeedOutputTriangleData ? &Instance->Strands.RestRootResource->BulkData : nullptr, bNeedOutputTriangleData);
 		for (const FCachedGeometry::Section& CachedGeometrySection : CachedGeometry.Sections)
 		{
 			InstanceData.MeshDataLOD.Sections.Add(ConvertMeshSection(CachedGeometry, CachedGeometrySection));
