Add debug configuration for visualizing Sim/Ren root deformation & RBF samples.

diff --git a/Engine/Plugins/Runtime/HairStrands/Source/HairStrandsCore/Private/GroomDebug.cpp b/Engine/Plugins/Runtime/HairStrands/Source/HairStrandsCore/Private/GroomDebug.cpp
--- a/Engine/Plugins/Runtime/HairStrands/Source/HairStrandsCore/Private/GroomDebug.cpp
+++ b/Engine/Plugins/Runtime/HairStrands/Source/HairStrandsCore/Private/GroomDebug.cpp
@@ -27,6 +27,7 @@
 static int32 GHairDebugMeshProjection_SkinCacheMesh = 0;
 static int32 GHairDebugMeshProjection_SkinCacheMeshInUVsSpace = 0;
 
+static int32 GHairDebugMeshProjection_Override = 0;
 static int32 GHairDebugMeshProjection_Sim_HairRestTriangles = 0;
 static int32 GHairDebugMeshProjection_Sim_HairRestFrames = 0;
 static int32 GHairDebugMeshProjection_Sim_HairRestSamples = 0;
@@ -44,12 +45,11 @@ static int32 GHairDebugMeshProjection_Render_HairDeformedSamples = 0;
 
 static FAutoConsoleVariableRef CVarHairDebugMeshProjection_SkinCacheMeshInUVsSpace(TEXT("r.HairStrands.MeshProjection.DebugInUVsSpace"), GHairDebugMeshProjection_SkinCacheMeshInUVsSpace, TEXT("Render debug mes projection in UVs space"));
 
+static FAutoConsoleVariableRef CVarHairDebugMeshProjection_Override(TEXT("r.HairStrands.MeshProjection"), GHairDebugMeshProjection_Override, TEXT("Override in shader settings for displaying root debug data"));
 static FAutoConsoleVariableRef CVarHairDebugMeshProjection_Render_HairRestTriangles(TEXT("r.HairStrands.MeshProjection.Render.Rest.Triangles"), GHairDebugMeshProjection_Render_HairRestTriangles, TEXT("Render strands rest triangles"));
 static FAutoConsoleVariableRef CVarHairDebugMeshProjection_Render_HairRestFrames(TEXT("r.HairStrands.MeshProjection.Render.Rest.Frames"), GHairDebugMeshProjection_Render_HairRestFrames, TEXT("Render strands rest frames"));
-static FAutoConsoleVariableRef CVarHairDebugMeshProjection_Render_HairRestSamples(TEXT("r.HairStrands.MeshProjection.Render.Rest.Samples"), GHairDebugMeshProjection_Render_HairRestSamples, TEXT("Render strands rest samples"));
 static FAutoConsoleVariableRef CVarHairDebugMeshProjection_Render_HairDeformedTriangles(TEXT("r.HairStrands.MeshProjection.Render.Deformed.Triangles"), GHairDebugMeshProjection_Render_HairDeformedTriangles, TEXT("Render strands deformed triangles"));
 static FAutoConsoleVariableRef CVarHairDebugMeshProjection_Render_HairDeformedFrames(TEXT("r.HairStrands.MeshProjection.Render.Deformed.Frames"), GHairDebugMeshProjection_Render_HairDeformedFrames, TEXT("Render strands deformed frames"));
-static FAutoConsoleVariableRef CVarHairDebugMeshProjection_Render_HairDeformedSamples(TEXT("r.HairStrands.MeshProjection.Render.Deformed.Samples"), GHairDebugMeshProjection_Render_HairDeformedSamples, TEXT("Render strands deformed samples"));
 
 static FAutoConsoleVariableRef CVarHairDebugMeshProjection_Sim_HairRestTriangles(TEXT("r.HairStrands.MeshProjection.Sim.Rest.Triangles"), GHairDebugMeshProjection_Sim_HairRestTriangles, TEXT("Render guides rest triangles"));
 static FAutoConsoleVariableRef CVarHairDebugMeshProjection_Sim_HairRestFrames(TEXT("r.HairStrands.MeshProjection.Sim.Rest.Frames"), GHairDebugMeshProjection_Sim_HairRestFrames, TEXT("Render guides rest frames"));
@@ -277,6 +277,21 @@ class FHairProjectionHairDebugCS : public FGlobalShader
 		SHADER_PARAMETER(uint32, MaxRootCount)
 		SHADER_PARAMETER(uint32, DeformedFrameEnable)
 		SHADER_PARAMETER(FMatrix44f, RootLocalToWorld)
+		SHADER_PARAMETER(uint32, DrawIndex)
+		SHADER_PARAMETER(uint32, bSimPass)
+		SHADER_PARAMETER(uint32, bRestPass)
+
+		SHADER_PARAMETER(uint32, bOverride)
+		SHADER_PARAMETER(uint32, Render_HairRestTriangles)
+		SHADER_PARAMETER(uint32, Render_HairRestFrames)
+		SHADER_PARAMETER(uint32, Render_HairDeformedTriangles)
+		SHADER_PARAMETER(uint32, Render_HairDeformedFrames)
+		SHADER_PARAMETER(uint32, Sim_HairRestTriangles)
+		SHADER_PARAMETER(uint32, Sim_HairRestFrames)
+		SHADER_PARAMETER(uint32, Sim_HairRestSamples)
+		SHADER_PARAMETER(uint32, Sim_HairDeformedTriangles)
+		SHADER_PARAMETER(uint32, Sim_HairDeformedFrames)
+		SHADER_PARAMETER(uint32, Sim_HairDeformedSamples)
 
 		SHADER_PARAMETER_RDG_BUFFER_SRV(Buffer, RestPositionBuffer)
 		SHADER_PARAMETER_RDG_BUFFER_SRV(Buffer, DeformedPositionBuffer)
@@ -323,6 +338,8 @@ static void AddDebugProjectionHairPass(
 	const FHairStrandsRestRootResource* RestRootResources,
 	const FHairStrandsDeformedRootResource* DeformedRootResources,
 	const FTransform& LocalToWorld,
+	const uint32 DrawIndex,
+	const bool bSim,
 	const EDebugProjectionHairType GeometryType,
 	const HairStrandsTriangleType PoseType)
 {
@@ -385,10 +402,24 @@ static void AddDebugProjectionHairPass(
 	const FIntPoint Resolution(Viewport.Width(), Viewport.Height());
 
 	FHairProjectionHairDebugCS::FParameters* Parameters = GraphBuilder.AllocParameters<FHairProjectionHairDebugCS::FParameters>();
-	Parameters->OutputResolution = Resolution;
-	Parameters->MaxRootCount = RootCount;
-	Parameters->RootLocalToWorld = FMatrix44f(LocalToWorld.ToMatrixWithScale());	// LWC_TODO: Precision loss
-	Parameters->DeformedFrameEnable = PoseType == HairStrandsTriangleType::DeformedPose;
+	Parameters->OutputResolution 			= Resolution;
+	Parameters->MaxRootCount 				= RootCount;
+	Parameters->RootLocalToWorld 			= FMatrix44f(LocalToWorld.ToMatrixWithScale());	// LWC_TODO: Precision loss
+	Parameters->DeformedFrameEnable 		= PoseType == HairStrandsTriangleType::DeformedPose;
+	Parameters->DrawIndex					= DrawIndex;
+	Parameters->bSimPass 					= bSim ? 1u : 0u;
+	Parameters->bRestPass 					= PoseType == HairStrandsTriangleType::RestPose ? 1u : 0u;
+	Parameters->bOverride 					= GHairDebugMeshProjection_Override > 0 ? 1u : 0u;
+	Parameters->Render_HairRestTriangles	= GHairDebugMeshProjection_Render_HairRestTriangles > 0 ? 1u : 0u;
+	Parameters->Render_HairRestFrames 		= GHairDebugMeshProjection_Render_HairRestFrames > 0 ? 1u : 0u;
+	Parameters->Render_HairDeformedTriangles= GHairDebugMeshProjection_Render_HairDeformedTriangles > 0 ? 1u : 0u;
+	Parameters->Render_HairDeformedFrames	= GHairDebugMeshProjection_Render_HairDeformedFrames > 0 ? 1u : 0u;
+	Parameters->Sim_HairRestTriangles		= GHairDebugMeshProjection_Sim_HairRestTriangles > 0 ? 1u : 0u;
+	Parameters->Sim_HairRestFrames			= GHairDebugMeshProjection_Sim_HairRestFrames > 0 ? 1u : 0u;
+	Parameters->Sim_HairRestSamples			= GHairDebugMeshProjection_Sim_HairRestSamples > 0 ? 1u : 0u;
+	Parameters->Sim_HairDeformedTriangles	= GHairDebugMeshProjection_Sim_HairDeformedTriangles > 0 ? 1u : 0u;
+	Parameters->Sim_HairDeformedFrames		= GHairDebugMeshProjection_Sim_HairDeformedFrames > 0 ? 1u : 0u;
+	Parameters->Sim_HairDeformedSamples		= GHairDebugMeshProjection_Sim_HairDeformedSamples > 0 ? 1u : 0u;
 
 	if (EDebugProjectionHairType::HairFrame == GeometryType)
 	{
@@ -1440,7 +1471,8 @@ void RunHairStrandsDebug(
 				RenderMeshProjection(GraphBuilder, EHairStrandsProjectionMeshType::TargetMesh);
 			}
 
-			auto RenderProjectionData = [&GraphBuilder, LocalView, ShaderMap, Viewport, &ViewUniformBuffer, Instances, ShaderPrintData](bool bGuide, bool bRestTriangle, bool bRestFrame, bool bRestSamples, bool bDeformedTriangle, bool bDeformedFrame, bool bDeformedSamples)
+			uint32 DrawIndex = 0;
+			auto RenderProjectionData = [&GraphBuilder, LocalView, ShaderMap, Viewport, &ViewUniformBuffer, Instances, ShaderPrintData, &DrawIndex](bool bSim)
 			{
 				TArray<int32> HairLODIndices;
 				for (FHairStrandsInstance* AbstractInstance : Instances)
@@ -1451,7 +1483,7 @@ void RunHairStrandsDebug(
 
 					FHairStrandsRestRootResource* RestRootResource = nullptr;
 					FHairStrandsDeformedRootResource* DeformedRootResource = nullptr;
-					if (bGuide)
+					if (bSim)
 					{
 						RestRootResource 	 = Instance->Guides.RestRootResource;
 						DeformedRootResource = Instance->Guides.DeformedRootResource;
@@ -1475,48 +1507,17 @@ void RunHairStrandsDebug(
 
 					const int32 MeshLODIndex = Instance->Debug.MeshLODIndex;
 
-					if (bRestTriangle)		{ AddDebugProjectionHairPass(GraphBuilder, LocalView, ShaderMap, ShaderPrintData, Viewport, ViewUniformBuffer, MeshLODIndex, RestRootResource, DeformedRootResource, Instance->HairGroupPublicData->VFInput.LocalToWorldTransform, EDebugProjectionHairType::HairTriangle, HairStrandsTriangleType::RestPose); }
-					if (bRestFrame)			{ AddDebugProjectionHairPass(GraphBuilder, LocalView, ShaderMap, ShaderPrintData, Viewport, ViewUniformBuffer, MeshLODIndex, RestRootResource, DeformedRootResource, Instance->HairGroupPublicData->VFInput.LocalToWorldTransform, EDebugProjectionHairType::HairFrame,    HairStrandsTriangleType::RestPose); }
-					if (bRestSamples)		{ AddDebugProjectionHairPass(GraphBuilder, LocalView, ShaderMap, ShaderPrintData, Viewport, ViewUniformBuffer, MeshLODIndex, RestRootResource, DeformedRootResource, Instance->HairGroupPublicData->VFInput.LocalToWorldTransform, EDebugProjectionHairType::HairSamples,  HairStrandsTriangleType::RestPose); }
-					if (bDeformedTriangle)	{ AddDebugProjectionHairPass(GraphBuilder, LocalView, ShaderMap, ShaderPrintData, Viewport, ViewUniformBuffer, MeshLODIndex, RestRootResource, DeformedRootResource, Instance->HairGroupPublicData->VFInput.LocalToWorldTransform, EDebugProjectionHairType::HairTriangle, HairStrandsTriangleType::DeformedPose); }
-					if (bDeformedFrame)		{ AddDebugProjectionHairPass(GraphBuilder, LocalView, ShaderMap, ShaderPrintData, Viewport, ViewUniformBuffer, MeshLODIndex, RestRootResource, DeformedRootResource, Instance->HairGroupPublicData->VFInput.LocalToWorldTransform, EDebugProjectionHairType::HairFrame,    HairStrandsTriangleType::DeformedPose); }
-					if (bDeformedSamples)	{ AddDebugProjectionHairPass(GraphBuilder, LocalView, ShaderMap, ShaderPrintData, Viewport, ViewUniformBuffer, MeshLODIndex, RestRootResource, DeformedRootResource, Instance->HairGroupPublicData->VFInput.LocalToWorldTransform, EDebugProjectionHairType::HairSamples,  HairStrandsTriangleType::DeformedPose); }
+					AddDebugProjectionHairPass(GraphBuilder, LocalView, ShaderMap, ShaderPrintData, Viewport, ViewUniformBuffer, MeshLODIndex, RestRootResource, DeformedRootResource, Instance->HairGroupPublicData->VFInput.LocalToWorldTransform, DrawIndex++, bSim, EDebugProjectionHairType::HairTriangle, HairStrandsTriangleType::RestPose);
+					AddDebugProjectionHairPass(GraphBuilder, LocalView, ShaderMap, ShaderPrintData, Viewport, ViewUniformBuffer, MeshLODIndex, RestRootResource, DeformedRootResource, Instance->HairGroupPublicData->VFInput.LocalToWorldTransform, DrawIndex++, bSim, EDebugProjectionHairType::HairFrame,    HairStrandsTriangleType::RestPose);
+					AddDebugProjectionHairPass(GraphBuilder, LocalView, ShaderMap, ShaderPrintData, Viewport, ViewUniformBuffer, MeshLODIndex, RestRootResource, DeformedRootResource, Instance->HairGroupPublicData->VFInput.LocalToWorldTransform, DrawIndex++, bSim, EDebugProjectionHairType::HairSamples,  HairStrandsTriangleType::RestPose);
+					AddDebugProjectionHairPass(GraphBuilder, LocalView, ShaderMap, ShaderPrintData, Viewport, ViewUniformBuffer, MeshLODIndex, RestRootResource, DeformedRootResource, Instance->HairGroupPublicData->VFInput.LocalToWorldTransform, DrawIndex++, bSim, EDebugProjectionHairType::HairTriangle, HairStrandsTriangleType::DeformedPose);
+					AddDebugProjectionHairPass(GraphBuilder, LocalView, ShaderMap, ShaderPrintData, Viewport, ViewUniformBuffer, MeshLODIndex, RestRootResource, DeformedRootResource, Instance->HairGroupPublicData->VFInput.LocalToWorldTransform, DrawIndex++, bSim, EDebugProjectionHairType::HairFrame,    HairStrandsTriangleType::DeformedPose);
+					AddDebugProjectionHairPass(GraphBuilder, LocalView, ShaderMap, ShaderPrintData, Viewport, ViewUniformBuffer, MeshLODIndex, RestRootResource, DeformedRootResource, Instance->HairGroupPublicData->VFInput.LocalToWorldTransform, DrawIndex++, bSim, EDebugProjectionHairType::HairSamples,  HairStrandsTriangleType::DeformedPose);
 				}
 			};
 
-			if (GHairDebugMeshProjection_Render_HairRestTriangles > 0 ||
-				GHairDebugMeshProjection_Render_HairRestFrames > 0 ||
-				GHairDebugMeshProjection_Render_HairDeformedTriangles > 0 ||
-				GHairDebugMeshProjection_Render_HairDeformedFrames > 0 ||
-				GHairDebugMeshProjection_Render_HairDeformedSamples > 0 ||
-				GHairDebugMeshProjection_Render_HairRestSamples > 0)
-			{
-				RenderProjectionData(
-					false,
-					GHairDebugMeshProjection_Render_HairRestTriangles > 0,
-					GHairDebugMeshProjection_Render_HairRestFrames > 0,
-					GHairDebugMeshProjection_Render_HairRestSamples > 0,
-					GHairDebugMeshProjection_Render_HairDeformedTriangles > 0,
-					GHairDebugMeshProjection_Render_HairDeformedFrames > 0,
-					GHairDebugMeshProjection_Render_HairDeformedSamples > 0);
-			}
-
-			if (GHairDebugMeshProjection_Sim_HairRestTriangles > 0 ||
-				GHairDebugMeshProjection_Sim_HairRestFrames > 0 ||
-				GHairDebugMeshProjection_Sim_HairDeformedTriangles > 0 ||
-				GHairDebugMeshProjection_Sim_HairDeformedFrames > 0 ||
-				GHairDebugMeshProjection_Sim_HairDeformedSamples > 0 ||
-				GHairDebugMeshProjection_Sim_HairRestSamples > 0)
-			{
-				RenderProjectionData(
-					true,
-					GHairDebugMeshProjection_Sim_HairRestTriangles > 0,
-					GHairDebugMeshProjection_Sim_HairRestFrames > 0,
-					GHairDebugMeshProjection_Sim_HairRestSamples > 0,
-					GHairDebugMeshProjection_Sim_HairDeformedTriangles > 0,
-					GHairDebugMeshProjection_Sim_HairDeformedFrames > 0,
-					GHairDebugMeshProjection_Sim_HairDeformedSamples > 0);
-			}
+			RenderProjectionData(false /*bSim*/);
+			RenderProjectionData(true  /*bSim*/);
 		}
 	}
 
diff --git a/Engine/Shaders/Private/HairStrands/HairStrandsDebug.usf b/Engine/Shaders/Private/HairStrands/HairStrandsDebug.usf
--- a/Engine/Shaders/Private/HairStrands/HairStrandsDebug.usf
+++ b/Engine/Shaders/Private/HairStrands/HairStrandsDebug.usf
@@ -316,12 +316,27 @@ void MainCS(uint3 DispatchThreadId : SV_DispatchThreadID)
 #define INPUT_ROOT_TRIANGLE 1
 #define INPUT_ROOT_SAMPLES 2
 
+uint bSimPass;
+uint bRestPass;
+uint DrawIndex;
 uint MaxUniqueSection;
 uint MaxUniqueTriangleCount;
 uint MaxRootCount;
 float4x4 RootLocalToWorld;
 uint DeformedFrameEnable;
 
+uint bOverride;
+uint Render_HairRestTriangles;
+uint Render_HairRestFrames;
+uint Render_HairDeformedTriangles;
+uint Render_HairDeformedFrames;
+uint Sim_HairRestTriangles;
+uint Sim_HairRestFrames;
+uint Sim_HairRestSamples;
+uint Sim_HairDeformedTriangles;
+uint Sim_HairDeformedFrames;
+uint Sim_HairDeformedSamples;
+
 #if PERMUTATION_INPUT_TYPE == INPUT_ROOT_FRAME
 Buffer<uint>   RootBarycentricBuffer;
 #endif
@@ -334,11 +349,160 @@ Buffer<float4>  DeformedPositionBuffer;
 StructuredBuffer<float4> RestSamplePositionsBuffer;
 StructuredBuffer<float4> DeformedSamplePositionsBuffer;
 
+
+void Move(inout FShaderPrintContext Ctx, uint2 InPos)
+{
+	if (Ctx.bIsActive)
+	{
+		float2 fPos = float2(InPos) / float2(Ctx.Config.Resolution);
+		Ctx.StartPos = fPos;
+		Ctx.Pos 	 = fPos;
+	}
+}
+
 [numthreads(256, 1, 1)]
 void MainCS(uint3 DispatchThreadId : SV_DispatchThreadID)
 {
 	ResolvedView = ResolveView();
 
+	// Widgets for draw configuration
+	{	
+		uint2 StartPos  = uint2(500, 50);
+		uint2 OffsetPos = uint2(180, 0);
+		FShaderPrintContext Ctx = InitShaderPrintContext(DispatchThreadId.x == 0 && DrawIndex == 0, StartPos);
+		bool bValid = false;
+
+		// Guides
+		{
+			bool bRest    = false;
+			bool bDeform  = false;
+			bool bTriangle= false;
+			bool bRoot    = false; 
+
+			Print(Ctx, TEXT("Guides "), FontYellow); Newline(Ctx);
+			if (bOverride>0)
+			{
+				bRest    = Sim_HairRestTriangles     || Sim_HairRestFrames;
+				bDeform  = Sim_HairDeformedTriangles || Sim_HairDeformedFrames;
+				bTriangle= Sim_HairRestTriangles     || Sim_HairDeformedTriangles;
+				bRoot    = Sim_HairRestFrames        || Sim_HairDeformedFrames;
+
+				Print(Ctx, TEXT("Sim Rest      "), FontWhite);	PrintBool(Ctx, bRest); 		Newline(Ctx);
+				Print(Ctx, TEXT("Sim Deformed  "), FontWhite);	PrintBool(Ctx, bDeform);	Newline(Ctx);
+				Newline(Ctx);
+	
+				Print(Ctx, TEXT("Sim Triangles "), FontWhite);	PrintBool(Ctx, bTriangle); 	Newline(Ctx);
+				Print(Ctx, TEXT("Sim Roots     "), FontWhite); 	PrintBool(Ctx, bRoot);		Newline(Ctx);
+				Newline(Ctx);
+	
+			}
+			else
+			{
+				bRest     = AddCheckbox(Ctx, TEXT("Sim Rest"), false, FontWhite); Newline(Ctx);
+				bDeform   = AddCheckbox(Ctx, TEXT("Sim Deformed"), false, FontWhite); Newline(Ctx);
+				Newline(Ctx);
+		
+				bTriangle = AddCheckbox(Ctx, TEXT("Sim Triangles"), false, FontWhite); Newline(Ctx);
+				bRoot     = AddCheckbox(Ctx, TEXT("Sim Roots"), false, FontWhite); Newline(Ctx);
+				Newline(Ctx);
+			}
+
+			if (bSimPass == 1)
+			{
+				if (PERMUTATION_INPUT_TYPE == INPUT_ROOT_TRIANGLE)
+				{
+					bValid = bTriangle  && ((bRestPass==1 && bRest) || (bRestPass==0 && bDeform));
+				}
+				if (PERMUTATION_INPUT_TYPE == INPUT_ROOT_FRAME)
+				{
+					bValid = bRoot && ((bRestPass==1 && bRest) || (bRestPass==0 && bDeform));
+				}
+			}
+		}
+
+		// Strands/Cards
+		{
+			bool bRest    = false;
+			bool bDeform  = false;
+			bool bTriangle= false;
+			bool bRoot    = false;
+
+			Move(Ctx, StartPos + OffsetPos);
+			Print(Ctx, TEXT("Strands / Cards"), FontYellow); Newline(Ctx);
+			if (bOverride>0)
+			{
+				bRest    = Render_HairRestTriangles     || Render_HairRestFrames;
+				bDeform  = Render_HairDeformedTriangles || Render_HairDeformedFrames;
+				bTriangle= Render_HairRestTriangles     || Render_HairDeformedTriangles;
+				bRoot    = Render_HairRestFrames        || Render_HairDeformedFrames;
+	
+				Print(Ctx, TEXT("Ren Rest      "), FontWhite);	PrintBool(Ctx, bRest); 		Newline(Ctx);
+				Print(Ctx, TEXT("Ren Deformed  "), FontWhite);	PrintBool(Ctx, bDeform);	Newline(Ctx);
+				Newline(Ctx);
+			
+				Print(Ctx, TEXT("Ren Triangles "), FontWhite);	PrintBool(Ctx, bTriangle); 	Newline(Ctx);
+				Print(Ctx, TEXT("Ren Roots     "), FontWhite); 	PrintBool(Ctx, bRoot);		Newline(Ctx);
+				Newline(Ctx);
+			}
+			else
+			{
+				bRest     = AddCheckbox(Ctx, TEXT("Ren Rest"), false, FontWhite); Newline(Ctx);
+				bDeform   = AddCheckbox(Ctx, TEXT("Ren Deformed"), true, FontWhite) ; Newline(Ctx);
+				Newline(Ctx);
+				
+				bTriangle = AddCheckbox(Ctx, TEXT("Ren Triangles"), false, FontWhite); Newline(Ctx);
+				bRoot     = AddCheckbox(Ctx, TEXT("Ren Roots"), true, FontWhite); Newline(Ctx);
+				Newline(Ctx);
+			}
+
+			if (bSimPass == 0)
+			{
+				if (PERMUTATION_INPUT_TYPE == INPUT_ROOT_TRIANGLE)
+				{
+					bValid = bTriangle && ((bRestPass==1 && bRest) || (bRestPass==0 && bDeform));
+				}
+				if (PERMUTATION_INPUT_TYPE == INPUT_ROOT_FRAME)
+				{
+					bValid = bRoot && ((bRestPass==1 && bRest) || (bRestPass==0 && bDeform));
+				}
+			}
+		}
+
+		// RBF
+		{
+			bool bRest   = false;
+			bool bDeform = false;
+
+			Move(Ctx, StartPos + OffsetPos*2);
+			Print(Ctx, TEXT("RBF Samples"), FontYellow); Newline(Ctx);
+			if (bOverride>0)
+			{
+				bRest	= Sim_HairRestSamples;
+				bDeform = Sim_HairDeformedSamples;
+
+				Print(Ctx, TEXT("RBF Rest     "), FontWhite); Print(Ctx, TEXT("|| "), FontPurple); PrintBool(Ctx, bRest) Newline(Ctx);
+				Print(Ctx, TEXT("RBF Deformed "), FontWhite); Print(Ctx, TEXT("|| "), FontYellow); PrintBool(Ctx, bDeform) Newline(Ctx);
+				Newline(Ctx);
+			}
+			else
+			{
+				bRest     = AddCheckbox(Ctx, TEXT("RBF Rest    "), false, FontWhite); Print(Ctx, TEXT("||"), FontPurple); Newline(Ctx);
+				bDeform   = AddCheckbox(Ctx, TEXT("RBF Deformed"), false, FontWhite); Print(Ctx, TEXT("||"), FontYellow); Newline(Ctx);
+				Newline(Ctx);	
+			}
+
+			if (bSimPass == 1 && (PERMUTATION_INPUT_TYPE == INPUT_ROOT_SAMPLES))
+			{
+				bValid = (bRestPass==1 && bRest) || (bRestPass==0 && bDeform);
+			}
+		}
+
+		if (!bValid)
+		{
+			return;
+		}
+	}
+
 #if PERMUTATION_INPUT_TYPE == INPUT_ROOT_FRAME  
 	const uint RootIndex = DispatchThreadId.x;
 	if (RootIndex >= MaxRootCount) return;
@@ -439,8 +603,9 @@ void MainCS(uint3 DispatchThreadId : SV_DispatchThreadID)
 	{
 		float3 P = (DeformedFrameEnable > 0) ? DeformedSamplePositionsBuffer[PointIndex].xyz : RestSamplePositionsBuffer[PointIndex].xyz;
 		P = mul(float4(P, 1), RootLocalToWorld).xyz;
-
-		AddCrossWS(P, 1.f, ColorYellow);
+		
+		const float4 RBFColor = bRestPass==1 ? ColorPurple : ColorYellow;
+		AddCrossWS(P, 1.f, RBFColor);
 	}
 #endif // INPUT_ROOT_SAMPLES
 }
