* Regroup skel. mesh update before simulation for all type of geometry (guides/strands/cards/meshes).
* Deduplicate skel. mesh update when manual hair skin cache is enabled.

This allows to save duplicated cost between instance and between guides/strands.

diff --git a/Engine/Plugins/Runtime/HairStrands/Source/HairStrandsCore/Private/GroomDebug.cpp b/Engine/Plugins/Runtime/HairStrands/Source/HairStrandsCore/Private/GroomDebug.cpp
--- a/Engine/Plugins/Runtime/HairStrands/Source/HairStrandsCore/Private/GroomDebug.cpp
+++ b/Engine/Plugins/Runtime/HairStrands/Source/HairStrandsCore/Private/GroomDebug.cpp
@@ -24,7 +24,6 @@
 
 ///////////////////////////////////////////////////////////////////////////////////////////////////
 
-static int32 GHairDebugMeshProjection_SkinCacheMesh = 0;
 static int32 GHairDebugMeshProjection_SkinCacheMeshInUVsSpace = 0;
 
 static int32 GHairDebugMeshProjection_Override = 0;
@@ -69,199 +68,6 @@ enum class HairStrandsTriangleType
 	DeformedPose,
 };
 
-FCachedGeometry GetCacheGeometryForHair(
-	FRDGBuilder& GraphBuilder,
-	FGlobalShaderMap* ShaderMap,
-	FSceneInterface* Scene,
-	FHairGroupInstance* Instance,
-	const FHairStrandsRootBulkData* RootBulkData,
-	const bool bOutputTriangleData);
-
-static void GetGroomInterpolationData(
-	FRDGBuilder& GraphBuilder,
-	FGlobalShaderMap* ShaderMap,
-	FSceneInterface* Scene,
-	const FHairStrandsInstances& Instances,
-	const EHairStrandsProjectionMeshType MeshType,
-	FHairStrandsProjectionMeshData::FLOD& OutGeometries)
-{
-	for (FHairStrandsInstance* AbstractInstance : Instances)
-	{
-		FHairGroupInstance* Instance = static_cast<FHairGroupInstance*>(AbstractInstance);
-
-		if (!Instance)
-		{
-			continue;
-		}
-
-		const FHairStrandsRestRootResource* RestRootResource = nullptr;
-		if (Instance->GetHairGeometry() == EHairGeometryType::Strands)
-		{
-			RestRootResource = Instance->Strands.RestRootResource;
-		}
-		else if (Instance->GetHairGeometry() == EHairGeometryType::Cards)
-		{
-			const int32 LODIndex = Instance->HairGroupPublicData->GetIntLODIndex();
-			if (Instance->Cards.IsValid(LODIndex))
-			{
-				RestRootResource = Instance->Cards.LODs[LODIndex].Guides.RestRootResource;
-			}
-		}
-
-		if (!RestRootResource)
-		{
-			continue;
-		}
-
-		const FCachedGeometry CachedGeometry = GetCacheGeometryForHair(GraphBuilder, ShaderMap, Scene, Instance, &RestRootResource->BulkData, true);
-		if (CachedGeometry.Sections.Num() == 0)
-			continue;
-
-		if (MeshType == EHairStrandsProjectionMeshType::DeformedMesh || MeshType == EHairStrandsProjectionMeshType::RestMesh)
-		{
-			for (const FCachedGeometry::Section& CachedGeometrySection : CachedGeometry.Sections)
-			{
-				FHairStrandsProjectionMeshData::FSection OutSection = ConvertMeshSection(CachedGeometry, CachedGeometrySection);
-				if (MeshType == EHairStrandsProjectionMeshType::RestMesh)
-				{					
-					// If the mesh has some mesh-tranferred data, we display that otherwise we use the rest data
-					const int32 SectionLodIndex = CachedGeometrySection .LODIndex;
-					const bool bHasTransferData = SectionLodIndex < Instance->Debug.TransferredPositions.Num();
-					if (bHasTransferData)
-					{
-						OutSection.PositionBuffer = Instance->Debug.TransferredPositions[SectionLodIndex].SRV;
-					}
-					else if (Instance->Debug.TargetMeshData.LODs.Num() > 0)
-					{
-						OutGeometries = Instance->Debug.TargetMeshData.LODs[0];
-					}
-				}
-				OutGeometries.Sections.Add(OutSection);
-			}
-		}
-
-		if (MeshType == EHairStrandsProjectionMeshType::TargetMesh && Instance->Debug.TargetMeshData.LODs.Num() > 0)
-		{
-			OutGeometries = Instance->Debug.TargetMeshData.LODs[0];
-		}
-
-		if (MeshType == EHairStrandsProjectionMeshType::SourceMesh && Instance->Debug.SourceMeshData.LODs.Num() > 0)
-		{
-			OutGeometries = Instance->Debug.SourceMeshData.LODs[0];
-		}
-	}
-}
-
-///////////////////////////////////////////////////////////////////////////////////////////////////
-
-class FHairProjectionMeshDebugCS : public FGlobalShader
-{
-	DECLARE_GLOBAL_SHADER(FHairProjectionMeshDebugCS);
-	SHADER_USE_PARAMETER_STRUCT(FHairProjectionMeshDebugCS, FGlobalShader);
-
-	class FInputType : SHADER_PERMUTATION_INT("PERMUTATION_INPUT_TYPE", 2);
-	using FPermutationDomain = TShaderPermutationDomain<FInputType>;
-
-	BEGIN_SHADER_PARAMETER_STRUCT(FParameters, )
-		SHADER_PARAMETER(FMatrix44f, LocalToWorld)
-		SHADER_PARAMETER(uint32, VertexOffset)
-		SHADER_PARAMETER(uint32, IndexOffset)
-		SHADER_PARAMETER(uint32, MaxIndexCount)
-		SHADER_PARAMETER(uint32, MaxVertexCount)
-		SHADER_PARAMETER(uint32, MeshUVsChannelOffset)
-		SHADER_PARAMETER(uint32, MeshUVsChannelCount)
-		SHADER_PARAMETER(uint32, bOutputInUVsSpace)
-		SHADER_PARAMETER(uint32, MeshType)
-		SHADER_PARAMETER(uint32, SectionIndex)
-		SHADER_PARAMETER(FVector2f, OutputResolution)
-		SHADER_PARAMETER_SRV(StructuredBuffer, InputIndexBuffer)
-		SHADER_PARAMETER_SRV(StructuredBuffer, InputVertexPositionBuffer)
-		SHADER_PARAMETER_SRV(StructuredBuffer, InputVertexUVsBuffer)
-		SHADER_PARAMETER_STRUCT_REF(FViewUniformShaderParameters, ViewUniformBuffer)
-		SHADER_PARAMETER_STRUCT_INCLUDE(ShaderPrint::FShaderParameters, ShaderPrintUniformBuffer)
-	END_SHADER_PARAMETER_STRUCT()
-
-	static bool ShouldCompilePermutation(const FGlobalShaderPermutationParameters& Parameters)
-	{
-		return IsHairStrandsSupported(EHairStrandsShaderType::All, Parameters.Platform) && ShaderPrint::IsSupported(Parameters.Platform);
-	}
-
-	static void ModifyCompilationEnvironment(const FGlobalShaderPermutationParameters& Parameters, FShaderCompilerEnvironment& OutEnvironment)
-	{
-		FGlobalShader::ModifyCompilationEnvironment(Parameters, OutEnvironment);
-		ShaderPrint::ModifyCompilationEnvironment(Parameters, OutEnvironment);
-
-		OutEnvironment.SetDefine(TEXT("SHADER_MESH_PROJECTION_SKIN_CACHE"), 1);
-	}
-};
-
-IMPLEMENT_GLOBAL_SHADER(FHairProjectionMeshDebugCS, "/Engine/Private/HairStrands/HairStrandsDebug.usf", "MainCS", SF_Compute);
-
-static void AddDebugProjectionMeshPass(
-	FRDGBuilder& GraphBuilder,
-	const FSceneView* View,
-	FGlobalShaderMap* ShaderMap,
-	const FShaderPrintData* ShaderPrintData,
-	const FIntRect Viewport,
-	const TUniformBufferRef<FViewUniformShaderParameters>& ViewUniformBuffer,
-	const EHairStrandsProjectionMeshType MeshType,
-	FHairStrandsProjectionMeshData::FSection& MeshSectionData)
-{
-	if (!View || !ShaderPrint::IsSupported(View->GetShaderPlatform()))
-	{
-		return;
-	}
-
-	// Force ShaderPrint on.
-	ShaderPrint::SetEnabled(true);
-
-	const bool bHasIndexBuffer = MeshSectionData.IndexBuffer != nullptr;
-	const uint32 PrimitiveCount = MeshSectionData.NumPrimitives;
-
-	if (!MeshSectionData.PositionBuffer || PrimitiveCount == 0)
-		return;
-
-	if (ShaderPrintData == nullptr || !ShaderPrint::IsEnabled(*ShaderPrintData))
-	{
-		return;
-	}
-
-	ShaderPrint::RequestSpaceForLines(PrimitiveCount);
-	ShaderPrint::RequestSpaceForTriangles(PrimitiveCount);
-
-	const FIntPoint Resolution(Viewport.Width(), Viewport.Height());
-
-	FHairProjectionMeshDebugCS::FParameters* Parameters = GraphBuilder.AllocParameters<FHairProjectionMeshDebugCS::FParameters>();
-	Parameters->LocalToWorld = FMatrix44f(MeshSectionData.LocalToWorld.ToMatrixWithScale());		// LWC_TODO: Precision loss
-	Parameters->OutputResolution = Resolution;
-	Parameters->MeshType = uint32(MeshType);
-	Parameters->bOutputInUVsSpace = GHairDebugMeshProjection_SkinCacheMeshInUVsSpace ? 1 : 0;
-	Parameters->VertexOffset = MeshSectionData.VertexBaseIndex;
-	Parameters->IndexOffset = MeshSectionData.IndexBaseIndex;
-	Parameters->MaxIndexCount = MeshSectionData.TotalIndexCount;
-	Parameters->MaxVertexCount = MeshSectionData.TotalVertexCount;
-	Parameters->MeshUVsChannelOffset = MeshSectionData.UVsChannelOffset;
-	Parameters->MeshUVsChannelCount = MeshSectionData.UVsChannelCount;
-	Parameters->InputIndexBuffer = MeshSectionData.IndexBuffer;
-	Parameters->InputVertexPositionBuffer = MeshSectionData.PositionBuffer;
-	Parameters->InputVertexUVsBuffer = MeshSectionData.UVsBuffer;
-	Parameters->SectionIndex = MeshSectionData.SectionIndex;
-	Parameters->ViewUniformBuffer = ViewUniformBuffer;
-	ShaderPrint::SetParameters(GraphBuilder, *ShaderPrintData, Parameters->ShaderPrintUniformBuffer);
-
-	FHairProjectionMeshDebugCS::FPermutationDomain PermutationVector;
-	PermutationVector.Set<FHairProjectionMeshDebugCS::FInputType>(bHasIndexBuffer ? 1 : 0);
-	TShaderMapRef<FHairProjectionMeshDebugCS> ComputeShader(ShaderMap, PermutationVector);
-	ClearUnusedGraphResources(ComputeShader, Parameters);
-
-	FComputeShaderUtils::AddPass(
-		GraphBuilder,
-		RDG_EVENT_NAME("HairStrands::MeshProjectionDebug(Mesh)"),
-		ComputeShader,
-		Parameters,
-		FIntVector(FMath::DivideAndRoundUp(PrimitiveCount, 256u), 1, 1));
-}
-
 ///////////////////////////////////////////////////////////////////////////////////////////////////
 
 class FHairProjectionHairDebugCS : public FGlobalShader
@@ -1453,24 +1259,6 @@ void RunHairStrandsDebug(
 		{
 			const FSceneView* LocalView = &View;
 
-			if (GHairDebugMeshProjection_SkinCacheMesh > 0)
-			{
-				auto RenderMeshProjection = [ShaderMap, LocalView, Scene, ShaderPrintData, Viewport, &ViewUniformBuffer, Instances, &GraphBuilder](FRDGBuilder& LocalGraphBuilder, EHairStrandsProjectionMeshType MeshType)
-				{
-					FHairStrandsProjectionMeshData::FLOD MeshProjectionLODData;
-					GetGroomInterpolationData(GraphBuilder, ShaderMap, Scene, Instances, MeshType, MeshProjectionLODData);
-					for (FHairStrandsProjectionMeshData::FSection& Section : MeshProjectionLODData.Sections)
-					{
-						AddDebugProjectionMeshPass(LocalGraphBuilder, LocalView, ShaderMap, ShaderPrintData, Viewport, ViewUniformBuffer, MeshType, Section);
-					}
-				};
-
-				RenderMeshProjection(GraphBuilder, EHairStrandsProjectionMeshType::DeformedMesh);
-				RenderMeshProjection(GraphBuilder, EHairStrandsProjectionMeshType::RestMesh);
-				RenderMeshProjection(GraphBuilder, EHairStrandsProjectionMeshType::SourceMesh);
-				RenderMeshProjection(GraphBuilder, EHairStrandsProjectionMeshType::TargetMesh);
-			}
-
 			uint32 DrawIndex = 0;
 			auto RenderProjectionData = [&GraphBuilder, LocalView, ShaderMap, Viewport, &ViewUniformBuffer, Instances, ShaderPrintData, &DrawIndex](bool bSim)
 			{
diff --git a/Engine/Plugins/Runtime/HairStrands/Source/HairStrandsCore/Private/GroomGeometryCache.cpp b/Engine/Plugins/Runtime/HairStrands/Source/HairStrandsCore/Private/GroomGeometryCache.cpp
--- a/Engine/Plugins/Runtime/HairStrands/Source/HairStrandsCore/Private/GroomGeometryCache.cpp
+++ b/Engine/Plugins/Runtime/HairStrands/Source/HairStrandsCore/Private/GroomGeometryCache.cpp
@@ -18,128 +18,9 @@
 #include "SkeletalMeshDeformerHelpers.h"
 #include "HairStrandsDatas.h"
 
-///////////////////////////////////////////////////////////////////////////////////////////////////
-
-FHairStrandsProjectionMeshData::FSection ConvertMeshSection(const FCachedGeometry& InCachedGeometry, const FCachedGeometry::Section & In)
-{
-	FHairStrandsProjectionMeshData::FSection Out;
-	Out.IndexBuffer = In.IndexBuffer;
-	Out.RDGPositionBuffer = In.RDGPositionBuffer;
-	Out.RDGPreviousPositionBuffer = In.RDGPreviousPositionBuffer;
-	Out.PositionBuffer = In.PositionBuffer;
-	Out.PreviousPositionBuffer = In.PreviousPositionBuffer;
-	Out.UVsBuffer = In.UVsBuffer;
-	Out.UVsChannelOffset = In.UVsChannelOffset;
-	Out.UVsChannelCount = In.UVsChannelCount;
-	Out.TotalVertexCount = In.TotalVertexCount;
-	Out.TotalIndexCount = In.TotalIndexCount;
-	Out.VertexBaseIndex = In.VertexBaseIndex;
-	Out.IndexBaseIndex = In.IndexBaseIndex;
-	Out.NumPrimitives = In.NumPrimitives;
-	Out.NumVertices= In.NumVertices;
-	Out.SectionIndex = In.SectionIndex;
-	Out.LODIndex = In.LODIndex;
-	Out.LocalToWorld = InCachedGeometry.LocalToWorld;
-	return Out;
-}
-
 ////////////////////////////////////////////////////////////////////////////////////////////////////////////////////////
 
-ENGINE_API void UpdateRefToLocalMatrices(TArray<FMatrix44f>& ReferenceToLocal, const USkinnedMeshComponent* InMeshComponent, const FSkeletalMeshRenderData* InSkeletalMeshRenderData, int32 LODIndex, const TArray<FBoneIndexType>* ExtraRequiredBoneIndices = NULL);
-ENGINE_API void UpdatePreviousRefToLocalMatrices(TArray<FMatrix44f>& ReferenceToLocal, const USkinnedMeshComponent* InMeshComponent, const FSkeletalMeshRenderData* InSkeletalMeshRenderData, int32 LODIndex, const TArray<FBoneIndexType>* ExtraRequiredBoneIndices = NULL);
-
- void BuildCacheGeometry(
-	FRDGBuilder& GraphBuilder,
-	FGlobalShaderMap* ShaderMap, 
-	const FSkeletalMeshSceneProxy* Proxy,
-	const FHairStrandsRootBulkData* RootBulkData,
-	const bool bOutputTriangleData,
-	FCachedGeometry& Out)
-{
-	if (Proxy == nullptr)
-	{
-		return;
-	}
-	
-	Out.LocalToWorld = FTransform(Proxy->GetLocalToWorld());
-
-	const FSkeletalMeshObject* SkeletalMeshObject = Proxy->GetMeshObject();
-	if (SkeletalMeshObject == nullptr)
-	{
-		return;
-	}
-
-	const int32 LODIndex = SkeletalMeshObject->GetLOD();
-	Out.LODIndex = LODIndex;
-
-	if (!bOutputTriangleData || !SkeletalMeshObject->HaveValidDynamicData())
-	{
-		return;
-	}
-
-	FSkeletalMeshLODRenderData& LODData = SkeletalMeshObject->GetSkeletalMeshRenderData().LODRenderData[LODIndex];
-	const uint32 SectionCount = LODData.RenderSections.Num();
-	if (SectionCount == 0)
-	{
-		return;
-	}
-
-	const bool bNeedPreviousPosition = IsHairStrandContinuousDecimationReorderingEnabled();
-
-	// Create deformed position buffer (output)
-	Out.DeformedPositionBuffer					= GraphBuilder.CreateBuffer(FRDGBufferDesc::CreateBufferDesc(sizeof(float), LODData.StaticVertexBuffers.PositionVertexBuffer.GetNumVertices() * 3), TEXT("Hair.SkinnedDeformedPositions"));
-	Out.DeformedPreviousPositionBuffer			= bNeedPreviousPosition ?  GraphBuilder.CreateBuffer(FRDGBufferDesc::CreateBufferDesc(sizeof(float), LODData.StaticVertexBuffers.PositionVertexBuffer.GetNumVertices() * 3), TEXT("Hair.SkinnedDeformedPreviousPositions")) : nullptr;
-	FRDGBufferSRVRef DeformedPositionSRV		= GraphBuilder.CreateSRV(Out.DeformedPositionBuffer, PF_R32_FLOAT);
-	FRDGBufferSRVRef DeformedPreviousPositionSRV= bNeedPreviousPosition ? GraphBuilder.CreateSRV(Out.DeformedPreviousPositionBuffer, PF_R32_FLOAT) : nullptr;
-
-	check(RootBulkData);
-	check(RootBulkData->Header.LODs.IsValidIndex(LODIndex));
-	const TArray<uint32>& UniqueSectionCount = RootBulkData->Header.LODs[LODIndex].UniqueSectionIndices;
-	const uint32 EffectiveSectionCount = UniqueSectionCount.Num();
-
-	// Fill in result
-	TArray<FSkinUpdateSection> UpdateSections;
-	UpdateSections.Reserve(EffectiveSectionCount);
-	for (uint32 SectionIndex : UniqueSectionCount)
-	{
-		const FSkelMeshRenderSection& Section = LODData.RenderSections[SectionIndex];
-
-		FSkinUpdateSection& UpdateSection 	= UpdateSections.AddDefaulted_GetRef();
-		UpdateSection.SectionIndex 			= SectionIndex;
-		UpdateSection.NumVertexToProcess 	= Section.NumVertices;
-		UpdateSection.SectionVertexBaseIndex= Section.BaseVertexIndex;
-		UpdateSection.BoneBuffer 			= FSkeletalMeshDeformerHelpers::GetBoneBufferForReading(SkeletalMeshObject, LODIndex, SectionIndex, false);
-		UpdateSection.BonePrevBuffer 		= FSkeletalMeshDeformerHelpers::GetBoneBufferForReading(SkeletalMeshObject, LODIndex, SectionIndex, true);
-
-		FCachedGeometry::Section& OutSection= Out.Sections.AddDefaulted_GetRef();
-		OutSection.RDGPositionBuffer 		= DeformedPositionSRV;
-		OutSection.RDGPreviousPositionBuffer= DeformedPreviousPositionSRV;
-		OutSection.PositionBuffer 			= nullptr; // Do not use the SRV slot, but instead use the RDG buffer created above (DeformedPositionSRV)
-		OutSection.PreviousPositionBuffer 	= nullptr; // Do not use the SRV slot, but instead use the RDG buffer created above (DeformedPositionSRV)
-		OutSection.UVsBuffer 				= LODData.StaticVertexBuffers.StaticMeshVertexBuffer.GetTexCoordsSRV();
-		OutSection.TotalVertexCount 		= LODData.StaticVertexBuffers.PositionVertexBuffer.GetNumVertices();
-		OutSection.IndexBuffer 				= LODData.MultiSizeIndexContainer.GetIndexBuffer()->GetSRV();
-		OutSection.TotalIndexCount 			= LODData.MultiSizeIndexContainer.GetIndexBuffer()->Num();
-		OutSection.UVsChannelCount 			= LODData.StaticVertexBuffers.StaticMeshVertexBuffer.GetNumTexCoords();
-		OutSection.NumPrimitives 			= Section.NumTriangles;
-		OutSection.NumVertices 				= Section.NumVertices;
-		OutSection.IndexBaseIndex 			= Section.BaseIndex;
-		OutSection.VertexBaseIndex 			= Section.BaseVertexIndex;
-		OutSection.SectionIndex 			= SectionIndex;
-		OutSection.LODIndex 				= LODIndex;
-		OutSection.UVsChannelOffset 		= 0; // Assume that we needs to pair meshes based on UVs 0
-	}
-
-	AddSkinUpdatePass(
-		GraphBuilder, 
-		ShaderMap, 
-		LODData, 
-		UpdateSections, 
-		Out.DeformedPositionBuffer, 
-		Out.DeformedPreviousPositionBuffer);
-}
-
- void BuildCacheGeometry(
+ void GetCachedGeometry(
 	 FRDGBuilder& GraphBuilder,
 	 FGlobalShaderMap* ShaderMap, 
 	 const FGeometryCacheSceneProxy* SceneProxy,
diff --git a/Engine/Plugins/Runtime/HairStrands/Source/HairStrandsCore/Private/GroomGeometryCache.h b/Engine/Plugins/Runtime/HairStrands/Source/HairStrandsCore/Private/GroomGeometryCache.h
--- a/Engine/Plugins/Runtime/HairStrands/Source/HairStrandsCore/Private/GroomGeometryCache.h
+++ b/Engine/Plugins/Runtime/HairStrands/Source/HairStrandsCore/Private/GroomGeometryCache.h
@@ -7,21 +7,9 @@
 #include "CachedGeometry.h"
 
 struct FCachedGeometry;
-class FSkeletalMeshSceneProxy;
 class FGlobalShaderMap;
 class FRDGBuilder;
 class FGeometryCacheSceneProxy;
-struct FHairStrandsRootBulkData;
-
-FHairStrandsProjectionMeshData::FSection ConvertMeshSection(const FCachedGeometry& InCachedGeometry, const FCachedGeometry::Section& In);
-
-void BuildCacheGeometry(
-	FRDGBuilder& GraphBuilder,
-	FGlobalShaderMap* ShaderMap, 
-	const FSkeletalMeshSceneProxy* Proxy,
-	const FHairStrandsRootBulkData* RootBulkData,
-	const bool bOutputTriangleData,
-	FCachedGeometry& OutCachedGeometry);
 
 void BuildCacheGeometry(
 	FRDGBuilder& GraphBuilder,
diff --git a/Engine/Plugins/Runtime/HairStrands/Source/HairStrandsCore/Private/GroomManager.cpp b/Engine/Plugins/Runtime/HairStrands/Source/HairStrandsCore/Private/GroomManager.cpp
--- a/Engine/Plugins/Runtime/HairStrands/Source/HairStrandsCore/Private/GroomManager.cpp
+++ b/Engine/Plugins/Runtime/HairStrands/Source/HairStrandsCore/Private/GroomManager.cpp
@@ -30,6 +30,7 @@
 #include "GroomVisualizationData.h"
 #include "SkeletalMeshDeformerHelpers.h"
 #include "HairStrandsInterface.h"
+#include "ShaderPlatformCachedIniValue.h"
 
 static int32 GHairStrandsMinLOD = 0;
 static FAutoConsoleVariableRef CVarGHairStrandsMinLOD(TEXT("r.HairStrands.MinLOD"), GHairStrandsMinLOD, TEXT("Clamp the min hair LOD to this value, preventing to reach lower/high-quality LOD."), ECVF_Scalability);
@@ -153,6 +154,9 @@ static bool IsInstanceFrustumCullingEnable()
 	return GHairStrands_InterpolationFrustumCullingEnable > 0;
 }
 
+///////////////////////////////////////////////////////////////////////////////////////////////////
+// Hair manual skin cache
+
 static bool IsSkeletalMeshEvaluationEnabled()
 {
 	// When deferred skel. mesh update is enabled, hair strands skeletal mesh deformation is not allowed, as skin-cached update happen after 
@@ -188,49 +192,550 @@ static const FPrimitiveSceneInfo* GetMeshSceneInfo(FSceneInterface* Scene, FHair
 	return PrimitiveSceneInfo;
 }
 
+struct FHairGeometryCacheKey
+{
+	uint32 ComponentId = 0;
+	uint32 LODIndex = 0;
+};
+static FORCEINLINE bool operator!= (const FHairGeometryCacheKey& A, const FHairGeometryCacheKey& B) 
+{ 
+	return A.ComponentId != B.ComponentId || A.LODIndex != B.LODIndex; 
+}
+static FORCEINLINE uint32 GetHash(const FHairGeometryCacheKey& In) 
+{ 
+	return Murmur32({In.ComponentId, In.LODIndex}); 
+}
+
+struct FHairGeometryCache
+{
+	enum class ECacheType
+	{
+		SkinCache,
+		HairCache,
+		GeomCache
+	};
+
+	struct FData
+	{
+		const FSkeletalMeshObject* MeshObject		= nullptr;
+		FSkeletalMeshLODRenderData* LODData 		= nullptr;
+		FRDGBufferRef PositionBuffer				= nullptr;
+		FRDGBufferRef PreviousPositionBuffer		= nullptr;
+		FRDGBufferSRVRef PositionSRV				= nullptr;
+		FRDGBufferSRVRef PreviousPositionSRV		= nullptr;
+
+		FHairGeometryCacheKey Key;
+		uint32 Hash 				= 0;
+		uint32 TotalSectionCount 	= 0;
+		TArray<uint32> RequestedSectionIndices;
+		TArray<FSkinUpdateSection> RequestedSections;
+		TBitArray<> RequestedSectionBits;
+	};
+
+	struct FDebugData
+	{
+		const FPrimitiveSceneProxy* Proxy = nullptr;
+		FString MeshComponentName;
+		int32  LODIndex = -1;
+		uint32 InstanceCount = 0;
+		uint32 CacheType = 0;
+		uint32 GeometryType = 0;
+		TBitArray<> SectionBits;
+	};
+
+	FHairGeometryCache(EGroomViewMode InViewMode)
+	{
+		bDebugEnable = InViewMode == EGroomViewMode::MeshProjection;
+	}
+
+	void GetOrAdd(
+		FRDGBuilder& GraphBuilder, 
+		const FSkeletalMeshObject* InMeshObject, 
+		FSkeletalMeshLODRenderData* InLODData,
+		uint32 InLODIndex, 
+		const TArray<uint32>& UniqueSections,
+		FRDGBufferSRVRef& Out, 
+		FRDGBufferSRVRef& OutPrev)
+	{
+		check(InLODData);
+		check(InMeshObject);
+		check(UniqueSections.Num() > 0);
+
+		// Find existing entry
+		const bool bNeedPreviousPosition = IsHairStrandContinuousDecimationReorderingEnabled();
+		const FHairGeometryCacheKey Key = { InMeshObject->GetComponentId(), InLODIndex };
+		const uint32 Hash = GetHash(Key);
+		uint32 Index = HashTable.First(Hash);
+		while (HashTable.IsValid(Index) && Datas[Index].Key != Key)
+		{
+			Index = HashTable.Next(Index);
+			check(Datas[Index].Hash == Hash); // Sanitycheck
+		}
+
+		// Add it if it does not exist
+		if (!HashTable.IsValid(Index))
+		{
+			// Or add it
+			Index = Datas.AddDefaulted();
+			HashTable.Add(Hash, Index);
+
+			FData& Data = Datas[Index];
+			Data.MeshObject 			= InMeshObject;
+			Data.LODData 				= InLODData;
+			Data.Hash					= Hash;
+			Data.Key 					= Key;
+			Data.TotalSectionCount 		= InLODData->RenderSections.Num();
+			Data.PositionBuffer 		= GraphBuilder.CreateBuffer(FRDGBufferDesc::CreateBufferDesc(sizeof(float), InLODData->StaticVertexBuffers.PositionVertexBuffer.GetNumVertices() * 3), TEXT("Hair.SkinnedDeformedPositions"));;
+			Data.PreviousPositionBuffer	= bNeedPreviousPosition ?  GraphBuilder.CreateBuffer(FRDGBufferDesc::CreateBufferDesc(sizeof(float), InLODData->StaticVertexBuffers.PositionVertexBuffer.GetNumVertices() * 3), TEXT("Hair.SkinnedDeformedPreviousPositions")) : nullptr;;
+			Data.PositionSRV 			= GraphBuilder.CreateSRV(Data.PositionBuffer, PF_R32_FLOAT);
+			Data.PreviousPositionSRV	= bNeedPreviousPosition ? GraphBuilder.CreateSRV(Data.PreviousPositionBuffer, PF_R32_FLOAT) : nullptr;
+
+			Data.RequestedSections.Reserve(Data.TotalSectionCount);
+			Data.RequestedSectionBits.Init(false, Data.TotalSectionCount);
+		}
+
+		// Add requested sections
+		FData& Data = Datas[Index];
+		check(Data.PositionSRV);
+		for (uint32 SectionIndex : UniqueSections)
+		{
+			if (!Data.RequestedSectionBits[SectionIndex])
+			{
+				const FSkelMeshRenderSection& Section = InLODData->RenderSections[SectionIndex];
+				FSkinUpdateSection& UpdateSection 	= Data.RequestedSections.AddDefaulted_GetRef();
+				UpdateSection.SectionIndex 			= SectionIndex;
+				UpdateSection.NumVertexToProcess 	= Section.NumVertices;
+				UpdateSection.SectionVertexBaseIndex= Section.BaseVertexIndex;
+				UpdateSection.BoneBuffer 			= FSkeletalMeshDeformerHelpers::GetBoneBufferForReading(InMeshObject, InLODIndex, SectionIndex, false);
+				UpdateSection.BonePrevBuffer 		= FSkeletalMeshDeformerHelpers::GetBoneBufferForReading(InMeshObject, InLODIndex, SectionIndex, true);
+				Data.RequestedSectionBits[SectionIndex] = true;
+			}
+		}
+
+		// Initialized returned values
+		Out 	= Data.PositionSRV;
+		OutPrev = Data.PreviousPositionSRV;
+	}
+
+	void AddDebug(const FHairGroupInstance* InInstance, const FPrimitiveSceneProxy* InProxy, const FCachedGeometry& InGeom, EHairPositionUpdateType InGeometryType, ECacheType InCacheType, uint32 InTotalSectionCount=0)
+	{
+		if (bDebugEnable)
+		{
+			// SkinCache or GeomCache returns the entire set of sections, while HairCache only returns the used sections
+			if (InCacheType != ECacheType::HairCache)
+			{
+				InTotalSectionCount = InGeom.Sections.Num();
+			}
+
+			// Convert a section array into a bitfield
+			auto ToBitArray = [InTotalSectionCount](const TArray<FCachedGeometry::Section>& In)
+			{
+				TBitArray<> Out;
+				Out.Init(false, InTotalSectionCount);
+				for (const FCachedGeometry::Section& S : In)
+				{
+					check(S.SectionIndex < InTotalSectionCount);
+					Out[S.SectionIndex] = true;
+				}
+				return Out;
+			};
+
+			for (FDebugData& DebugData : DebugDatas)
+			{
+				if (DebugData.Proxy == InProxy && DebugData.LODIndex == InGeom.LODIndex)
+				{
+					++DebugData.InstanceCount;
+					DebugData.CacheType    |= 1u << uint32(InCacheType);
+					DebugData.GeometryType |= 1u << uint32(InGeometryType);
+					DebugData.SectionBits.BitwiseOR(DebugData.SectionBits, ToBitArray(InGeom.Sections), EBitwiseOperatorFlags::MaxSize);
+					return;
+				}
+			}
+		
+			FDebugData& DebugData = DebugDatas.AddDefaulted_GetRef();
+			DebugData.MeshComponentName = InInstance->Debug.MeshComponentForDebug->GetName();
+			DebugData.Proxy             = InProxy;
+			DebugData.LODIndex          = InGeom.LODIndex;
+			DebugData.InstanceCount     = 1;
+			DebugData.CacheType        |= 1u << uint32(InCacheType);
+			DebugData.GeometryType     |= 1u << uint32(InGeometryType);
+			DebugData.SectionBits       = ToBitArray(InGeom.Sections);
+		}
+	}
+
+	FHashTable HashTable;
+	TArray<FData> Datas;
+
+	bool bDebugEnable = false;
+	TArray<FDebugData> DebugDatas;
+};
+
+///////////////////////////////////////////////////////////////////////////////////////////////////
+
+class FHairDebugPrintHairSkinCacheCS : public FGlobalShader
+{
+	DECLARE_GLOBAL_SHADER(FHairDebugPrintHairSkinCacheCS);
+	SHADER_USE_PARAMETER_STRUCT(FHairDebugPrintHairSkinCacheCS, FGlobalShader);
+
+	using FPermutationDomain = TShaderPermutationDomain<>;
+
+	BEGIN_SHADER_PARAMETER_STRUCT(FParameters, )
+		SHADER_PARAMETER(uint32, HairSkinCacheEnable)
+		SHADER_PARAMETER(uint32, GPUSkinCacheEnable)
+		SHADER_PARAMETER(uint32, InstanceCount)
+		SHADER_PARAMETER(uint32, UniqueMeshCount)
+		SHADER_PARAMETER_STRUCT(ShaderPrint::FStrings::FShaderParameters, UniqueMeshNames)
+		SHADER_PARAMETER_RDG_BUFFER_SRV(StructuredBuffer<uint4>, Infos)
+		SHADER_PARAMETER_STRUCT_INCLUDE(ShaderPrint::FShaderParameters, ShaderPrintUniformBuffer)
+	END_SHADER_PARAMETER_STRUCT()
+
+public:
+	static bool ShouldCompilePermutation(const FGlobalShaderPermutationParameters& Parameters) { return IsHairStrandsSupported(EHairStrandsShaderType::All, Parameters.Platform) && ShaderPrint::IsSupported(Parameters.Platform); }
+	static void ModifyCompilationEnvironment(const FGlobalShaderPermutationParameters& Parameters, FShaderCompilerEnvironment& OutEnvironment)
+	{
+		// Skip optimization for avoiding long compilation time due to large UAV writes
+		FGlobalShader::ModifyCompilationEnvironment(Parameters, OutEnvironment);
+		ShaderPrint::ModifyCompilationEnvironment(Parameters, OutEnvironment);
+
+		OutEnvironment.CompilerFlags.Add(CFLAG_Debug);
+		OutEnvironment.SetDefine(TEXT("SHADER_PRINT_HAIR_SKIN_CACHE"), 1);
+	}
+};
+
+IMPLEMENT_GLOBAL_SHADER(FHairDebugPrintHairSkinCacheCS, "/Engine/Private/HairStrands/HairStrandsDebug.usf", "MainCS", SF_Compute);
+
+static void AddHairSkinCacheDebugPass(
+	FRDGBuilder& GraphBuilder,
+	FGlobalShaderMap* ShaderMap,
+	const FSceneView* View,
+	const FShaderPrintData* ShaderPrintData,
+	const FHairStrandsInstances& Instances,
+	FHairGeometryCache& HairGeometryCache)
+{
+	if (!HairGeometryCache.bDebugEnable || !ShaderPrint::IsSupported(View->GetShaderPlatform()) || !ShaderPrintData || Instances.Num() == 0)
+	{
+		return;
+	}
+
+	bool bIsGPUSkinCacheEnable = false;
+	{
+		static FShaderPlatformCachedIniValue<int32> PerPlatformCVar(TEXT("r.SkinCache.Mode"));
+		bIsGPUSkinCacheEnable = (PerPlatformCVar.Get(View->GetShaderPlatform()) != 0);
+	}
+
+	const uint32 InstanceCount = Instances.Num();
+	const uint32 UniqueMeshCount = HairGeometryCache.DebugDatas.Num();
+
+	// Force ShaderPrint on.
+	ShaderPrint::SetEnabled(true);
+	// Request more drawing primitives & characters for printing if needed	
+	ShaderPrint::RequestSpaceForLines(UniqueMeshCount * 64u);
+	ShaderPrint::RequestSpaceForCharacters(UniqueMeshCount * 256 + 512);
+
+	ShaderPrint::FStrings UniqueMeshNames(UniqueMeshCount);
+	struct FUniqueMeshInfo
+	{
+		FUintVector4 Data0 = { 0,0,0,0 };
+		FUintVector4 Data1 = { 0,0,0,0 };
+	};
+	TArray<FUniqueMeshInfo> Infos;
+	Infos.Reserve(UniqueMeshCount);
+	for (uint32 UniqueMeshIndex = 0; UniqueMeshIndex < UniqueMeshCount; ++UniqueMeshIndex)
+	{
+		const FHairGeometryCache::FDebugData& DebugData = HairGeometryCache.DebugDatas[UniqueMeshIndex];
+
+		UniqueMeshNames.Add(DebugData.MeshComponentName, UniqueMeshIndex);
+
+		FUniqueMeshInfo& D = Infos.AddDefaulted_GetRef();
+		D.Data0.X = UniqueMeshIndex;
+		D.Data0.Y = DebugData.GeometryType;
+		D.Data0.Z = DebugData.CacheType;
+		D.Data0.W = DebugData.InstanceCount;
+
+		D.Data1.X = DebugData.SectionBits.CountSetBits();
+		D.Data1.Y = DebugData.SectionBits.Num();
+		D.Data1.Z = DebugData.LODIndex;
+		D.Data1.W = 0;
+	}
+
+	const uint32 InfoInBytes  = sizeof(FUniqueMeshInfo);
+	const uint32 InfoInUints = sizeof(FUniqueMeshInfo) / sizeof(uint32);
+	FRDGBufferRef InfoBuffer = CreateVertexBuffer(GraphBuilder, TEXT("Hair.Debug.UniqueMeshNames"), FRDGBufferDesc::CreateBufferDesc(4, InfoInUints * Infos.Num()), Infos.GetData(), InfoInBytes * Infos.Num());
+
+	FHairDebugPrintHairSkinCacheCS::FParameters* Parameters = GraphBuilder.AllocParameters<FHairDebugPrintHairSkinCacheCS::FParameters>();
+	Parameters->HairSkinCacheEnable = GHairStrands_ManualSkinCache > 0 ? 1u : 0u;
+	Parameters->GPUSkinCacheEnable	= bIsGPUSkinCacheEnable ? 1u : 0u;
+	Parameters->InstanceCount 		= InstanceCount;
+	Parameters->UniqueMeshCount 	= UniqueMeshCount;
+	Parameters->UniqueMeshNames 	= UniqueMeshNames.GetParameters(GraphBuilder);
+	Parameters->Infos 				= GraphBuilder.CreateSRV(InfoBuffer, PF_R32_UINT);
+	ShaderPrint::SetParameters(GraphBuilder, *ShaderPrintData, Parameters->ShaderPrintUniformBuffer);
+	TShaderMapRef<FHairDebugPrintHairSkinCacheCS> ComputeShader(ShaderMap);
+
+	FComputeShaderUtils::AddPass(
+		GraphBuilder,
+		RDG_EVENT_NAME("HairStrands::DebugPrintSkinCache"),
+		ComputeShader,
+		Parameters,
+		FIntVector(1, 1, 1));
+}
+
+static void GetOrAllocateCachedGeometry(
+	FRDGBuilder& GraphBuilder,
+	FGlobalShaderMap* ShaderMap, 
+	const FSkeletalMeshSceneProxy* Proxy,
+	const FHairStrandsRootBulkData* RootBulkData,
+	const bool bOutputTriangleData,
+	FCachedGeometry& Out,
+	FHairGeometryCache& OutHairGeometryCache,
+	uint32& OutTotalSectionCount)
+{
+	if (Proxy == nullptr)
+	{
+		return;
+	}
+	
+	Out.LocalToWorld = FTransform(Proxy->GetLocalToWorld());
+
+	const FSkeletalMeshObject* SkeletalMeshObject = Proxy->GetMeshObject();
+	if (SkeletalMeshObject == nullptr)
+	{
+		return;
+	}
+
+	const int32 LODIndex = SkeletalMeshObject->GetLOD();
+	Out.LODIndex = LODIndex;
+
+	FSkeletalMeshLODRenderData& LODData = SkeletalMeshObject->GetSkeletalMeshRenderData().LODRenderData[LODIndex];
+	OutTotalSectionCount = LODData.RenderSections.Num();
+	if (!bOutputTriangleData || OutTotalSectionCount == 0)
+	{
+		return;
+	}
+
+	check(RootBulkData);
+	check(RootBulkData->Header.LODs.IsValidIndex(LODIndex));
+	const TArray<uint32>& UniqueSections = RootBulkData->Header.LODs[LODIndex].UniqueSectionIndices;
+
+	// Create deformed position buffer (output)
+	FRDGBufferSRVRef DeformedPositionSRV = nullptr;
+	FRDGBufferSRVRef DeformedPreviousPositionSRV = nullptr;
+	OutHairGeometryCache.GetOrAdd(GraphBuilder, SkeletalMeshObject, &LODData, LODIndex, UniqueSections, DeformedPositionSRV, DeformedPreviousPositionSRV);
+
+	// Fill in result
+	for (uint32 SectionIndex : UniqueSections)
+	{
+		FCachedGeometry::Section& OutSection= Out.Sections.AddDefaulted_GetRef();
+		OutSection.RDGPositionBuffer 		= DeformedPositionSRV;
+		OutSection.RDGPreviousPositionBuffer= DeformedPreviousPositionSRV;
+		OutSection.PositionBuffer 			= nullptr; // Do not use the SRV slot, but instead use the RDG buffer created above (DeformedPositionSRV)
+		OutSection.PreviousPositionBuffer 	= nullptr; // Do not use the SRV slot, but instead use the RDG buffer created above (DeformedPositionSRV)
+		OutSection.UVsBuffer 				= LODData.StaticVertexBuffers.StaticMeshVertexBuffer.GetTexCoordsSRV();
+		OutSection.TotalVertexCount 		= LODData.StaticVertexBuffers.PositionVertexBuffer.GetNumVertices();
+		OutSection.IndexBuffer 				= LODData.MultiSizeIndexContainer.GetIndexBuffer()->GetSRV();
+		OutSection.TotalIndexCount 			= LODData.MultiSizeIndexContainer.GetIndexBuffer()->Num();
+		OutSection.UVsChannelCount 			= LODData.StaticVertexBuffers.StaticMeshVertexBuffer.GetNumTexCoords();
+		OutSection.NumPrimitives 			= LODData.RenderSections[SectionIndex].NumTriangles;
+		OutSection.NumVertices 				= LODData.RenderSections[SectionIndex].NumVertices;
+		OutSection.IndexBaseIndex 			= LODData.RenderSections[SectionIndex].BaseIndex;
+		OutSection.VertexBaseIndex 			= LODData.RenderSections[SectionIndex].BaseVertexIndex;
+		OutSection.SectionIndex 			= SectionIndex;
+		OutSection.LODIndex 				= LODIndex;
+		OutSection.UVsChannelOffset 		= 0; // Assume that we needs to pair meshes based on UVs 0
+	}
+}
+
+static int32 GetMeshLODIndex(
+	FRDGBuilder& GraphBuilder, 
+	FGlobalShaderMap* ShaderMap,
+	FSceneInterface* Scene,
+	FHairGroupInstance* Instance,
+	FTransform& OutMeshLODLocalToWorld)
+{
+	int32 OutLODIndex = -1;
+	if (const FPrimitiveSceneInfo* PrimitiveSceneInfo = GetMeshSceneInfo(Scene, Instance))
+	{
+		if (Instance->Debug.GroomBindingType == EGroomBindingMeshType::SkeletalMesh && IsSkeletalMeshEvaluationEnabled())
+		{
+			if (const FSkeletalMeshSceneProxy* SceneProxy = static_cast<const FSkeletalMeshSceneProxy*>(PrimitiveSceneInfo->Proxy))
+			{
+				OutLODIndex = SceneProxy->GetMeshObject() ? SceneProxy->GetMeshObject()->GetLOD() : -1;
+				OutMeshLODLocalToWorld = FTransform(SceneProxy->GetLocalToWorld());
+			}
+		}
+		else if (Instance->Debug.GroomBindingType == EGroomBindingMeshType::GeometryCache)
+		{
+			if (const FGeometryCacheSceneProxy* SceneProxy = static_cast<const FGeometryCacheSceneProxy*>(PrimitiveSceneInfo->Proxy))
+			{
+				OutLODIndex = 0;
+				OutMeshLODLocalToWorld = FTransform(SceneProxy->GetLocalToWorld());
+			}
+		}
+	}
+
+	return OutLODIndex;
+}
+
+void GetCachedGeometry(
+	FRDGBuilder& GraphBuilder,
+	FGlobalShaderMap* ShaderMap, 
+	const FGeometryCacheSceneProxy* SceneProxy,
+	const bool bOutputTriangleData,
+	FCachedGeometry& Out);
+
 // Returns the cached geometry of the underlying geometry on which a hair instance is attached to
-FCachedGeometry GetCacheGeometryForHair(
+static FCachedGeometry GetCacheGeometryForHair(
 	FRDGBuilder& GraphBuilder, 
 	FGlobalShaderMap* ShaderMap,
 	FSceneInterface* Scene,
 	FHairGroupInstance* Instance, 
 	const FHairStrandsRootBulkData* RootBulkData,
-	const bool bOutputTriangleData)
+	const bool bOutputTriangleData,
+	const EHairPositionUpdateType PositionUpdateType,
+	FHairGeometryCache& OutHairGeometryCache)
 {
 	FCachedGeometry Out;
-	if (Instance->Debug.GroomBindingType == EGroomBindingMeshType::SkeletalMesh)
+	if (const FPrimitiveSceneInfo* PrimitiveSceneInfo = GetMeshSceneInfo(Scene, Instance))
 	{
-		if (IsSkeletalMeshEvaluationEnabled())
+		if (Instance->Debug.GroomBindingType == EGroomBindingMeshType::SkeletalMesh && IsSkeletalMeshEvaluationEnabled())
 		{
-			if (const FPrimitiveSceneInfo* PrimitiveSceneInfo = GetMeshSceneInfo(Scene, Instance))
+			if (const FSkeletalMeshSceneProxy* SceneProxy = static_cast<const FSkeletalMeshSceneProxy*>(PrimitiveSceneInfo->Proxy))
 			{
-				if (const FSkeletalMeshSceneProxy* SceneProxy = static_cast<const FSkeletalMeshSceneProxy*>(PrimitiveSceneInfo->Proxy))
-				{				
-					SceneProxy->GetCachedGeometry(Out);
-
-					if (GHairStrands_ManualSkinCache > 0 && Out.Sections.Num() == 0)
-					{
-						//#hair_todo: Need to have a (frame) cache to insure that we don't recompute the same projection several time
-						// Actual populate the cache with only the needed part based on the groom projection data. At the moment it recompute everything ...
-						BuildCacheGeometry(GraphBuilder, ShaderMap, SceneProxy, RootBulkData, bOutputTriangleData, Out);
-					}
+				// 1. Get cached geometry from GPU SkinCache, if enabled
+				SceneProxy->GetCachedGeometry(Out);
+				if (Out.Sections.Num() > 0)
+				{
+					OutHairGeometryCache.AddDebug(Instance, SceneProxy, Out, PositionUpdateType, FHairGeometryCache::ECacheType::SkinCache);
+				}
+				// 2. If no cached geometry is  extracted from the skel. mesh proxy, compute it using the manual hair skin cache
+				else if (GHairStrands_ManualSkinCache > 0)
+				{
+					uint32 OutTotalSectionCount = 0;
+					GetOrAllocateCachedGeometry(GraphBuilder, ShaderMap, SceneProxy, RootBulkData, bOutputTriangleData, Out, OutHairGeometryCache, OutTotalSectionCount);
+					OutHairGeometryCache.AddDebug(Instance, SceneProxy, Out, PositionUpdateType, FHairGeometryCache::ECacheType::HairCache, OutTotalSectionCount);
 				}
 			}
 		}
-	}
-	else if (Instance->Debug.GroomBindingType == EGroomBindingMeshType::GeometryCache)
-	{
-		if (const FPrimitiveSceneInfo* PrimitiveSceneInfo = GetMeshSceneInfo(Scene, Instance))
+		else if (Instance->Debug.GroomBindingType == EGroomBindingMeshType::GeometryCache)
 		{
 			if (const FGeometryCacheSceneProxy* SceneProxy = static_cast<const FGeometryCacheSceneProxy*>(PrimitiveSceneInfo->Proxy))
 			{
-				BuildCacheGeometry(GraphBuilder, ShaderMap, SceneProxy, bOutputTriangleData, Out);
+				// Get cached geometry from the geometry cache surface
+				GetCachedGeometry(GraphBuilder, ShaderMap, SceneProxy, bOutputTriangleData, Out);
+				OutHairGeometryCache.AddDebug(Instance, SceneProxy, Out, PositionUpdateType, FHairGeometryCache::ECacheType::GeomCache);
 			}
 		}
 	}
 	return Out;
 }
 
+// Binding surface parameters (skel.mesh/geom. cache)
+static void RunHairStrandsSkinCache(
+	FRDGBuilder& GraphBuilder,
+	FSceneInterface* Scene,
+	const FSceneView* View,
+	const uint32 ViewUniqueID,
+	const FHairStrandsInstances& Instances,
+	const FShaderPrintData* ShaderPrintData,
+	FHairTransientResources& TransientResources,
+	FGlobalShaderMap* ShaderMap)
+{
+	const uint32 TotalInstanceCount = Instances.Num();
+	TransientResources.SimMeshDatas.SetNum(TotalInstanceCount);
+	TransientResources.RenMeshDatas.SetNum(TotalInstanceCount);
+
+	const bool bStrandSupported = IsHairStrandsEnabled(EHairStrandsShaderType::Strands, Scene->GetShaderPlatform());
+	const bool bCardSupported   = IsHairStrandsEnabled(EHairStrandsShaderType::Cards, Scene->GetShaderPlatform());
+
+ 	FHairGeometryCache HairGeometryCache(GetGroomViewMode(*View));
+	for (FHairStrandsInstance* AbstractInstance : Instances)
+	{
+		FHairGroupInstance* Instance = static_cast<FHairGroupInstance*>(AbstractInstance);
+		check(Instance);
+		if (Instance->GeometryType == EHairGeometryType::NoneGeometry)
+		{
+			continue;
+		}
+
+		check(Instance->HairGroupPublicData);
+		const uint32 HairLODIndex = Instance->HairGroupPublicData->LODIndex;
+		int32 MeshLODIndex = -1;
+
+		// 1. Guides
+		{
+			FHairStrandsRootBulkData* RootBulkData = nullptr;
+			if (bStrandSupported && Instance->Guides.IsValid())
+			{	
+				const bool bNeedSurfaceUpdate = 
+					Instance->BindingType == EHairBindingType::Skinning &&
+					(Instance->HairGroupPublicData->IsGlobalInterpolationEnable(HairLODIndex) ||
+					Instance->HairGroupPublicData->IsSimulationEnable(HairLODIndex) ||
+					Instance->HairGroupPublicData->bIsDeformationEnable);
+				if (bNeedSurfaceUpdate)
+				{
+					// Extract MeshLODData and MeshLODIndex
+					check(Instance->Guides.RestRootResource);
+					RootBulkData = &Instance->Guides.RestRootResource->BulkData;
+				}
+			}
+			TransientResources.SimMeshDatas[Instance->RegisteredIndex] = GetCacheGeometryForHair(GraphBuilder, ShaderMap, Scene, Instance, RootBulkData, RootBulkData != nullptr /*bOutputTriangleData*/, EHairPositionUpdateType::Guides, HairGeometryCache);
+			MeshLODIndex = TransientResources.SimMeshDatas[Instance->RegisteredIndex].LODIndex;
+		}
+
+		// 2. Strands/Cards
+		{
+			// Strands
+			if (Instance->GeometryType == EHairGeometryType::Strands && bStrandSupported && Instance->Strands.IsValid())
+			{
+				// Extract MeshLODData and compute MeshLODIndex
+				const bool bNeedOutputTriangleData = Instance->Strands.RestRootResource != nullptr;
+				FHairStrandsRootBulkData* RootBulkData = bNeedOutputTriangleData ? &Instance->Strands.RestRootResource->BulkData : nullptr;
+				TransientResources.RenMeshDatas[Instance->RegisteredIndex] = GetCacheGeometryForHair(GraphBuilder, ShaderMap, Scene, Instance, RootBulkData, RootBulkData != nullptr /*bOutputTriangleData*/, EHairPositionUpdateType::Strands, HairGeometryCache);
+				MeshLODIndex = TransientResources.RenMeshDatas[Instance->RegisteredIndex].LODIndex;
+			}
+			// Cards 
+			// This is only needed for card geometry. Mesh geometry only uses RBF deformation, which are initalized by the guide pass.
+			else if (Instance->GeometryType == EHairGeometryType::Cards && bCardSupported && Instance->Cards.IsValid(HairLODIndex))
+			{
+				FHairStrandsRootBulkData* RootBulkData = nullptr;
+				const bool bNeedSurfaceUpdate = 
+					Instance->BindingType == EHairBindingType::Skinning || 
+					Instance->HairGroupPublicData->IsGlobalInterpolationEnable(HairLODIndex) || 
+					Instance->HairGroupPublicData->bIsDeformationEnable;
+				if (bNeedSurfaceUpdate)
+				{	
+					// Extract MeshLODData and MeshLODIndex
+					check(Instance->Cards.LODs.IsValidIndex(HairLODIndex));
+					RootBulkData = &Instance->Cards.LODs[HairLODIndex].Guides.RestRootResource->BulkData;
+				}
+				TransientResources.RenMeshDatas[Instance->RegisteredIndex] = GetCacheGeometryForHair(GraphBuilder, ShaderMap, Scene, Instance, RootBulkData, RootBulkData != nullptr /*bOutputTriangleData*/, EHairPositionUpdateType::Cards, HairGeometryCache);
+				MeshLODIndex = TransientResources.RenMeshDatas[Instance->RegisteredIndex].LODIndex;
+			}
+		}
+
+		// Update MeshLODIndex
+		Instance->HairGroupPublicData->MeshLODIndex = MeshLODIndex;
+		Instance->Debug.MeshLODIndex 				= MeshLODIndex;
+	}
+
+	// Process manual skin cache requests
+	for (FHairGeometryCache::FData& Data : HairGeometryCache.Datas)
+	{
+		AddSkinUpdatePass(
+			GraphBuilder, 
+			ShaderMap, 
+			*Data.LODData, 
+			Data.RequestedSections, 
+			Data.PositionBuffer, 
+			Data.PreviousPositionBuffer);
+	}
+
+	AddHairSkinCacheDebugPass(GraphBuilder, ShaderMap, View, ShaderPrintData, Instances, HairGeometryCache);
+}
+
+///////////////////////////////////////////////////////////////////////////////////////////////////
+// Guide interpolation
+
 static void RunHairStrandsInterpolation_Guide(
 	FRDGBuilder& GraphBuilder,
 	FSceneInterface* Scene,
@@ -238,6 +743,7 @@ static void RunHairStrandsInterpolation_Guide(
 	const uint32 ViewUniqueID,
 	const FHairStrandsInstances& Instances,
 	const FShaderPrintData* ShaderPrintData,
+	FHairTransientResources& TransientResources,
 	FGlobalShaderMap* ShaderMap)
 {
 	check(IsInRenderingThread());
@@ -257,7 +763,7 @@ static void RunHairStrandsInterpolation_Guide(
 		bool bDeformationEnable = false;
 		bool bGlobalDeformationEnable = false;
 
-		FHairStrandsProjectionMeshData::FLOD MeshDataLOD;
+		FCachedGeometry MeshLODData;
 		FHairGroupInstance* Instance = nullptr;
 
 		bool NeedsMeshUpdate(bool bCheckMeshLODIndex=true) const
@@ -290,30 +796,8 @@ static void RunHairStrandsInterpolation_Guide(
 		InstanceData.bSimulationEnable 			= Instance->HairGroupPublicData->IsSimulationEnable(InstanceData.HairLODIndex);
 		InstanceData.bDeformationEnable 		= Instance->HairGroupPublicData->bIsDeformationEnable;
 		InstanceData.bGlobalDeformationEnable 	= Instance->HairGroupPublicData->IsGlobalInterpolationEnable(InstanceData.HairLODIndex);
-		InstanceData.MeshLODIndex				= -1;
-
-		if (InstanceData.NeedsMeshUpdate(false/*bCheckMeshLODIndex*/))
-		{
-			// Extract MeshDataLOD and MeshLODIndex
-			check(Instance->Guides.RestRootResource);
-			const FCachedGeometry CachedGeometry = GetCacheGeometryForHair(GraphBuilder, ShaderMap, Scene, Instance, &Instance->Guides.RestRootResource->BulkData, true /*bOutputTriangleData*/);
-			for (const FCachedGeometry::Section& CachedGeometrySection : CachedGeometry.Sections)
-			{
-				InstanceData.MeshDataLOD.Sections.Add(ConvertMeshSection(CachedGeometry, CachedGeometrySection));
-			}
-			InstanceData.MeshLODIndex    = CachedGeometry.LODIndex;
-			Instance->Debug.MeshLODIndex = CachedGeometry.LODIndex;
-
-			// Sanity check
-			check(InstanceData.Instance->Guides.IsValid());
-		}
-		else
-		{
-			// Extract MeshLODIndex
-			const FCachedGeometry CachedGeometry = GetCacheGeometryForHair(GraphBuilder, ShaderMap, Scene, Instance, nullptr, false /*bOutputTriangleData*/);
-			InstanceData.MeshLODIndex    = CachedGeometry.LODIndex;
-			Instance->Debug.MeshLODIndex = CachedGeometry.LODIndex;
-		}
+		InstanceData.MeshLODIndex				= Instance->HairGroupPublicData->MeshLODIndex;
+		InstanceData.MeshLODData				= TransientResources.GetMeshLODData(Instance->RegisteredIndex, true /*bSim*/);
 
 		if (InstanceData.ActiveGroomCacheType == EGroomCacheType::Guides)  { bHasAnySimCacheInstances = true; }
 	}	
@@ -331,8 +815,8 @@ static void RunHairStrandsInterpolation_Guide(
 			AddHairStrandUpdateMeshTrianglesPass(
 				GraphBuilder,
 				ShaderMap,
-				InstanceData.Instance->Debug.MeshLODIndex,
-				InstanceData.MeshDataLOD,
+				InstanceData.MeshLODIndex,
+				InstanceData.MeshLODData,
 				InstanceData.Instance->Guides.RestRootResource,
 				InstanceData.Instance->Guides.DeformedRootResource);
 		}
@@ -349,8 +833,8 @@ static void RunHairStrandsInterpolation_Guide(
 				AddHairStrandInitMeshSamplesPass(
 					GraphBuilder,
 					ShaderMap,
-					InstanceData.Instance->Debug.MeshLODIndex,
-					InstanceData.MeshDataLOD,
+					InstanceData.MeshLODIndex,
+					InstanceData.MeshLODData,
 					InstanceData.Instance->Guides.RestRootResource,
 					InstanceData.Instance->Guides.DeformedRootResource);
 			}
@@ -368,8 +852,8 @@ static void RunHairStrandsInterpolation_Guide(
 				AddHairStrandUpdateMeshSamplesPass(
 					GraphBuilder,
 					ShaderMap,
-					InstanceData.Instance->Debug.MeshLODIndex,
-					InstanceData.MeshDataLOD,
+					InstanceData.MeshLODIndex,
+					InstanceData.MeshLODData,
 					InstanceData.Instance->Guides.RestRootResource,
 					InstanceData.Instance->Guides.DeformedRootResource);
 			}
@@ -389,7 +873,7 @@ static void RunHairStrandsInterpolation_Guide(
 				EHairPositionUpdateType::Guides,
 				InstanceData.Instance->RegisteredIndex,
 				InstanceData.HairLODIndex,
-				InstanceData.Instance->Debug.MeshLODIndex,
+				InstanceData.MeshLODIndex,
 				InstanceData.Instance->Guides.DeformedRootResource,
 				InstanceData.Instance->Guides.DeformedResource);
 
@@ -418,7 +902,7 @@ static void RunHairStrandsInterpolation_Guide(
 				EHairPositionUpdateType::Guides,
 				InstanceData.Instance->RegisteredIndex,
 				InstanceData.HairLODIndex,
-				InstanceData.Instance->Debug.MeshLODIndex,
+				InstanceData.MeshLODIndex,
 				nullptr,
 				InstanceData.Instance->Guides.DeformedResource);
 
@@ -442,7 +926,7 @@ static void RunHairStrandsInterpolation_Guide(
 			GraphBuilder,
 			ShaderMap,
 			InstanceData.Instance->RegisteredIndex,
-			InstanceData.Instance->Debug.MeshLODIndex,
+			InstanceData.MeshLODIndex,
 			InstanceData.Instance->Guides.RestResource->GetPointCount(),
 			InstanceData.Instance->Guides.RestRootResource,
 			InstanceData.Instance->Guides.DeformedRootResource,
@@ -538,6 +1022,9 @@ static void RunHairStrandsInterpolation_Guide(
 	}
 }
 
+///////////////////////////////////////////////////////////////////////////////////////////////////
+// Strands interpolation
+
 void AddDrawDebugClusterPass(
 	FRDGBuilder& GraphBuilder,
 	const FSceneView& View,
@@ -605,7 +1092,7 @@ static void RunHairStrandsInterpolation_Strands(
 
 		FInstanceRDGResources RDGResources;
 
-		FHairStrandsProjectionMeshData::FLOD MeshDataLOD;
+		FCachedGeometry MeshLODData;
 		FHairGroupInstance* Instance = nullptr;
 		FRDGHairStrandsCullingData CullingData;
 
@@ -645,7 +1132,7 @@ static void RunHairStrandsInterpolation_Strands(
 		InstanceData.bSimulationEnable 			= Instance->HairGroupPublicData->IsSimulationEnable(InstanceData.HairLODIndex);
 		InstanceData.bDeformationEnable 		= Instance->HairGroupPublicData->bIsDeformationEnable;
 		InstanceData.bGlobalDeformationEnable 	= Instance->HairGroupPublicData->IsGlobalInterpolationEnable(InstanceData.HairLODIndex);
-		InstanceData.MeshLODIndex				= -1;
+		InstanceData.MeshLODIndex				= Instance->HairGroupPublicData->MeshLODIndex;
 		InstanceData.bNeedDeformation 			= Instance->Strands.DeformedResource != nullptr;
 		InstanceData.bNeedRaytracing			= false;
 		InstanceData.bNeedRaytracingUpdate		= false;
@@ -690,15 +1177,7 @@ static void RunHairStrandsInterpolation_Strands(
 		if (InstanceData.ActiveGroomCacheType == EGroomCacheType::Guides)  { bHasAnySimCacheInstances = true; }
 		if (InstanceData.ActiveGroomCacheType == EGroomCacheType::Strands) { bHasAnyRenCacheInstances = true; }
 
-		// Extract MeshDataLOD and compute MeshLODIndex
-		const bool bNeedOutputTriangleData = Instance->Strands.RestRootResource != nullptr;
-		const FCachedGeometry CachedGeometry = GetCacheGeometryForHair(GraphBuilder, ShaderMap, Scene, Instance, bNeedOutputTriangleData ? &Instance->Strands.RestRootResource->BulkData : nullptr, bNeedOutputTriangleData);
-		for (const FCachedGeometry::Section& CachedGeometrySection : CachedGeometry.Sections)
-		{
-			InstanceData.MeshDataLOD.Sections.Add(ConvertMeshSection(CachedGeometry, CachedGeometrySection));
-		}
-		InstanceData.MeshLODIndex    = CachedGeometry.LODIndex;
-		Instance->Debug.MeshLODIndex = CachedGeometry.LODIndex;
+		InstanceData.MeshLODData = TransientResources.GetMeshLODData(Instance->RegisteredIndex, false /*bSim*/);
 		Instance->HairGroupPublicData->VFInput.Strands	= FHairGroupPublicData::FVertexFactoryInput::FStrands();
 
 		if (InstanceData.MeshLODIndex >= 0 && (InstanceData.BindingType == EHairBindingType::Skinning || InstanceData.bGlobalDeformationEnable))
@@ -770,7 +1249,7 @@ static void RunHairStrandsInterpolation_Strands(
 					GraphBuilder,
 					ShaderMap,
 					InstanceData.MeshLODIndex,
-					InstanceData.MeshDataLOD,
+					InstanceData.MeshLODData,
 					InstanceData.Instance->Strands.RestRootResource,
 					InstanceData.Instance->Strands.DeformedRootResource);
 				Transitions.Add(RegisterAsSRV(GraphBuilder, InstanceData.Instance->Strands.DeformedRootResource->LODs[InstanceData.MeshLODIndex].GetDeformedUniqueTrianglePositionBuffer(FHairStrandsDeformedRootResource::FLOD::Current)));
@@ -1223,6 +1702,9 @@ static void RunHairStrandsInterpolation_Strands(
 	ExternalAccessQueue.Submit(GraphBuilder);
 }
 
+///////////////////////////////////////////////////////////////////////////////////////////////////
+// Cards/Meshes interpolation
+
 enum class EHairCardsSimulationType
 {
 	None,
@@ -1240,6 +1722,7 @@ static void RunHairStrandsInterpolation_Cards(
 	const uint32 ViewUniqueID,
 	const FHairStrandsInstances& Instances,
 	const FShaderPrintData* ShaderPrintData,
+	FHairTransientResources& TransientResources,
 	FGlobalShaderMap* ShaderMap)
 {
 	const bool bCardSupported = IsHairStrandsEnabled(EHairStrandsShaderType::Cards, Scene->GetShaderPlatform());
@@ -1268,7 +1751,7 @@ static void RunHairStrandsInterpolation_Cards(
 		EHairCardsSimulationType CardsSimulationType = EHairCardsSimulationType::None;
 		FHairGroupInstance::FCards::FLOD*  CardInstance = nullptr;
 		FHairGroupInstance::FMeshes::FLOD* MeshInstance = nullptr;
-		FHairStrandsProjectionMeshData::FLOD MeshDataLOD;
+		FCachedGeometry MeshLODData;
 		FHairGroupInstance* Instance = nullptr;
 	};
 
@@ -1306,34 +1789,9 @@ static void RunHairStrandsInterpolation_Cards(
 			InstanceData.HairLODIndex 				= HairLODIndex;
 			InstanceData.bSimulationEnable 			= Instance->HairGroupPublicData->IsSimulationEnable(HairLODIndex);
 			InstanceData.bDeformationEnable 		= Instance->HairGroupPublicData->bIsDeformationEnable;
-			InstanceData.bGlobalDeformationEnable 	= Instance->HairGroupPublicData->IsGlobalInterpolationEnable(HairLODIndex);
-			InstanceData.MeshLODIndex				= -1;
-
-			// Binding surface parameters (skel.mesh/geom. cache)
-			// This is only needed for card geometry. Mesh geometry only uses RBF deformation, which are initalized by the guide pass.
-			if (Instance->GeometryType == EHairGeometryType::Cards)
-			{				
-				const bool bNeedSurfaceUpdate = InstanceData.BindingType == EHairBindingType::Skinning || InstanceData.bGlobalDeformationEnable || InstanceData.bDeformationEnable;
-				if (bNeedSurfaceUpdate)
-				{	
-					// Extract MeshDataLOD and MeshLODIndex
-					check(Instance->Cards.LODs.IsValidIndex(InstanceData.HairLODIndex));
-					const FCachedGeometry CachedGeometry = GetCacheGeometryForHair(GraphBuilder, ShaderMap, Scene, Instance, &Instance->Cards.LODs[InstanceData.HairLODIndex].Guides.RestRootResource->BulkData, true /*bOutputTriangleData*/);
-					for (const FCachedGeometry::Section& CachedGeometrySection : CachedGeometry.Sections)
-					{
-						InstanceData.MeshDataLOD.Sections.Add(ConvertMeshSection(CachedGeometry, CachedGeometrySection));
-					}
-					InstanceData.MeshLODIndex    = CachedGeometry.LODIndex;
-					Instance->Debug.MeshLODIndex = CachedGeometry.LODIndex;
-				}
-				else
-				{
-					// Extract MeshLODIndex
-					const FCachedGeometry CachedGeometry = GetCacheGeometryForHair(GraphBuilder, ShaderMap, Scene, Instance, nullptr, false /*bOutputTriangleData*/);
-					InstanceData.MeshLODIndex    = CachedGeometry.LODIndex;
-					Instance->Debug.MeshLODIndex = CachedGeometry.LODIndex;
-				}
-			}
+			InstanceData.bGlobalDeformationEnable 	= Instance->HairGroupPublicData->IsGlobalInterpolationEnable(HairLODIndex);			
+			InstanceData.MeshLODIndex				= Instance->HairGroupPublicData->MeshLODIndex;
+			InstanceData.MeshLODData  				= TransientResources.GetMeshLODData(Instance->RegisteredIndex, false /*bSim*/);
 
 			// Card or Mesh specific parameters
 			if (Instance->GeometryType == EHairGeometryType::Cards)
@@ -1406,7 +1864,7 @@ static void RunHairStrandsInterpolation_Cards(
 				GraphBuilder,
 				ShaderMap,
 				InstanceData.MeshLODIndex,
-				InstanceData.MeshDataLOD,
+				InstanceData.MeshLODData,
 				InstanceData.CardInstance->Guides.RestRootResource,
 				InstanceData.CardInstance->Guides.DeformedRootResource);
 		}
@@ -1693,7 +2151,6 @@ static void RunHairBufferSwap(const FHairStrandsInstances& Instances, const TArr
 	{
 		FHairGroupInstance* Instance = static_cast<FHairGroupInstance*>(AbstractInstance);
 
-		int32 MeshLODIndex = -1;
 		check(Instance);
 		if (!Instance)
 			continue;
@@ -1930,8 +2387,8 @@ static void RunHairLODSelection(
 
 		check(Instance);
 		check(Instance->HairGroupPublicData);
-		const FCachedGeometry CachedGeometry = GetCacheGeometryForHair(GraphBuilder, ShaderMap, Scene, Instance, nullptr, false);
-		const int32 MeshLODIndex = CachedGeometry.LODIndex;
+		FTransform MeshLODLocalToWorld = FTransform::Identity;
+		const int32 MeshLODIndex = GetMeshLODIndex(GraphBuilder, ShaderMap, Scene, Instance, MeshLODLocalToWorld);
 
 		// Perform LOD selection based on all the views	
 		// CPU LOD selection. 
@@ -2248,7 +2705,7 @@ static void RunHairLODSelection(
 		if (GetHairSwapBufferType() != EHairBufferSwapType::Tick)
 		{
 			Instance->Debug.SkinningPreviousLocalToWorld = Instance->Debug.SkinningCurrentLocalToWorld;
-			Instance->Debug.SkinningCurrentLocalToWorld = CachedGeometry.LocalToWorld;
+			Instance->Debug.SkinningCurrentLocalToWorld = MeshLODLocalToWorld;
 		}
 		Instance->LocalToWorld = Instance->GetCurrentLocalToWorld();
 	}
@@ -2366,12 +2823,14 @@ void ProcessHairStrandsBookmark(
 				Parameters.AllViews);
 		}
 		check(GraphBuilder);
+
 		RunHairLODSelection(
 			*GraphBuilder,
 			Parameters.Scene,
 			*Parameters.Instances,
 			Parameters.AllViews,
 			Parameters.ShaderMap);
+
 	}
 	else if (Bookmark == EHairStrandsBookmark::ProcessEndOfFrame)
 	{
@@ -2385,6 +2844,17 @@ void ProcessHairStrandsBookmark(
 	else if (Bookmark == EHairStrandsBookmark::ProcessGuideInterpolation)
 	{
 		check(GraphBuilder);
+
+		RunHairStrandsSkinCache(
+			*GraphBuilder,
+			Parameters.Scene,
+			Parameters.View,
+			Parameters.ViewUniqueID,
+			*Instances,
+			Parameters.ShaderPrintData,
+			*Parameters.TransientResources,
+			Parameters.ShaderMap);
+
 		RunHairStrandsInterpolation_Guide(
 			*GraphBuilder,
 			Parameters.Scene,
@@ -2392,6 +2862,7 @@ void ProcessHairStrandsBookmark(
 			Parameters.ViewUniqueID,
 			*Instances,
 			Parameters.ShaderPrintData,
+			*Parameters.TransientResources,
 			Parameters.ShaderMap);
 	}
 	else if (Bookmark == EHairStrandsBookmark::ProcessCardsAndMeshesInterpolation_PrimaryView || Bookmark == EHairStrandsBookmark::ProcessCardsAndMeshesInterpolation_ShadowView)
@@ -2405,6 +2876,7 @@ void ProcessHairStrandsBookmark(
 			Parameters.ViewUniqueID,
 			*Instances,
 			Parameters.ShaderPrintData,
+			*Parameters.TransientResources,
 			Parameters.ShaderMap);
 	}
 	else if (Bookmark == EHairStrandsBookmark::ProcessStrandsInterpolation)
diff --git a/Engine/Plugins/Runtime/HairStrands/Source/HairStrandsCore/Private/HairStrandsMeshProjection.cpp b/Engine/Plugins/Runtime/HairStrands/Source/HairStrandsCore/Private/HairStrandsMeshProjection.cpp
--- a/Engine/Plugins/Runtime/HairStrands/Source/HairStrandsCore/Private/HairStrandsMeshProjection.cpp
+++ b/Engine/Plugins/Runtime/HairStrands/Source/HairStrandsCore/Private/HairStrandsMeshProjection.cpp
@@ -181,21 +181,21 @@ public:
 
 IMPLEMENT_GLOBAL_SHADER(FHairUpdateMeshTriangleCS, "/Engine/Private/HairStrands/HairStrandsMesh.usf", "MainCS", SF_Compute);
 
-bool AddHairStrandUpdateMeshTrianglesPass(
+static bool AddHairStrandUpdateMeshTrianglesPass(
 	FRDGBuilder& GraphBuilder,
 	FGlobalShaderMap* ShaderMap,
 	const uint32 UniqueTriangleCount, 
 	const uint32 TotalSectionCount, // Total section count of the underlying mesh
-	const FHairStrandsProjectionMeshData::FLOD& MeshData,
+	const FCachedGeometry& MeshLODData,
 	FRDGBufferSRVRef UniqueTriangleIndexSRV,
 	FRDGBufferUAVRef OutputCurrUAV,
 	FRDGBufferUAVRef OutputPrevUAV)
 {
-	const uint32 EffectiveSectionCount 	= MeshData.Sections.Num(); // Section count which are valid & used by groom data
+	const uint32 EffectiveSectionCount 	= MeshLODData.Sections.Num(); // Section count which are valid & used by groom data
 
 	// If a skel. mesh is not streaming yet, its SRV will be null
 	const bool bValid = EffectiveSectionCount > 0 && TotalSectionCount < GetHairStrandsMaxSectionCount();
-	const bool bReady = MeshData.Sections[0].IndexBuffer != nullptr && MeshData.Sections[0].UVsBuffer != nullptr && (MeshData.Sections[0].RDGPositionBuffer != nullptr || MeshData.Sections[0].PositionBuffer != nullptr);
+	const bool bReady = MeshLODData.Sections[0].IndexBuffer != nullptr && MeshLODData.Sections[0].UVsBuffer != nullptr && (MeshLODData.Sections[0].RDGPositionBuffer != nullptr || MeshLODData.Sections[0].PositionBuffer != nullptr);
 	if (!bReady || !bValid)
 	{
 		return false;
@@ -204,12 +204,12 @@ bool AddHairStrandUpdateMeshTrianglesPass(
 	FHairUpdateMeshTriangleCS::FParameters CommonParameters;
 	CommonParameters.MaxUniqueTriangleCount 		= UniqueTriangleCount;
 	CommonParameters.MaxSectionCount 				= TotalSectionCount;
-	CommonParameters.RDGMeshPositionBuffer			= MeshData.Sections[0].RDGPositionBuffer;
-	CommonParameters.RDGMeshPreviousPositionBuffer 	= MeshData.Sections[0].RDGPreviousPositionBuffer;
-	CommonParameters.MeshPositionBuffer				= MeshData.Sections[0].PositionBuffer;
-	CommonParameters.MeshPreviousPositionBuffer		= MeshData.Sections[0].PreviousPositionBuffer;
-	CommonParameters.MeshIndexBuffer				= MeshData.Sections[0].IndexBuffer;
-	CommonParameters.MeshUVsBuffer					= MeshData.Sections[0].UVsBuffer;
+	CommonParameters.RDGMeshPositionBuffer			= MeshLODData.Sections[0].RDGPositionBuffer;
+	CommonParameters.RDGMeshPreviousPositionBuffer 	= MeshLODData.Sections[0].RDGPreviousPositionBuffer;
+	CommonParameters.MeshPositionBuffer				= MeshLODData.Sections[0].PositionBuffer;
+	CommonParameters.MeshPreviousPositionBuffer		= MeshLODData.Sections[0].PreviousPositionBuffer;
+	CommonParameters.MeshIndexBuffer				= MeshLODData.Sections[0].IndexBuffer;
+	CommonParameters.MeshUVsBuffer					= MeshLODData.Sections[0].UVsBuffer;
 	CommonParameters.UniqueTriangleIndices 			= UniqueTriangleIndexSRV;
 	CommonParameters.OutUniqueTriangleCurrPosition	= OutputCurrUAV;
 	CommonParameters.OutUniqueTrianglePrevPosition	= OutputPrevUAV;
@@ -239,7 +239,7 @@ bool AddHairStrandUpdateMeshTrianglesPass(
 	// Allocate data for *all* sections, but only fill in the used/valid sections
 	TArray<FSectionData> SectionDatas;
 	SectionDatas.Init(Default, TotalSectionCount); 
-	for (const FHairStrandsProjectionMeshData::FSection& MeshSectionData : MeshData.Sections)
+	for (const FCachedGeometry::Section& MeshSectionData : MeshLODData.Sections)
 	{
 		const uint32 SectionIndex = MeshSectionData.SectionIndex;
 		SectionDatas[SectionIndex].TotalIndexCount	= MeshSectionData.TotalIndexCount;
@@ -301,7 +301,7 @@ void AddHairStrandUpdateMeshTrianglesPass(
 	FRDGBuilder& GraphBuilder,
 	FGlobalShaderMap* ShaderMap,
 	const int32 MeshLODIndex,
-	const FHairStrandsProjectionMeshData::FLOD& MeshData,
+	const FCachedGeometry& MeshLODData,
 	FHairStrandsRestRootResource* RestResources,
 	FHairStrandsDeformedRootResource* DeformedResources)
 {	
@@ -333,7 +333,7 @@ void AddHairStrandUpdateMeshTrianglesPass(
 			ShaderMap,
 			UniqueTriangleCount,
 			TotalMeshSectionCount,
-			MeshData,
+			MeshLODData,
 			RegisterAsSRV(GraphBuilder, RestLODData.UniqueTriangleIndexBuffer),
 			OutputCurrBuffer.UAV,
 			OutputPrevBuffer.UAV))
@@ -515,7 +515,7 @@ void AddHairStrandInitMeshSamplesPass(
 	FRDGBuilder& GraphBuilder,
 	FGlobalShaderMap* ShaderMap,
 	const int32 LODIndex,
-	const FHairStrandsProjectionMeshData::FLOD& MeshData,
+	const FCachedGeometry& MeshData,
 	FHairStrandsRestRootResource* RestResources,
 	FHairStrandsDeformedRootResource* DeformedResources)
 {
@@ -556,7 +556,7 @@ void AddHairStrandInitMeshSamplesPass(
 	};
 	TArray<FSectionData> SectionDatas;
 	SectionDatas.Init(FSectionData(), TotalSectionCount);
-	for (const FHairStrandsProjectionMeshData::FSection& Section : MeshData.Sections)
+	for (const FCachedGeometry::Section& Section : MeshData.Sections)
 	{
 		const uint32 SectionIndex = Section.SectionIndex;
 		SectionDatas[SectionIndex].bIsSwapped = 
@@ -631,7 +631,7 @@ void AddHairStrandUpdateMeshSamplesPass(
 	FRDGBuilder& GraphBuilder,
 	FGlobalShaderMap* ShaderMap,
 	const int32 LODIndex,
-	const FHairStrandsProjectionMeshData::FLOD& MeshData,
+	const FCachedGeometry& MeshData,
 	FHairStrandsRestRootResource* RestResources,
 	FHairStrandsDeformedRootResource* DeformedResources)
 {
diff --git a/Engine/Plugins/Runtime/HairStrands/Source/HairStrandsCore/Public/GroomInstance.h b/Engine/Plugins/Runtime/HairStrands/Source/HairStrandsCore/Public/GroomInstance.h
--- a/Engine/Plugins/Runtime/HairStrands/Source/HairStrandsCore/Public/GroomInstance.h
+++ b/Engine/Plugins/Runtime/HairStrands/Source/HairStrandsCore/Public/GroomInstance.h
@@ -250,11 +250,6 @@ struct HAIRSTRANDSCORE_API FHairGroupInstance : public FHairStrandsInstance
 
 		TSharedPtr<class IGroomCacheBuffers, ESPMode::ThreadSafe> GroomCacheBuffers;
 
-		// Transfer
-		TArray<FRWBuffer> TransferredPositions;
-		FHairStrandsProjectionMeshData SourceMeshData;
-		FHairStrandsProjectionMeshData TargetMeshData;
-
 		// Resources
 		FHairStrandsDebugDatas::FResources* HairDebugResource = nullptr;
 	} Debug;
diff --git a/Engine/Plugins/Runtime/HairStrands/Source/HairStrandsCore/Public/HairStrandsMeshProjection.h b/Engine/Plugins/Runtime/HairStrands/Source/HairStrandsCore/Public/HairStrandsMeshProjection.h
--- a/Engine/Plugins/Runtime/HairStrands/Source/HairStrandsCore/Public/HairStrandsMeshProjection.h
+++ b/Engine/Plugins/Runtime/HairStrands/Source/HairStrandsCore/Public/HairStrandsMeshProjection.h
@@ -19,36 +19,6 @@ struct FHairCardsRestResource;
 struct FHairCardsDeformedResource;
 struct FHairStrandsDeformedResource;
 
-struct FHairStrandsProjectionMeshData
-{
-	struct FSection
-	{
-		FTransform LocalToWorld;
-		FRDGBufferSRVRef RDGPositionBuffer = nullptr;
-		FRDGBufferSRVRef RDGPreviousPositionBuffer = nullptr;
-		FRHIShaderResourceView* PositionBuffer = nullptr;
-		FRHIShaderResourceView* PreviousPositionBuffer = nullptr;
-		FRHIShaderResourceView* UVsBuffer = nullptr;
-		FRHIShaderResourceView* IndexBuffer = nullptr;
-		uint32 UVsChannelCount = 0;
-		uint32 UVsChannelOffset = 0;
-		uint32 NumPrimitives = 0;
-		uint32 NumVertices = 0;
-		uint32 VertexBaseIndex = 0;
-		uint32 IndexBaseIndex = 0;
-		uint32 TotalVertexCount = 0;
-		uint32 TotalIndexCount = 0;
-		uint32 SectionIndex = 0;
-		int32 LODIndex = 0;
-	};
-
-	struct FLOD
-	{
-		TArray<FSection> Sections;
-	};
-	TArray<FLOD> LODs;
-};
-
 // Return the max number of section/triangle a skeletal mesh can have. After this count, binding will be disabled
 uint32 GetHairStrandsMaxSectionCount();
 uint32 GetHairStrandsMaxTriangleCount();
@@ -57,8 +27,8 @@ uint32 GetHairStrandsMaxTriangleCount();
 void AddHairStrandUpdateMeshTrianglesPass(
 	FRDGBuilder& GraphBuilder,
 	FGlobalShaderMap* ShaderMap,
-	const int32 LODIndex,
-	const FHairStrandsProjectionMeshData::FLOD& MeshData,
+	const int32 MeshLODIndex,
+	const FCachedGeometry& MeshLODData,
 	FHairStrandsRestRootResource* RestResources,
 	FHairStrandsDeformedRootResource* DeformedResources);
 
@@ -66,8 +36,8 @@ void AddHairStrandUpdateMeshTrianglesPass(
 void AddHairStrandInitMeshSamplesPass(
 	FRDGBuilder& GraphBuilder,
 	FGlobalShaderMap* ShaderMap,
-	const int32 LODIndex,
-	const FHairStrandsProjectionMeshData::FLOD& MeshData,
+	const int32 MeshLODIndex,
+	const FCachedGeometry& MeshLODData,
 	FHairStrandsRestRootResource* RestResources,
 	FHairStrandsDeformedRootResource* DeformedResources);
 
@@ -75,8 +45,8 @@ void AddHairStrandInitMeshSamplesPass(
 void AddHairStrandUpdateMeshSamplesPass(
 	FRDGBuilder& GraphBuilder,
 	FGlobalShaderMap* ShaderMap,
-	const int32 LODIndex,
-	const FHairStrandsProjectionMeshData::FLOD& ProjectionMeshData,
+	const int32 MeshLODIndex,
+	const FCachedGeometry& MeshLODData,
 	FHairStrandsRestRootResource* RestResources,
 	FHairStrandsDeformedRootResource* DeformedResources);
 
diff --git a/Engine/Shaders/Private/HairStrands/HairStrandsDebug.usf b/Engine/Shaders/Private/HairStrands/HairStrandsDebug.usf
--- a/Engine/Shaders/Private/HairStrands/HairStrandsDebug.usf
+++ b/Engine/Shaders/Private/HairStrands/HairStrandsDebug.usf
@@ -4,7 +4,7 @@
 #include "../HairShadingCommon.ush"
 #include "../ColorMap.ush"
 
-#if SHADER_PRINT_INSTANCE || SHADER_PRINT || SHADER_PRINT_MEMORY || SHADER_CLUSTERAABB || SHADER_DOMINFO
+#if SHADER_PRINT_INSTANCE || SHADER_PRINT || SHADER_PRINT_MEMORY || SHADER_CLUSTERAABB || SHADER_DOMINFO || SHADER_PRINT_HAIR_SKIN_CACHE
 /////////////////////////////////////////////////////////////////////////////////////////
 // Utils function
 
@@ -88,7 +88,7 @@ void PrintPercentage(inout FShaderPrintContext Context, float Value, FFontColor
 	Print(Context, TEXT("% "), InColor);
 }
 
-#endif // SHADER_PRINT_INSTANCE || SHADER_PRINT || SHADER_PRINT_MEMORY || SHADER_CLUSTERAABB || SHADER_DOMINFO
+#endif // SHADER_PRINT_INSTANCE || SHADER_PRINT || SHADER_PRINT_MEMORY || SHADER_CLUSTERAABB || SHADER_DOMINFO || SHADER_PRINT_HAIR_SKIN_CACHE
 
 /////////////////////////////////////////////////////////////////////////////////////////
 // Visualize hair strands data as full screen view mode (BaseColor, Tangent, Roughness, ...)
@@ -505,7 +505,7 @@ void MainCS(uint3 DispatchThreadId : SV_DispatchThreadID)
 
 #if PERMUTATION_INPUT_TYPE == INPUT_ROOT_FRAME  
 	const uint RootIndex = DispatchThreadId.x;
-	if (RootIndex >= MaxRootCount) return;
+	if (RootIndex < MaxRootCount)
 	{
 		const uint TriangleIndex = RootToUniqueTriangleIndexBuffer[RootIndex];
 		const FHairMeshTriangleWithUVs RestTriangle		= GetTriangleWithUVsTransformation(TriangleIndex, RestPositionBuffer);
@@ -611,100 +611,6 @@ void MainCS(uint3 DispatchThreadId : SV_DispatchThreadID)
 }
 #endif // SHADER_MESH_PROJECTION_HAIR
 
-
-/////////////////////////////////////////////////////////////////////////////////////////
-// Visualize skin cache data, extracted & processed for hair strands deformation (Rest & Deformed)
-
-#if SHADER_MESH_PROJECTION_SKIN_CACHE
-
-#include "HairStrandsBindingCommon.ush"
-#include "../ShaderPrint.ush"
-
-float4x4 LocalToWorld;
-uint VertexOffset;
-uint IndexOffset;
-uint MaxIndexCount;
-uint MaxVertexCount;
-uint MeshUVsChannelOffset;
-uint MeshUVsChannelCount;
-uint bOutputInUVsSpace;
-uint MeshType;
-uint SectionIndex;
-float2 OutputResolution;
-
-Buffer<float> InputVertexPositionBuffer;
-Buffer<float2> InputVertexUVsBuffer;
-#if PERMUTATION_INPUT_TYPE == 1
-Buffer<uint> InputIndexBuffer;
-#endif
-
-#define REST_MESH 0
-#define DEFORMED_MESH 1
-#define SOURCE_MESH 2
-#define TARGET_MESH 3
-
-[numthreads(256, 1, 1)]
-void MainCS(uint3 DispatchThreadId : SV_DispatchThreadID)
-{
-	ResolvedView = ResolveView();
-
-	float4 Color = 0;
-	if (MeshType == REST_MESH)		Color = float4(1, 0, 0, 1);
-	if (MeshType == DEFORMED_MESH)	Color = float4(0, 1, 0, 1);
-	if (MeshType == SOURCE_MESH)	Color = float4(0, 0, 1, 1);
-	if (MeshType == TARGET_MESH)	Color = float4(1, 1, 0, 1);
-
-#if PERMUTATION_INPUT_TYPE == 0
-	const uint VertexId = DispatchThreadId.x * 3;
-	const uint VertexIndex0 = min(VertexOffset + VertexId, MaxVertexCount-1);
-
-	FHairMeshTriangleWithUVs Tri = (FHairMeshTriangleWithUVs)0;
-	Tri.P0 = InputVertexPositionBuffer[VertexIndex0 + 0];
-	Tri.P1 = InputVertexPositionBuffer[VertexIndex0 + 1];
-	Tri.P2 = InputVertexPositionBuffer[VertexIndex0 + 2];
-
-	Tri.UV0 = InputVertexUVsBuffer[VertexIndex0 + 0];
-	Tri.UV1 = InputVertexUVsBuffer[VertexIndex0 + 1];
-	Tri.UV2 = InputVertexUVsBuffer[VertexIndex0 + 2];
-
-	Tri.N = cross(Tri.P1 - Tri.P0, Tri.P2 - Tri.P0);
-#endif
-
-#if PERMUTATION_INPUT_TYPE == 1
-	const uint TriangleIndex = DispatchThreadId.x;
-	const FHairMeshTriangleWithUVs Tri = GetTriangleData(TriangleIndex, InputVertexPositionBuffer, InputVertexUVsBuffer, InputIndexBuffer, IndexOffset, MaxIndexCount, MaxVertexCount, MeshUVsChannelOffset, MeshUVsChannelCount);
-#endif
-
-	// Quick filter
-#if 0
-	//if (MeshType != REST_MESH)
-	//if (MeshType != DEFORMED_MESH)
-	//if (MeshType != SOURCE_MESH)
-	//if (MeshType != TARGET_MESH)
-	if (MeshType != SOURCE_MESH && MeshType != TARGET_MESH || SectionIndex != 0)
-	{
-		return;
-	}
-#endif
-
-	Color.xyz *= Tri.N * 0.5f + 0.5f;
-
-	if (bOutputInUVsSpace)
-	{
-		AddFilledTriangleWS(float3(Tri.UV0, 0), float3(Tri.UV1, 0), float3(Tri.UV2, 0), Color);
-	}
-	else
-	{
-		AddFilledTriangleWS(
-			mul(float4(Tri.P0, 1), LocalToWorld).xyz, 
-			mul(float4(Tri.P1, 1), LocalToWorld).xyz, 
-			mul(float4(Tri.P2, 1), LocalToWorld).xyz,
-			Color);
-	}
-}
-
-#endif // SHADER_MESH_PROJECTION_SKIN_CACHE
-
 /////////////////////////////////////////////////////////////////////////////////////////
 // Visualized stats for PPLL primary visibility
 
@@ -3100,4 +3006,144 @@ void MainDrawDebugAABBCS(uint3 DispatchThreadId : SV_DispatchThreadID)
 	}
 }
 
-#endif // SHADER_CLUSTERAABB
\ No newline at end of file
+#endif // SHADER_CLUSTERAABB
+
+/////////////////////////////////////////////////////////////////////////////////////////////////////////////////////////////
+/// Print skin cache stats
+#if SHADER_PRINT_HAIR_SKIN_CACHE
+
+uint			HairSkinCacheEnable;
+uint			GPUSkinCacheEnable;
+uint			InstanceCount;
+uint			UniqueMeshCount;
+Buffer<uint>	Infos;
+
+FSTRINGS(UniqueMeshNames)
+
+struct FInstanceInfo
+{
+	uint InstanceIndex;
+	uint GeometryType;
+	uint CacheType;
+	uint InstanceCount;
+	uint ActiveSectionCount;
+	uint TotalSectionCount;
+	uint LOD;
+	uint Pad;
+};
+
+FInstanceInfo ReadInstanceInfo(uint Index)
+{
+	const uint Index8 = Index * 8u;
+
+	uint4 Data0 = 0;
+	Data0.x = Infos[Index8 + 0];
+	Data0.y = Infos[Index8 + 1];
+	Data0.z = Infos[Index8 + 2];
+	Data0.w = Infos[Index8 + 3];
+
+	uint4 Data1 = 0;
+	Data1.x = Infos[Index8 + 4];
+	Data1.y = Infos[Index8 + 5];
+	Data1.z = Infos[Index8 + 6];
+	Data1.w = Infos[Index8 + 7];
+
+	FInstanceInfo Out = (FInstanceInfo)0;
+	Out.InstanceIndex		= Data0.x;
+	Out.GeometryType		= Data0.y;
+	Out.CacheType			= Data0.z;
+	Out.InstanceCount		= Data0.w;
+	Out.ActiveSectionCount	= Data1.x;
+	Out.TotalSectionCount	= Data1.y;
+	Out.LOD					= Data1.z;
+	Out.Pad					= Data1.w;
+	return Out;
+}
+
+[numthreads(1, 1, 1)]
+void MainCS(uint GroupIndex : SV_GroupIndex, uint3 DispatchThreadId : SV_DispatchThreadID)
+{
+	// Draw AABB of the instance groups
+	if (any(DispatchThreadId != 0))
+	{
+		return;
+	}
+
+	FShaderPrintContext Ctx = InitShaderPrintContext(true, uint2(50, 100));
+
+	// Common
+	bool bAdvanced = false;
+	{
+		Print(Ctx, TEXT("Instance count    : "), FontWhite);
+		Print(Ctx, InstanceCount, FontYellow, 8);
+		Newline(Ctx);
+	
+		Print(Ctx, TEXT("Unique mesh count : "), FontWhite);
+		Print(Ctx, UniqueMeshCount, FontYellow, 8);
+		Newline(Ctx);
+	
+		Print(Ctx, TEXT("GPU  skin cache   : "), FontWhite);
+		PrintBool(Ctx, GPUSkinCacheEnable);
+		Newline(Ctx);
+
+		Print(Ctx, TEXT("Hair skin cache   : "), FontWhite);
+		PrintBool(Ctx, HairSkinCacheEnable);
+		Newline(Ctx);
+
+		bAdvanced = AddCheckbox(Ctx, TEXT("Advanced"), false, FontWhite); 
+		Newline(Ctx);
+	}
+	Newline(Ctx);
+
+	// Draw instance details
+	Print(Ctx, TEXT("  Inst. Guide Strands Cards CacheType Sections LOD Name"));
+	Newline(Ctx);
+
+	LOOP
+	for (uint UniqueMeshIndex = 0; UniqueMeshIndex < UniqueMeshCount; ++UniqueMeshIndex)
+	{
+		const FInstanceInfo Info = ReadInstanceInfo(UniqueMeshIndex);
+		if (!bAdvanced && Info.ActiveSectionCount == 0)
+		{
+			continue;
+		}
+
+		// Index
+		Print(Ctx,UniqueMeshIndex, FontSilver, 2);
+		Print(Ctx, TEXT("  "));
+
+		// Instance count
+		Print(Ctx,Info.InstanceCount, FontYellow, 3, 0);
+		Print(Ctx, TEXT(" "));
+
+		// Geometry type
+		if (Info.GeometryType & 1) { Print(Ctx, TEXT("x     "), FontGreen); } else { Print(Ctx, TEXT(".     "), FontRed); }
+		if (Info.GeometryType & 2) { Print(Ctx, TEXT("x       "), FontGreen); } else { Print(Ctx, TEXT(".       "), FontRed); }
+		if (Info.GeometryType & 4) { Print(Ctx, TEXT("x     "), FontGreen); } else { Print(Ctx, TEXT(".     "), FontRed); }
+
+		// Cache type
+		     if (Info.CacheType == 1){ Print(Ctx, TEXT("Skin"), FontEmerald); }
+		else if (Info.CacheType == 2){ Print(Ctx, TEXT("Hair"), FontEmerald); }
+		else if (Info.CacheType == 4){ Print(Ctx, TEXT("Geom"), FontEmerald); }
+		else                         { Print(Ctx, TEXT("Unkw"), FontRed); }
+		Print(Ctx, TEXT("      "));
+
+
+		// Section
+		Print(Ctx,Info.ActiveSectionCount, FontYellow, 3, 0);
+		Print(Ctx, TEXT("/"));
+		Print(Ctx,Info.TotalSectionCount, FontYellow, 3, 0);
+		Print(Ctx, TEXT("    "));
+
+		// LOD
+		Print(Ctx,Info.LOD, FontYellow, 2, 0);
+
+		// Name
+		PrintUniqueMeshNames(Ctx, UniqueMeshIndex, FontWhite);
+
+		Newline(Ctx);
+	}
+	Newline(Ctx);
+}
+
+#endif // SHADER_PRINT_HAIR_SKIN_CACHE
\ No newline at end of file
diff --git a/Engine/Source/Runtime/Engine/Public/CachedGeometry.h b/Engine/Source/Runtime/Engine/Public/CachedGeometry.h
--- a/Engine/Source/Runtime/Engine/Public/CachedGeometry.h
+++ b/Engine/Source/Runtime/Engine/Public/CachedGeometry.h
@@ -40,9 +40,5 @@ struct FCachedGeometry
 
 	int32 LODIndex = -1;
 	TArray<Section> Sections;
-
-	FRDGBufferRef DeformedPositionBuffer = nullptr;
-	FRDGBufferRef DeformedPreviousPositionBuffer = nullptr;
-
 	FTransform LocalToWorld = FTransform::Identity;
 };
diff --git a/Engine/Source/Runtime/Renderer/Private/DeferredShadingRenderer.cpp b/Engine/Source/Runtime/Renderer/Private/DeferredShadingRenderer.cpp
--- a/Engine/Source/Runtime/Renderer/Private/DeferredShadingRenderer.cpp
+++ b/Engine/Source/Runtime/Renderer/Private/DeferredShadingRenderer.cpp
@@ -2077,7 +2077,8 @@ void FDeferredShadingSceneRenderer::Render(FRDGBuilder& GraphBuilder)
 	if (IsHairStrandsEnabled(EHairStrandsShaderType::All, Scene->GetShaderPlatform()) && RendererOutput == ERendererOutput::FinalSceneColor)
 	{
 		CreateHairStrandsBookmarkParameters(Scene, Views, AllFamilyViews, HairStrandsBookmarkParameters);
-		HairStrandsBookmarkParameters.TransientResources = AllocateHairTransientResourcse(GraphBuilder, Scene);
+		check(Scene->HairStrandsSceneData.TransientResources);
+		HairStrandsBookmarkParameters.TransientResources = Scene->HairStrandsSceneData.TransientResources;
 		RunHairStrandsBookmark(GraphBuilder, EHairStrandsBookmark::ProcessTasks, HairStrandsBookmarkParameters);
 
 		// Interpolation needs to happen after the skin cache run as there is a dependency 
diff --git a/Engine/Source/Runtime/Renderer/Private/HairStrands/HairStrandsRendering.cpp b/Engine/Source/Runtime/Renderer/Private/HairStrands/HairStrandsRendering.cpp
--- a/Engine/Source/Runtime/Renderer/Private/HairStrands/HairStrandsRendering.cpp
+++ b/Engine/Source/Runtime/Renderer/Private/HairStrands/HairStrandsRendering.cpp
@@ -85,7 +85,7 @@ void AddMeshDrawTransitionPass(
 	const FViewInfo& ViewInfo,
 	const FHairStrandsMacroGroupDatas& MacroGroupDatas);
 
-FHairTransientResources* AllocateHairTransientResourcse(FRDGBuilder& GraphBuilder, FScene* Scene)
+FHairTransientResources* AllocateHairTransientResources(FRDGBuilder& GraphBuilder, FScene* Scene)
 {
 	if (Scene->HairStrandsSceneData.TransientResources == nullptr)
 	{
diff --git a/Engine/Source/Runtime/Renderer/Private/HairStrands/HairStrandsRendering.h b/Engine/Source/Runtime/Renderer/Private/HairStrands/HairStrandsRendering.h
--- a/Engine/Source/Runtime/Renderer/Private/HairStrands/HairStrandsRendering.h
+++ b/Engine/Source/Runtime/Renderer/Private/HairStrands/HairStrandsRendering.h
@@ -21,7 +21,7 @@
 #include "HairStrandsInterface.h"
 #include "HairStrandsData.h"
 
-FHairTransientResources* AllocateHairTransientResourcse(FRDGBuilder& GraphBuilder, FScene* Scene);
+FHairTransientResources* AllocateHairTransientResources(FRDGBuilder& GraphBuilder, FScene* Scene);
 
 void RenderHairPrePass(
 	FRDGBuilder& GraphBuilder,
diff --git a/Engine/Source/Runtime/Renderer/Private/SceneVisibility.cpp b/Engine/Source/Runtime/Renderer/Private/SceneVisibility.cpp
--- a/Engine/Source/Runtime/Renderer/Private/SceneVisibility.cpp
+++ b/Engine/Source/Runtime/Renderer/Private/SceneVisibility.cpp
@@ -4436,6 +4436,7 @@ void FSceneRenderer::PreVisibilityFrameSetup(FRDGBuilder& GraphBuilder)
 		{
 			FHairStrandsBookmarkParameters Parameters; 
 			CreateHairStrandsBookmarkParameters(Scene, Views, AllFamilyViews, Parameters, false /*bComputeVisibleInstances*/);
+			Parameters.TransientResources = AllocateHairTransientResources(GraphBuilder, Scene);
 			if (Parameters.HasInstances())
 			{
 				// Prepare (skel.) data for guide/simulation update
diff --git a/Engine/Source/Runtime/Renderer/Public/HairStrandsInterface.h b/Engine/Source/Runtime/Renderer/Public/HairStrandsInterface.h
--- a/Engine/Source/Runtime/Renderer/Public/HairStrandsInterface.h
+++ b/Engine/Source/Runtime/Renderer/Public/HairStrandsInterface.h
@@ -17,6 +17,7 @@
 #include "ShaderPrintParameters.h"
 #include "GroomVisualizationData.h"
 #include "HairStrandsDefinitions.h"
+#include "CachedGeometry.h"
 
 class UTexture2D;
 class FSceneInterface;
@@ -439,6 +440,11 @@ struct FHairTransientResources
 	// Indirect dispatch args
 	FRDGBufferRef    IndirectDispatchArgsBuffer = nullptr;
 	FRDGBufferSRVRef IndirectDispatchArgsSRV = nullptr;
+
+	// Mesh data
+	TArray<FCachedGeometry> SimMeshDatas;
+	TArray<FCachedGeometry> RenMeshDatas;
+	const FCachedGeometry& GetMeshLODData(uint32 InRegisteredIndex, bool bSim) const { if (bSim) { check(SimMeshDatas.IsValidIndex(InRegisteredIndex)); return SimMeshDatas[InRegisteredIndex]; } else {check(RenMeshDatas.IsValidIndex(InRegisteredIndex)); return RenMeshDatas[InRegisteredIndex];} }
 };
 
 ////////////////////////////////////////////////////////////////////////////////////////////////////
