* Fix crash when running sky capture with hair.
* Fix incorect hair mesh RBF deformation when RBF is disabled

diff --git a/Engine/Plugins/Runtime/HairStrands/Source/HairStrandsCore/Private/GroomManager.cpp b/Engine/Plugins/Runtime/HairStrands/Source/HairStrandsCore/Private/GroomManager.cpp
--- a/Engine/Plugins/Runtime/HairStrands/Source/HairStrandsCore/Private/GroomManager.cpp
+++ b/Engine/Plugins/Runtime/HairStrands/Source/HairStrandsCore/Private/GroomManager.cpp
@@ -630,7 +630,7 @@ static FCachedGeometry GetCacheGeometryForHair(
 }
 
 // Binding surface parameters (skel.mesh/geom. cache)
-static void RunHairStrandsSkinCache(
+static void RunHairBindingSurfaceUpdate(
 	FRDGBuilder& GraphBuilder,
 	FSceneInterface* Scene,
 	const FSceneView* View,
@@ -1832,7 +1832,6 @@ static void RunHairStrandsInterpolation_Cards(
 				check(Instance->Meshes.IsValid(InstanceData.HairLODIndex));
 				InstanceData.bNeedDeformation 	= Instance->Meshes.LODs[InstanceData.HairLODIndex].DeformedResource != nullptr;
 				InstanceData.MeshInstance 		= &Instance->Meshes.LODs[InstanceData.HairLODIndex];
-				InstanceData.MeshLODIndex 		= Instance->Debug.MeshLODIndex; // Debug.MeshLODIndex is initialized during the guide pass (if guides are needed updated)
 				if (InstanceData.bNeedDeformation)
 				{
 					check(Instance->BindingType == EHairBindingType::Skinning);
@@ -2052,8 +2051,9 @@ static void RunHairStrandsInterpolation_Cards(
 	// Meshes only - Deform final mesh geometry (using RBF)
 	for (uint32 InstanceIndex : MeshInstances)
 	{
+		FInstanceData& InstanceData = InstanceDatas[InstanceIndex];
+		if (InstanceData.bNeedDeformation)
 		{
-			FInstanceData& InstanceData = InstanceDatas[InstanceIndex];
 			AddHairMeshesRBFInterpolationPass(
 				GraphBuilder,
 				ShaderMap,
@@ -2841,11 +2841,11 @@ void ProcessHairStrandsBookmark(
 				Parameters.AllViews);
 		}
 	}
-	else if (Bookmark == EHairStrandsBookmark::ProcessGuideInterpolation)
+	else if (Bookmark == EHairStrandsBookmark::ProcessBindingSurfaceUpdate)
 	{
 		check(GraphBuilder);
 
-		RunHairStrandsSkinCache(
+		RunHairBindingSurfaceUpdate(
 			*GraphBuilder,
 			Parameters.Scene,
 			Parameters.View,
@@ -2854,6 +2854,9 @@ void ProcessHairStrandsBookmark(
 			Parameters.ShaderPrintData,
 			*Parameters.TransientResources,
 			Parameters.ShaderMap);
+	}
+	else if (Bookmark == EHairStrandsBookmark::ProcessGuideInterpolation)
+	{
 
 		RunHairStrandsInterpolation_Guide(
 			*GraphBuilder,
diff --git a/Engine/Source/Runtime/Renderer/Private/SceneVisibility.cpp b/Engine/Source/Runtime/Renderer/Private/SceneVisibility.cpp
--- a/Engine/Source/Runtime/Renderer/Private/SceneVisibility.cpp
+++ b/Engine/Source/Runtime/Renderer/Private/SceneVisibility.cpp
@@ -4439,13 +4439,18 @@ void FSceneRenderer::PreVisibilityFrameSetup(FRDGBuilder& GraphBuilder)
 			Parameters.TransientResources = AllocateHairTransientResources(GraphBuilder, Scene);
 			if (Parameters.HasInstances())
 			{
-				// Prepare (skel.) data for guide/simulation update
-				// If we are rendering from scene capture we don't need to run another time the hair bookmarks.
-				if (Scene && IsHairStrandsEnabled(EHairStrandsShaderType::All, Scene->GetShaderPlatform()) && Views[0].AllowGPUParticleUpdate())
+				if (Scene && IsHairStrandsEnabled(EHairStrandsShaderType::All, Scene->GetShaderPlatform()))
 				{
+					// 1.Update binding surfaces
+					// Prepare surface data (MeshLODData)
 					Scene->WaitForGPUSkinCacheTask();
+					RunHairStrandsBookmark(GraphBuilder, EHairStrandsBookmark::ProcessBindingSurfaceUpdate, Parameters);
 
-					RunHairStrandsBookmark(GraphBuilder, EHairStrandsBookmark::ProcessGuideInterpolation, Parameters);
+					// 2. Prepare surface data for guides
+					if (Views[0].AllowGPUParticleUpdate())
+					{
+						RunHairStrandsBookmark(GraphBuilder, EHairStrandsBookmark::ProcessGuideInterpolation, Parameters);
+					}
 				}
 			}
 		}
diff --git a/Engine/Source/Runtime/Renderer/Public/HairStrandsInterface.h b/Engine/Source/Runtime/Renderer/Public/HairStrandsInterface.h
--- a/Engine/Source/Runtime/Renderer/Public/HairStrandsInterface.h
+++ b/Engine/Source/Runtime/Renderer/Public/HairStrandsInterface.h
@@ -493,6 +493,7 @@ enum class EHairStrandsBookmark : uint8
 {
 	ProcessTasks,
 	ProcessLODSelection,
+	ProcessBindingSurfaceUpdate,
 	ProcessGuideInterpolation,
 	ProcessCardsAndMeshesInterpolation_PrimaryView,
 	ProcessCardsAndMeshesInterpolation_ShadowView,
