Cook hair strands if they're supported by any shader formats used on the target cook platform.

The previous behaviour was to cook them only if they were supported by *all* shader formats on the target platform. This meant that if you had ES31 shaders as an option on Windows, strands would be cooked out.

diff --git a/Engine/Plugins/Runtime/HairStrands/Source/HairStrandsCore/Private/GroomAsset.cpp b/Engine/Plugins/Runtime/HairStrands/Source/HairStrandsCore/Private/GroomAsset.cpp
--- a/Engine/Plugins/Runtime/HairStrands/Source/HairStrandsCore/Private/GroomAsset.cpp
+++ b/Engine/Plugins/Runtime/HairStrands/Source/HairStrandsCore/Private/GroomAsset.cpp
@@ -579,13 +579,18 @@ uint8 UGroomAsset::GenerateClassStripFlags(FArchive& Ar)
 	bool bIsMeshesStripped = false;
 	if (bIsCook)
 	{
-		// Determine if strands are supported on the target cook platform
+		// Determine if strands are supported on any shader formats used by the target cook platform
 		TArray<FName> ShaderFormats;
 		CookTarget->GetAllTargetedShaderFormats(ShaderFormats);
+		bIsStrandsSupportedOnTargetPlatform = false;
 		for (int32 FormatIndex = 0; FormatIndex < ShaderFormats.Num(); ++FormatIndex)
 		{
 			const EShaderPlatform ShaderPlatform = ShaderFormatToLegacyShaderPlatform(ShaderFormats[FormatIndex]);
-			bIsStrandsSupportedOnTargetPlatform &= IsHairStrandsSupported(EHairStrandsShaderType::Strands, ShaderPlatform);
+			if (IsHairStrandsSupported(EHairStrandsShaderType::Strands, ShaderPlatform))
+			{
+				bIsStrandsSupportedOnTargetPlatform = true;
+				break;
+			}
 		}
 
 		// Determine the platform min LOD if the stripping hasn't been disabled
