Fix uninitialized Groom deformer OutputMask value causing the output position to not being updated consistently.

diff --git a/Engine/Plugins/Runtime/HairStrands/Source/HairStrandsCore/Private/Deformer/DeformerDataInterfaceGroomWrite.cpp b/Engine/Plugins/Runtime/HairStrands/Source/HairStrandsCore/Private/Deformer/DeformerDataInterfaceGroomWrite.cpp
--- a/Engine/Plugins/Runtime/HairStrands/Source/HairStrandsCore/Private/Deformer/DeformerDataInterfaceGroomWrite.cpp
+++ b/Engine/Plugins/Runtime/HairStrands/Source/HairStrandsCore/Private/Deformer/DeformerDataInterfaceGroomWrite.cpp
@@ -166,6 +166,7 @@ FComputeDataProviderRenderProxy* UOptimusGroomWriteDataProvider::GetRenderProxy(
 
 FOptimusGroomWriteDataProviderProxy::FOptimusGroomWriteDataProviderProxy(UGroomComponent* InGroomComponent, uint64 InOutputMask)
 {
+	OutputMask = InOutputMask;
 	const uint32 InstanceCount = InGroomComponent ? InGroomComponent->GetGroupCount() : 0;
 	for (uint32 Index = 0; Index < InstanceCount; ++Index)
 	{
diff --git a/Engine/Plugins/Runtime/HairStrands/Source/HairStrandsCore/Private/Deformer/DeformerDataInterfaceGroomWrite.h b/Engine/Plugins/Runtime/HairStrands/Source/HairStrandsCore/Private/Deformer/DeformerDataInterfaceGroomWrite.h
--- a/Engine/Plugins/Runtime/HairStrands/Source/HairStrandsCore/Private/Deformer/DeformerDataInterfaceGroomWrite.h
+++ b/Engine/Plugins/Runtime/HairStrands/Source/HairStrandsCore/Private/Deformer/DeformerDataInterfaceGroomWrite.h
@@ -52,7 +52,7 @@ public:
 	UPROPERTY(EditAnywhere, BlueprintReadWrite, Category = Binding)
 	TObjectPtr<UGroomComponent> GroomComponent = nullptr;
 
-	uint64 OutputMask;
+	uint64 OutputMask = 0;
 
 	//~ Begin UComputeDataProvider Interface
 	FComputeDataProviderRenderProxy* GetRenderProxy() override;
@@ -74,7 +74,7 @@ private:
 	using FParameters = FGroomWriteDataInterfaceParameters;
 
 	TArray<FHairGroupInstance*> Instances;
-	uint64 OutputMask;
+	uint64 OutputMask = 0;
 
 	struct FResources
 	{
