Fix missing groom on skel. when manual skin cache is required and GDME is enabled.

diff --git a/Engine/Plugins/Runtime/HairStrands/Source/HairStrandsCore/Private/GroomManager.cpp b/Engine/Plugins/Runtime/HairStrands/Source/HairStrandsCore/Private/GroomManager.cpp
--- a/Engine/Plugins/Runtime/HairStrands/Source/HairStrandsCore/Private/GroomManager.cpp
+++ b/Engine/Plugins/Runtime/HairStrands/Source/HairStrandsCore/Private/GroomManager.cpp
@@ -157,14 +157,6 @@ static bool IsInstanceFrustumCullingEnable()
 ///////////////////////////////////////////////////////////////////////////////////////////////////
 // Hair manual skin cache
 
-static bool IsSkeletalMeshEvaluationEnabled()
-{
-	// When deferred skel. mesh update is enabled, hair strands skeletal mesh deformation is not allowed, as skin-cached update happen after 
-	// the PreInitView calls, which causes hair LOD selection and hair simulation to have invalid value & resources
-	static const auto CVarSkelMeshGDME = IConsoleManager::Get().FindTConsoleVariableDataInt(TEXT("r.DeferSkeletalDynamicDataUpdateUntilGDME"));
-	return CVarSkelMeshGDME ? CVarSkelMeshGDME->GetValueOnAnyThread() == 0 : true;
-}
-
 // Retrive the skel. mesh scene info
 static const FPrimitiveSceneInfo* GetMeshSceneInfo(FSceneInterface* Scene, FHairGroupInstance* Instance)
 {
@@ -555,7 +547,7 @@ static int32 GetMeshLODIndex(
 	int32 OutLODIndex = -1;
 	if (const FPrimitiveSceneInfo* PrimitiveSceneInfo = GetMeshSceneInfo(Scene, Instance))
 	{
-		if (Instance->Debug.GroomBindingType == EGroomBindingMeshType::SkeletalMesh && IsSkeletalMeshEvaluationEnabled())
+		if (Instance->Debug.GroomBindingType == EGroomBindingMeshType::SkeletalMesh)
 		{
 			if (const FSkeletalMeshSceneProxy* SceneProxy = static_cast<const FSkeletalMeshSceneProxy*>(PrimitiveSceneInfo->Proxy))
 			{
@@ -597,7 +589,7 @@ static FCachedGeometry GetCacheGeometryForHair(
 	FCachedGeometry Out;
 	if (const FPrimitiveSceneInfo* PrimitiveSceneInfo = GetMeshSceneInfo(Scene, Instance))
 	{
-		if (Instance->Debug.GroomBindingType == EGroomBindingMeshType::SkeletalMesh && IsSkeletalMeshEvaluationEnabled())
+		if (Instance->Debug.GroomBindingType == EGroomBindingMeshType::SkeletalMesh)
 		{
 			if (const FSkeletalMeshSceneProxy* SceneProxy = static_cast<const FSkeletalMeshSceneProxy*>(PrimitiveSceneInfo->Proxy))
 			{
