Fix mesh cards panel not displaying the texture layout details

diff --git a/Engine/Plugins/Runtime/HairStrands/Source/HairStrandsEditor/Private/GroomAssetDetails.cpp b/Engine/Plugins/Runtime/HairStrands/Source/HairStrandsEditor/Private/GroomAssetDetails.cpp
--- a/Engine/Plugins/Runtime/HairStrands/Source/HairStrandsEditor/Private/GroomAssetDetails.cpp
+++ b/Engine/Plugins/Runtime/HairStrands/Source/HairStrandsEditor/Private/GroomAssetDetails.cpp
@@ -1596,7 +1596,32 @@ void FGroomRenderingDetails::OnGenerateElementForHairGroup(TSharedRef<IPropertyH
 		break;
 		case EMaterialPanelType::Meshes:
 		{
-			AddPropertyWithCustomReset(ChildHandle, ChildrenBuilder, GroupIndex, -1);
+			if (PropertyName == GET_MEMBER_NAME_CHECKED(FHairGroupsMeshesSourceDescription, Textures))
+			{
+				EHairTextureLayout LayoutType = EHairTextureLayout::Layout0;
+				if (GroomAsset != nullptr && GroomAsset->GetHairGroupsMeshes().IsValidIndex(GroupIndex))
+				{
+					LayoutType = GroomAsset->GetHairGroupsMeshes()[GroupIndex].Textures.Layout;
+				}
+
+				IDetailGroup& TextureGroup = ChildrenBuilder.AddGroup(TEXT("HairMeshesTextures"), LOCTEXT("HairMeshesTextures", "Textures"));
+				{
+					// Layout type
+					TextureGroup.AddPropertyRow(ChildHandle->GetChildHandle(0).ToSharedRef());
+
+					// Textures
+					const uint32 TextureCount = GroomAsset->GetHairGroupsMeshes()[GroupIndex].Textures.Textures.Num();
+					TSharedPtr<IPropertyHandle> TextureArrayHandle = ChildHandle->GetChildHandle(1);
+					for (uint32 TextureIt = 0; TextureIt < TextureCount; ++TextureIt)
+					{
+						TextureGroup.AddPropertyRow(TextureArrayHandle->GetChildHandle(TextureIt).ToSharedRef()).DisplayName(FTextStringHelper::CreateFromBuffer(GetHairTextureLayoutTextureName(LayoutType, TextureIt, true /*bDetail*/)));
+					}
+				}
+			}
+			else
+			{
+				AddPropertyWithCustomReset(ChildHandle, ChildrenBuilder, GroupIndex, -1);
+			}
 		}
 		break;
 		case EMaterialPanelType::Interpolation:
