Fix invalid assert when switching groom texture layout.

diff --git a/Engine/Plugins/Runtime/HairStrands/Source/HairStrandsCore/Private/GroomAssetCards.cpp b/Engine/Plugins/Runtime/HairStrands/Source/HairStrandsCore/Private/GroomAssetCards.cpp
--- a/Engine/Plugins/Runtime/HairStrands/Source/HairStrandsCore/Private/GroomAssetCards.cpp
+++ b/Engine/Plugins/Runtime/HairStrands/Source/HairStrandsCore/Private/GroomAssetCards.cpp
@@ -93,7 +93,7 @@ uint32 GetHairTextureLayoutTextureCount(EHairTextureLayout In)
 	{
 		case EHairTextureLayout::Layout0: OutCount = 6u; break;
 		case EHairTextureLayout::Layout1: OutCount = 6u; break;
-		case EHairTextureLayout::Layout2: OutCount = 4u; break;
+		case EHairTextureLayout::Layout2: OutCount = 3u; break;
 		case EHairTextureLayout::Layout3: OutCount = 4u; break;
 	}
 	check(OutCount <= HAIR_CARDS_MAX_TEXTURE_COUNT);
