Fix incorrect check.

diff --git a/Engine/Plugins/Runtime/HairStrands/Source/HairStrandsCore/Private/GroomManager.cpp b/Engine/Plugins/Runtime/HairStrands/Source/HairStrandsCore/Private/GroomManager.cpp
--- a/Engine/Plugins/Runtime/HairStrands/Source/HairStrandsCore/Private/GroomManager.cpp
+++ b/Engine/Plugins/Runtime/HairStrands/Source/HairStrandsCore/Private/GroomManager.cpp
@@ -1816,7 +1816,7 @@ static void RunHairStrandsInterpolation_Cards(
 				{
 					check(InstanceData.CardInstance->Guides.IsValid());
 					check(InstanceData.CardInstance->Guides.HasValidRootData());
-					check(InstanceData.CardInstance->Guides.DeformedRootResource->IsValid(InstanceData.MeshLODIndex));
+					check(InstanceData.MeshLODIndex == -1 || InstanceData.CardInstance->Guides.DeformedRootResource->IsValid(InstanceData.MeshLODIndex));
 				}
 			}
 			else if (Instance->GeometryType == EHairGeometryType::Meshes)
