Fix holdout for hair cards/meshes.

diff --git a/Engine/Plugins/Runtime/HairStrands/Source/HairStrandsCore/Private/GroomComponent.cpp b/Engine/Plugins/Runtime/HairStrands/Source/HairStrandsCore/Private/GroomComponent.cpp
--- a/Engine/Plugins/Runtime/HairStrands/Source/HairStrandsCore/Private/GroomComponent.cpp
+++ b/Engine/Plugins/Runtime/HairStrands/Source/HairStrandsCore/Private/GroomComponent.cpp
@@ -41,6 +41,7 @@
 #include "PrimitiveSceneInfo.h"
 #include "PSOPrecache.h"
 #include "SceneInterface.h"
+#include "PrimitiveUniformShaderParametersBuilder.h"
 
 #if WITH_EDITORONLY_DATA
 #include "DerivedDataCache.h"
@@ -983,11 +984,29 @@ public:
 			PreviousLocalToWorld = CurrentLocalToWorld;
 		}
 
-		FDynamicPrimitiveUniformBuffer& DynamicPrimitiveUniformBuffer = Collector.AllocateOneFrameResource<FDynamicPrimitiveUniformBuffer>();
-		DynamicPrimitiveUniformBuffer.Set(Collector.GetRHICommandList(), CurrentLocalToWorld, PreviousLocalToWorld, GetBounds(), GetLocalBounds(), true, false, bOutputVelocity);
-		BatchElement.PrimitiveUniformBufferResource = &DynamicPrimitiveUniformBuffer.UniformBuffer; // automatic copy to the gpu scene buffer
-		//primtiveid is set to 0
+		// Update primitive uniform buffer
+		{
+			// Use default SceneProxy builder values
+			FPrimitiveUniformShaderParametersBuilder Builder;
+			BuildUniformShaderParameters(Builder);
+
+			// Override transforms
+			Builder
+				.LocalToWorld(CurrentLocalToWorld)
+				.PreviousLocalToWorld(PreviousLocalToWorld)
+				.OutputVelocity(bOutputVelocity)
+				.UseVolumetricLightmap(false);
+
+			// Create primitive uniform buffer
+			FRHICommandListBase& RHICmdList = Collector.GetRHICommandList();
+			FDynamicPrimitiveUniformBuffer& DynamicPrimitiveUniformBuffer = Collector.AllocateOneFrameResource<FDynamicPrimitiveUniformBuffer>();
+			DynamicPrimitiveUniformBuffer.UniformBuffer.BufferUsage = UniformBuffer_SingleFrame;
+			DynamicPrimitiveUniformBuffer.UniformBuffer.SetContents(RHICmdList, Builder.Build());
+			DynamicPrimitiveUniformBuffer.UniformBuffer.InitResource(RHICmdList);
+			BatchElement.PrimitiveUniformBufferResource = &DynamicPrimitiveUniformBuffer.UniformBuffer; // automatic copy to the gpu scene buffer
+		}
 
+		//primtiveid is set to 0
 		BatchElement.FirstIndex = 0;
 		BatchElement.NumInstances = 1;
 		BatchElement.PrimitiveIdMode = GeometryType == EHairGeometryType::Strands ? PrimID_ForceZero : PrimID_DynamicPrimitiveShaderData;
