Split Groom LOD selection.

diff --git a/Engine/Plugins/Runtime/HairStrands/Source/HairStrandsCore/Private/GroomManager.cpp b/Engine/Plugins/Runtime/HairStrands/Source/HairStrandsCore/Private/GroomManager.cpp
--- a/Engine/Plugins/Runtime/HairStrands/Source/HairStrandsCore/Private/GroomManager.cpp
+++ b/Engine/Plugins/Runtime/HairStrands/Source/HairStrandsCore/Private/GroomManager.cpp
@@ -2347,6 +2347,320 @@ static uint32 ComputeActiveCurveCount(float InScreenSize, uint32 InCurveCount, u
 	return FMath::Clamp(OutCurveCount, 1, InCurveCount);
 }
 
+struct FHairLOD
+{
+	float HairLODIndex = -1.f;
+	float MinHairLOD = 0.f;
+	uint32 HairLODCount = 0;
+
+	float LODPredictedIndex = -1.f;
+	float DebugScreenSize = 0.f;
+
+	uint32 ContinuousLODPointCount = 0;
+	uint32 ContinuousLODCurveCount = 0;
+	float ContinuousLODScreenSize = 1.f;
+	FVector2f ContinuousLODScreenPos = FVector2f(0,0);
+	FBoxSphereBounds ContinuousLODBounds;
+	float ContinuousLODCoverageScale = 1.f;
+};
+
+static FHairLOD ComputeHairLODIndex(const FHairGroupInstance* Instance, const TArray<const FSceneView*>& Views)
+{
+	check(Instance);
+
+	FHairLOD Out;
+
+	// Insure that MinLOD is necessary taken into account if a force LOD is request (i.e., LODIndex>=0). If a Force LOD 
+	// is not requested (i.e., LODIndex<0), the MinLOD is applied after ViewLODIndex has been determined in the codeblock below
+	Out.HairLODCount = Instance->HairGroupPublicData->GetLODVisibilities().Num();
+	Out.MinHairLOD = FMath::Max(0, GHairStrandsMinLOD);
+	Out.HairLODIndex = IsHairVisibilityComputeRasterContinuousLODEnabled() ? 0.0 : (Instance->Debug.LODForcedIndex >= 0 ? FMath::Max(Instance->Debug.LODForcedIndex, Out.MinHairLOD) : -1.0f);
+
+	FVector2f MaxContinuousLODScreenPos = FVector2f(0.f, 0.f);
+	float LODViewIndex = -1;
+	float MaxScreenSize_RestBound = 0.f;
+	float MaxScreenSize_Bound = 0.f;
+	const FSphere SphereBound = Instance->GetBounds().GetSphere();
+	for (const FSceneView* View : Views)
+	{
+		const FVector3d BoundScale = Instance->LocalToWorld.GetScale3D();
+		const FVector3f BoundExtent = Instance->Strands.Data ? FVector3f(Instance->Strands.Data->Header.BoundingBox.GetExtent()) : FVector3f(0,0,0);
+		const float BoundRadius = FMath::Max3(BoundExtent.X, BoundExtent.Y, BoundExtent.Z) * FMath::Max3(BoundScale.X, BoundScale.Y, BoundScale.Z);
+
+		const float ScreenSize_RestBound = FMath::Clamp(ComputeBoundsScreenSize(FVector4(SphereBound.Center, 1), BoundRadius, *View), 0.f, 1.0f);
+		const float ScreenSize_Bound = FMath::Clamp(ComputeBoundsScreenSize(FVector4(SphereBound.Center, 1), SphereBound.W, *View), 0.f, 1.0f);
+
+		const float CurrLODViewIndex = FMath::Max(Out.MinHairLOD, GetHairInstanceLODIndex(Instance->HairGroupPublicData->GetLODScreenSizes(), ScreenSize_Bound, Instance->Strands.Modifier.LODBias));
+
+		MaxScreenSize_Bound = FMath::Max(MaxScreenSize_Bound, ScreenSize_Bound);
+		MaxScreenSize_RestBound = FMath::Max(MaxScreenSize_RestBound, ScreenSize_RestBound);
+		MaxContinuousLODScreenPos = ComputeProjectedScreenPos(SphereBound.Center, *View);
+
+		// Select highest LOD across all views
+		LODViewIndex = LODViewIndex < 0 ? CurrLODViewIndex : FMath::Min(LODViewIndex, CurrLODViewIndex);
+	}
+
+	if (Out.HairLODIndex < 0)
+	{
+		Out.HairLODIndex = LODViewIndex;
+	}
+	Out.HairLODIndex = FMath::Clamp(Out.HairLODIndex, 0.f, float(Out.HairLODCount - 1));
+
+	// Feedback game thread with LOD selection 
+	Out.LODPredictedIndex = LODViewIndex;
+	Out.DebugScreenSize = MaxScreenSize_Bound;
+	Out.ContinuousLODPointCount = Instance->HairGroupPublicData->RestPointCount;
+	Out.ContinuousLODCurveCount = Instance->HairGroupPublicData->RestCurveCount;
+	Out.ContinuousLODScreenSize = MaxScreenSize_RestBound;
+	Out.ContinuousLODScreenPos = MaxContinuousLODScreenPos;
+	Out.ContinuousLODBounds = SphereBound;
+	Out.ContinuousLODCoverageScale = 1.f;
+
+	// Auto LOD
+	if (Instance->Strands.ClusterResource && Instance->Strands.Data)
+	{
+		uint32 EffectiveCurveCount = 0;
+		if (Instance->HairGroupPublicData->bAutoLOD || IsHairStrandsForceAutoLODEnabled())
+		{
+			EffectiveCurveCount = ComputeActiveCurveCount(Out.ContinuousLODScreenSize, Instance->HairGroupPublicData->RestCurveCount, Instance->HairGroupPublicData->ClusterCount);
+			Out.HairLODIndex = 0;
+		}
+		else
+		{
+			EffectiveCurveCount = Instance->Strands.ClusterResource->BulkData.GetCurveCount(Out.HairLODIndex);
+		}
+		check(EffectiveCurveCount <= uint32(Instance->Strands.Data->Header.CurveToPointCount.Num()));
+
+		Out.ContinuousLODCurveCount = EffectiveCurveCount;
+		Out.ContinuousLODPointCount = EffectiveCurveCount > 0 ? Instance->Strands.Data->Header.CurveToPointCount[EffectiveCurveCount - 1] : 0;
+		Out.ContinuousLODCoverageScale = ComputeActiveCurveCoverageScale(EffectiveCurveCount, Instance->HairGroupPublicData->RestCurveCount);
+	}
+
+	return Out;
+}
+
+static void ApplyHairLOD(const FHairLOD& In, FHairGroupInstance* OutInstance)
+{
+	// Feedback game thread with LOD selection 
+	OutInstance->Debug.LODPredictedIndex = In.LODPredictedIndex;
+	OutInstance->HairGroupPublicData->DebugScreenSize = In.DebugScreenSize;
+	OutInstance->HairGroupPublicData->ContinuousLODPointCount = In.ContinuousLODPointCount;
+	OutInstance->HairGroupPublicData->ContinuousLODCurveCount = In.ContinuousLODCurveCount;
+	OutInstance->HairGroupPublicData->ContinuousLODScreenSize = In.ContinuousLODScreenSize;
+	OutInstance->HairGroupPublicData->ContinuousLODScreenPos = In.ContinuousLODScreenPos;
+	OutInstance->HairGroupPublicData->ContinuousLODBounds = In.ContinuousLODBounds;
+	OutInstance->HairGroupPublicData->ContinuousLODCoverageScale = In.ContinuousLODCoverageScale;
+}
+
+// Function for selecting, loading, & initializing LOD resources
+static bool SelectValidLOD(
+	FRDGBuilder& GraphBuilder, 
+	FGlobalShaderMap* ShaderMap, 
+	EShaderPlatform ShaderPlatform, 
+	FHairGroupInstance* Instance, 
+	float HairLODIndex,
+	float PrevHairLODIndex, 
+	float HairLODCount, 
+	float MeshLODIndex, 
+	float PrevMeshLODIndex, 
+	bool bHasPathTracingView)
+{
+	const int32 IntHairLODIndex = FMath::Clamp(FMath::FloorToInt(HairLODIndex), 0, HairLODCount - 1);
+	const TArray<EHairGeometryType>& LODGeometryTypes = Instance->HairGroupPublicData->GetLODGeometryTypes();
+	const TArray<bool>& LODVisibilities = Instance->HairGroupPublicData->GetLODVisibilities();
+	const bool bIsVisible = LODVisibilities[IntHairLODIndex];
+	const bool bForceCards = GHairStrands_UseCards > 0 || Instance->bForceCards; // todo
+	EHairGeometryType GeometryType = ConvertLODGeometryType(LODGeometryTypes[IntHairLODIndex], bForceCards, ShaderPlatform);
+
+	// Skip cluster culling if strands have a single LOD (i.e.: LOD0), and the instance is static (i.e., no skinning, no simulation, ...)
+	bool bCullingEnable = false;
+	if (GeometryType == EHairGeometryType::Meshes)
+	{
+		if (!Instance->Meshes.IsValid(IntHairLODIndex))
+		{
+			GeometryType = EHairGeometryType::NoneGeometry;
+		}
+	}
+	else if (GeometryType == EHairGeometryType::Cards)
+	{
+		if (!Instance->Cards.IsValid(IntHairLODIndex))
+		{
+			GeometryType = EHairGeometryType::NoneGeometry;
+		}
+	}
+	else if (GeometryType == EHairGeometryType::Strands)
+	{
+		if (!Instance->Strands.IsValid())
+		{
+			GeometryType = EHairGeometryType::NoneGeometry;
+		}
+		else
+		{
+			const bool bNeedLODing		= HairLODIndex > 0;
+			const bool bNeedDeformation = Instance->Strands.DeformedResource != nullptr;
+			bCullingEnable = (bNeedLODing || bNeedDeformation);
+		}
+	}
+
+	if (!bIsVisible)
+	{
+		GeometryType = EHairGeometryType::NoneGeometry;
+	}
+
+	// Transform ContinuousLOD CurveCount/PointCount into streaming request
+	uint32 RequestedCurveCount = Instance->HairGroupPublicData->RestCurveCount;
+	uint32 RequestedPointCount = Instance->HairGroupPublicData->RestPointCount;
+	const bool bStreamingEnabled = GHairStrands_Streaming > 0;
+	if (GeometryType == EHairGeometryType::Strands && bStreamingEnabled && Instance->bSupportStreaming)
+	{
+		// Round Curve/Point request to curve 'page'
+		RequestedCurveCount = GetRoundedCurveCount(Instance->HairGroupPublicData->ContinuousLODCurveCount, Instance->HairGroupPublicData->RestCurveCount);
+		RequestedPointCount = RequestedCurveCount > 0 ? Instance->Strands.Data->Header.CurveToPointCount[RequestedCurveCount - 1] : 0;
+	}
+
+	const bool bSimulationEnable			= Instance->HairGroupPublicData->IsSimulationEnable(IntHairLODIndex);
+	const bool bDeformationEnable			= Instance->HairGroupPublicData->bIsDeformationEnable;
+	const bool bGlobalInterpolationEnable	= Instance->HairGroupPublicData->IsGlobalInterpolationEnable(IntHairLODIndex);
+	const bool bSimulationCacheEnable		= Instance->HairGroupPublicData->bIsSimulationCacheEnable;
+	const bool bLODNeedsGuides				= bSimulationEnable || bDeformationEnable || bGlobalInterpolationEnable || bSimulationCacheEnable;
+
+	const EHairResourceLoadingType LoadingType = GetHairResourceLoadingType(GeometryType, IntHairLODIndex);
+	EHairResourceStatus ResourceStatus;
+	ResourceStatus.Status 				= EHairResourceStatus::EStatus::None;
+	ResourceStatus.AvailableCurveCount 	= Instance->HairGroupPublicData->RestCurveCount;
+
+	// Lazy allocation of resources
+	// Note: Allocation will only be done if the resources is not initialized yet. Guides deformed position are also initialized from the Rest position at creation time.
+	if (Instance->Guides.Data && bLODNeedsGuides)
+	{
+		if (Instance->Guides.RestRootResource)			{ Instance->Guides.RestRootResource->Allocate(GraphBuilder, LoadingType, ResourceStatus, MeshLODIndex); }
+		if (Instance->Guides.RestResource)				{ Instance->Guides.RestResource->Allocate(GraphBuilder, LoadingType, ResourceStatus); }
+		if (Instance->Guides.DeformedRootResource)		{ Instance->Guides.DeformedRootResource->Allocate(GraphBuilder, LoadingType, ResourceStatus, MeshLODIndex); }
+		if (Instance->Guides.DeformedResource)			
+		{ 
+			// Ensure the rest resources are correctly loaded prior to initialized and copy the rest position into the deformed positions
+			if (Instance->Guides.RestResource->bIsInitialized)
+			{
+				const bool bNeedCopy = !Instance->Guides.DeformedResource->bIsInitialized; 
+				Instance->Guides.DeformedResource->Allocate(GraphBuilder, EHairResourceLoadingType::Sync); 
+				if (bNeedCopy) 
+				{ 
+					AddCopyHairStrandsPositionPass(GraphBuilder, ShaderMap, *Instance->Guides.RestResource, *Instance->Guides.DeformedResource); 
+				}
+			}
+		}
+	}
+
+	if (GeometryType == EHairGeometryType::Meshes)
+	{
+		FHairGroupInstance::FMeshes::FLOD& InstanceLOD = Instance->Meshes.LODs[IntHairLODIndex];
+
+		if (InstanceLOD.DeformedResource)				{ InstanceLOD.DeformedResource->Allocate(GraphBuilder, LoadingType, ResourceStatus); }
+#if RHI_RAYTRACING
+		if (InstanceLOD.RaytracingResource)				{ InstanceLOD.RaytracingResource->Allocate(GraphBuilder, LoadingType, ResourceStatus);}
+#endif
+
+		InstanceLOD.InitVertexFactory();
+	}
+	else if (GeometryType == EHairGeometryType::Cards)
+	{
+		FHairGroupInstance::FCards::FLOD& InstanceLOD = Instance->Cards.LODs[IntHairLODIndex];
+
+		if (InstanceLOD.InterpolationResource)			{ InstanceLOD.InterpolationResource->Allocate(GraphBuilder, LoadingType, ResourceStatus); }
+		if (InstanceLOD.DeformedResource)				{ InstanceLOD.DeformedResource->Allocate(GraphBuilder, LoadingType, ResourceStatus); }
+		if (InstanceLOD.Guides.RestRootResource)		{ InstanceLOD.Guides.RestRootResource->Allocate(GraphBuilder, LoadingType, ResourceStatus, MeshLODIndex); }
+		if (InstanceLOD.Guides.RestResource)			{ InstanceLOD.Guides.RestResource->Allocate(GraphBuilder, LoadingType, ResourceStatus); }
+		if (InstanceLOD.Guides.DeformedRootResource)	{ InstanceLOD.Guides.DeformedRootResource->Allocate(GraphBuilder, LoadingType, ResourceStatus, MeshLODIndex); }
+		if (InstanceLOD.Guides.DeformedResource)		{ InstanceLOD.Guides.DeformedResource->Allocate(GraphBuilder, LoadingType, ResourceStatus); }
+		if (InstanceLOD.Guides.InterpolationResource)	{ InstanceLOD.Guides.InterpolationResource->Allocate(GraphBuilder, LoadingType, ResourceStatus); }
+#if RHI_RAYTRACING
+		if (InstanceLOD.RaytracingResource)				{ InstanceLOD.RaytracingResource->Allocate(GraphBuilder, LoadingType, ResourceStatus);}
+#endif
+		InstanceLOD.InitVertexFactory();
+	}
+	else if (GeometryType == EHairGeometryType::Strands)
+	{
+		check(Instance->HairGroupPublicData);
+
+		if (Instance->Strands.RestRootResource)			{ Instance->Strands.RestRootResource->Allocate(GraphBuilder, EHairResourceLoadingType::Async, ResourceStatus, RequestedCurveCount, RequestedPointCount, MeshLODIndex); }
+		if (Instance->Strands.RestResource)				{ Instance->Strands.RestResource->Allocate(GraphBuilder, EHairResourceLoadingType::Async, ResourceStatus, RequestedCurveCount, RequestedPointCount); }
+		if (Instance->Strands.ClusterResource)			{ Instance->Strands.ClusterResource->Allocate(GraphBuilder, EHairResourceLoadingType::Async, ResourceStatus, RequestedCurveCount, RequestedPointCount); }
+		if (Instance->Strands.InterpolationResource)	{ Instance->Strands.InterpolationResource->Allocate(GraphBuilder, EHairResourceLoadingType::Async, ResourceStatus, RequestedCurveCount, RequestedPointCount); }
+		if (Instance->Strands.CullingResource)			{ Instance->Strands.CullingResource->Allocate(GraphBuilder, EHairResourceLoadingType::Async, ResourceStatus, RequestedCurveCount, RequestedPointCount); }
+
+		if (Instance->Strands.DeformedRootResource)		{ Instance->Strands.DeformedRootResource->Allocate(GraphBuilder, EHairResourceLoadingType::Async, ResourceStatus, RequestedCurveCount, RequestedPointCount, MeshLODIndex); }
+		if (Instance->Strands.DeformedResource)			{ Instance->Strands.DeformedResource->Allocate(GraphBuilder, EHairResourceLoadingType::Async, ResourceStatus, RequestedCurveCount, RequestedPointCount); }
+#if RHI_RAYTRACING
+		if (bHasPathTracingView)						{ AllocateRaytracingResources(Instance); }
+		if (Instance->Strands.RenRaytracingResource)	{ Instance->Strands.RenRaytracingResource->Allocate(GraphBuilder, EHairResourceLoadingType::Async, ResourceStatus, RequestedCurveCount, RequestedPointCount); }
+#endif
+		Instance->Strands.VertexFactory->InitResources(GraphBuilder.RHICmdList);
+
+		// Early initialization, so that when filtering MeshBatch block in SceneVisiblity, we can use this value to know 
+		// if the hair instance is visible or not (i.e., HairLengthScale > 0)
+		Instance->HairGroupPublicData->VFInput.Strands.Common.LengthScale = Instance->Strands.Modifier.HairLengthScale;
+	}
+
+	// Only switch LOD if the data are ready to be used
+	const bool bIsLODDataReady = !ResourceStatus.HasStatus(EHairResourceStatus::EStatus::Loading) || (bStreamingEnabled && ResourceStatus.AvailableCurveCount > 0);
+	if (bIsLODDataReady)
+	{
+		const uint32 IntPrevHairLODIndex = FMath::FloorToInt(PrevHairLODIndex); 
+		const EHairBindingType CurrBindingType = Instance->HairGroupPublicData->GetBindingType(IntHairLODIndex);
+		const EHairBindingType PrevBindingType = Instance->HairGroupPublicData->GetBindingType(IntPrevHairLODIndex);
+
+		Instance->HairGroupPublicData->SetLODVisibility(bIsVisible);
+		Instance->HairGroupPublicData->SetLODIndex(HairLODIndex);
+		Instance->HairGroupPublicData->SetLODBias(0);
+		Instance->HairGroupPublicData->SetMeshLODIndex(MeshLODIndex);
+		Instance->HairGroupPublicData->VFInput.GeometryType = GeometryType;
+		Instance->HairGroupPublicData->VFInput.BindingType = CurrBindingType;
+		Instance->HairGroupPublicData->VFInput.bHasLODSwitch = IntPrevHairLODIndex != IntHairLODIndex;
+		Instance->HairGroupPublicData->VFInput.bHasLODSwitchBindingType = CurrBindingType != PrevBindingType;
+		Instance->GeometryType = GeometryType;
+		Instance->BindingType = CurrBindingType;
+		Instance->Guides.bIsSimulationEnable = Instance->HairGroupPublicData->IsSimulationEnable(IntHairLODIndex);
+		Instance->Guides.bHasGlobalInterpolation = Instance->HairGroupPublicData->IsGlobalInterpolationEnable(IntHairLODIndex);
+		Instance->Guides.bIsDeformationEnable = Instance->HairGroupPublicData->bIsDeformationEnable;
+		Instance->Guides.bIsSimulationCacheEnable = Instance->HairGroupPublicData->bIsSimulationCacheEnable;
+		Instance->Strands.bCullingEnable = bCullingEnable;
+	}
+
+	// Ensure that if the binding type is set to Skinning (or has RBF enabled), the skel. mesh is valid and support skin cache.
+	if ((Instance->Guides.bHasGlobalInterpolation || Instance->BindingType == EHairBindingType::Skinning) && MeshLODIndex < 0)
+	{
+		return false;
+	}
+
+	// If requested LOD's resources are not ready yet, and the previous LOD was invalid or if the MeshLODIndex has changed 
+	// (which would required extra data loading) change the geometry to be invalid, so that it don't get processed this frame.
+	const bool bIsLODValid = PrevHairLODIndex >= 0;
+	const bool bHasMeshLODChanged = PrevMeshLODIndex != MeshLODIndex && (Instance->Guides.bHasGlobalInterpolation || Instance->BindingType == EHairBindingType::Skinning);
+	if (!bIsLODDataReady && (!bIsLODValid || bHasMeshLODChanged))
+	{
+		return false;
+	}
+
+	// Adapt the number curve/point based on available curves/points
+	if (GeometryType == EHairGeometryType::Strands && bStreamingEnabled)
+	{
+		if (!bIsLODDataReady)
+		{
+			return false;
+		}
+		check(bIsLODDataReady);
+		check(Instance->Strands.RestResource);
+
+		// Adapt CurveCount/PointCount/CoverageScale based on what is actually available
+		const uint32 EffectiveCurveCount = FMath::Min(ResourceStatus.AvailableCurveCount, Instance->HairGroupPublicData->ContinuousLODCurveCount);
+		Instance->HairGroupPublicData->ContinuousLODCurveCount = EffectiveCurveCount;
+		Instance->HairGroupPublicData->ContinuousLODPointCount = EffectiveCurveCount > 0 ? Instance->Strands.Data->Header.CurveToPointCount[EffectiveCurveCount - 1] : 0;
+		Instance->HairGroupPublicData->ContinuousLODCoverageScale = ComputeActiveCurveCoverageScale(EffectiveCurveCount, Instance->HairGroupPublicData->RestCurveCount);
+		check(Instance->HairGroupPublicData->ContinuousLODPointCount <= Instance->HairGroupPublicData->RestPointCount);
+	}
+	return true;
+}
+
 static void RunHairLODSelection(
 	FRDGBuilder& GraphBuilder, 
 	FSceneInterface* Scene,
@@ -2361,8 +2675,8 @@ static void RunHairLODSelection(
 	}
 
 	// Detect if one of the current view has path tracing enabled for allocating raytracing resources
-#if RHI_RAYTRACING
 	bool bHasPathTracingView = false;
+#if RHI_RAYTRACING
 	for (const FSceneView* View : Views)
 	{
 		if (View->Family->EngineShowFlags.PathTracing)
@@ -2373,14 +2687,12 @@ static void RunHairLODSelection(
 	}
 #endif
 
-	for (FHairStrandsInstance* AbstractInstance : Instances)
+	struct FInstanceData
 	{
-		FHairGroupInstance* Instance = static_cast<FHairGroupInstance*>(AbstractInstance);
-
-		check(Instance);
-		check(Instance->HairGroupPublicData);
+		FHairGroupInstance* Instance = nullptr;
+		FHairGroupPublicData* HairGroupPublicData = nullptr;
+		FHairLOD HairLOD;
 		FTransform MeshLODLocalToWorld = FTransform::Identity;
-		const int32 MeshLODIndex = GetMeshLODIndex(GraphBuilder, ShaderMap, Scene, Instance, MeshLODLocalToWorld);
 
 		// Perform LOD selection based on all the views	
 		// CPU LOD selection. 
@@ -2389,283 +2701,53 @@ static void RunHairLODSelection(
 		// 
 		// Use the forced LOD index if set, otherwise compute the LOD based on the maximal screensize accross all views
 		// Compute the view where the screen size is maximale
-		const float PrevLODIndex = Instance->HairGroupPublicData->GetLODIndex();
-		const float PrevMeshLODIndex = Instance->HairGroupPublicData->GetMeshLODIndex();
+		int32 MeshLODIndex = -1;
+		int32 PrevMeshLODIndex = -1.f;
+		float PrevHairLODIndex = -1.f;
 
-		// 0. Initial LOD index picking
-		// Insure that MinLOD is necessary taken into account if a force LOD is request (i.e., LODIndex>=0). If a Force LOD 
-		// is not requested (i.e., LODIndex<0), the MinLOD is applied after ViewLODIndex has been determined in the codeblock below
-		const int32 LODCount = Instance->HairGroupPublicData->GetLODVisibilities().Num();
-		const float MinLOD = FMath::Max(0, GHairStrandsMinLOD);
-		
+		// 0. Initial LOD index picking	
 		// If continuous LOD is enabled, we bypass all other type of geometric representation, and only use LOD0
-		float LODIndex = IsHairVisibilityComputeRasterContinuousLODEnabled() ? 0.0 : (Instance->Debug.LODForcedIndex >= 0 ? FMath::Max(Instance->Debug.LODForcedIndex, MinLOD) : -1.0f);
-		{
-			FVector2f MaxContinuousLODScreenPos = FVector2f(0.f, 0.f);
-			float LODViewIndex = -1;
-			float MaxScreenSize_RestBound = 0.f;
-			float MaxScreenSize_Bound = 0.f;
-			const FSphere SphereBound = Instance->GetBounds().GetSphere();
-			for (const FSceneView* View : Views)
-			{
-				const FVector3d BoundScale = Instance->LocalToWorld.GetScale3D();
-				const FVector3f BoundExtent = Instance->Strands.Data ? FVector3f(Instance->Strands.Data->Header.BoundingBox.GetExtent()) : FVector3f(0,0,0);
-				const float BoundRadius = FMath::Max3(BoundExtent.X, BoundExtent.Y, BoundExtent.Z) * FMath::Max3(BoundScale.X, BoundScale.Y, BoundScale.Z);
-
-				const float ScreenSize_RestBound = FMath::Clamp(ComputeBoundsScreenSize(FVector4(SphereBound.Center, 1), BoundRadius, *View), 0.f, 1.0f);
-				const float ScreenSize_Bound = FMath::Clamp(ComputeBoundsScreenSize(FVector4(SphereBound.Center, 1), SphereBound.W, *View), 0.f, 1.0f);
-
-				const float CurrLODViewIndex = FMath::Max(MinLOD, GetHairInstanceLODIndex(Instance->HairGroupPublicData->GetLODScreenSizes(), ScreenSize_Bound, Instance->Strands.Modifier.LODBias));
 
-				MaxScreenSize_Bound = FMath::Max(MaxScreenSize_Bound, ScreenSize_Bound);
-				MaxScreenSize_RestBound = FMath::Max(MaxScreenSize_RestBound, ScreenSize_RestBound);
-				MaxContinuousLODScreenPos = ComputeProjectedScreenPos(SphereBound.Center, *View);
-
-				// Select highest LOD across all views
-				LODViewIndex = LODViewIndex < 0 ? CurrLODViewIndex : FMath::Min(LODViewIndex, CurrLODViewIndex);
-			}
-
-			if (LODIndex < 0)
-			{
-				LODIndex = LODViewIndex;
-			}
-			LODIndex = FMath::Clamp(LODIndex, 0.f, float(LODCount - 1));
-
-			// Feedback game thread with LOD selection 
-			Instance->Debug.LODPredictedIndex = LODViewIndex;
-			Instance->HairGroupPublicData->DebugScreenSize = MaxScreenSize_Bound;
-			Instance->HairGroupPublicData->ContinuousLODPointCount = Instance->HairGroupPublicData->RestPointCount;
-			Instance->HairGroupPublicData->ContinuousLODCurveCount = Instance->HairGroupPublicData->RestCurveCount;
-			Instance->HairGroupPublicData->ContinuousLODScreenSize = MaxScreenSize_RestBound;
-			Instance->HairGroupPublicData->ContinuousLODScreenPos = MaxContinuousLODScreenPos;
-			Instance->HairGroupPublicData->ContinuousLODBounds = SphereBound;
-			Instance->HairGroupPublicData->ContinuousLODCoverageScale = 1.f;
-
-			if (Instance->Strands.ClusterResource)
-			{
-				uint32 EffectiveCurveCount = 0;
-				if (Instance->HairGroupPublicData->bAutoLOD || IsHairStrandsForceAutoLODEnabled())
-				{
-					EffectiveCurveCount = ComputeActiveCurveCount(Instance->HairGroupPublicData->ContinuousLODScreenSize, Instance->HairGroupPublicData->RestCurveCount, Instance->HairGroupPublicData->ClusterCount);
-					LODIndex = 0;
-				}
-				else
-				{
-					EffectiveCurveCount = Instance->Strands.ClusterResource->BulkData.GetCurveCount(LODIndex);
-				}
-				check(EffectiveCurveCount <= uint32(Instance->Strands.Data->Header.CurveToPointCount.Num()));
-
-				Instance->HairGroupPublicData->ContinuousLODCurveCount = EffectiveCurveCount;
-				Instance->HairGroupPublicData->ContinuousLODPointCount = EffectiveCurveCount > 0 ? Instance->Strands.Data->Header.CurveToPointCount[EffectiveCurveCount - 1] : 0;
-				Instance->HairGroupPublicData->ContinuousLODCoverageScale = ComputeActiveCurveCoverageScale(EffectiveCurveCount, Instance->HairGroupPublicData->RestCurveCount);
-			}
-		}
-
-		// Function for selecting, loading, & initializing LOD resources
-		auto SelectValidLOD = [Instance, ShaderPlatform, ShaderMap, PrevLODIndex, LODCount, MeshLODIndex, PrevMeshLODIndex
-		#if RHI_RAYTRACING
-		, bHasPathTracingView
-		#endif // RHI_RAYTRACING
-		] (FRDGBuilder& GraphBuilder, float LODIndex) -> bool
-		{
-			const int32 IntLODIndex = FMath::Clamp(FMath::FloorToInt(LODIndex), 0, LODCount - 1);
-			const TArray<EHairGeometryType>& LODGeometryTypes = Instance->HairGroupPublicData->GetLODGeometryTypes();
-			const TArray<bool>& LODVisibilities = Instance->HairGroupPublicData->GetLODVisibilities();
-			const bool bIsVisible = LODVisibilities[IntLODIndex];
-			const bool bForceCards = GHairStrands_UseCards > 0 || Instance->bForceCards; // todo
-			EHairGeometryType GeometryType = ConvertLODGeometryType(LODGeometryTypes[IntLODIndex], bForceCards, ShaderPlatform);
-
-			// Skip cluster culling if strands have a single LOD (i.e.: LOD0), and the instance is static (i.e., no skinning, no simulation, ...)
-			bool bCullingEnable = false;
-			if (GeometryType == EHairGeometryType::Meshes)
-			{
-				if (!Instance->Meshes.IsValid(IntLODIndex))
-				{
-					GeometryType = EHairGeometryType::NoneGeometry;
-				}
-			}
-			else if (GeometryType == EHairGeometryType::Cards)
-			{
-				if (!Instance->Cards.IsValid(IntLODIndex))
-				{
-					GeometryType = EHairGeometryType::NoneGeometry;
-				}
-			}
-			else if (GeometryType == EHairGeometryType::Strands)
-			{
-				if (!Instance->Strands.IsValid())
-				{
-					GeometryType = EHairGeometryType::NoneGeometry;
-				}
-				else
-				{
-					const bool bNeedLODing		= LODIndex > 0;
-					const bool bNeedDeformation = Instance->Strands.DeformedResource != nullptr;
-					bCullingEnable = (bNeedLODing || bNeedDeformation);
-				}
-			}
-
-			if (!bIsVisible)
-			{
-				GeometryType = EHairGeometryType::NoneGeometry;
-			}
-
-			// Transform ContinuousLOD CurveCount/PointCount into streaming request
-			uint32 RequestedCurveCount = Instance->HairGroupPublicData->RestCurveCount;
-			uint32 RequestedPointCount = Instance->HairGroupPublicData->RestPointCount;
-			const bool bStreamingEnabled = GHairStrands_Streaming > 0;
-			if (GeometryType == EHairGeometryType::Strands && bStreamingEnabled && Instance->bSupportStreaming)
-			{
-				// Round Curve/Point request to curve 'page'
-				RequestedCurveCount = GetRoundedCurveCount(Instance->HairGroupPublicData->ContinuousLODCurveCount, Instance->HairGroupPublicData->RestCurveCount);
-				RequestedPointCount = RequestedCurveCount > 0 ? Instance->Strands.Data->Header.CurveToPointCount[RequestedCurveCount - 1] : 0;
-			}
-
-			const bool bSimulationEnable			= Instance->HairGroupPublicData->IsSimulationEnable(IntLODIndex);
-			const bool bDeformationEnable			= Instance->HairGroupPublicData->bIsDeformationEnable;
-			const bool bGlobalInterpolationEnable	= Instance->HairGroupPublicData->IsGlobalInterpolationEnable(IntLODIndex);
-			const bool bSimulationCacheEnable		= Instance->HairGroupPublicData->bIsSimulationCacheEnable;
-			const bool bLODNeedsGuides				= bSimulationEnable || bDeformationEnable || bGlobalInterpolationEnable || bSimulationCacheEnable;
-
-			const EHairResourceLoadingType LoadingType = GetHairResourceLoadingType(GeometryType, IntLODIndex);
-			EHairResourceStatus ResourceStatus;
-			ResourceStatus.Status 				= EHairResourceStatus::EStatus::None;
-			ResourceStatus.AvailableCurveCount 	= Instance->HairGroupPublicData->RestCurveCount;
-
-			// Lazy allocation of resources
-			// Note: Allocation will only be done if the resources is not initialized yet. Guides deformed position are also initialized from the Rest position at creation time.
-			if (Instance->Guides.Data && bLODNeedsGuides)
-			{
-				if (Instance->Guides.RestRootResource)			{ Instance->Guides.RestRootResource->Allocate(GraphBuilder, LoadingType, ResourceStatus, MeshLODIndex); }
-				if (Instance->Guides.RestResource)				{ Instance->Guides.RestResource->Allocate(GraphBuilder, LoadingType, ResourceStatus); }
-				if (Instance->Guides.DeformedRootResource)		{ Instance->Guides.DeformedRootResource->Allocate(GraphBuilder, LoadingType, ResourceStatus, MeshLODIndex); }
-				if (Instance->Guides.DeformedResource)			
-				{ 
-					// Ensure the rest resources are correctly loaded prior to initialized and copy the rest position into the deformed positions
-					if (Instance->Guides.RestResource->bIsInitialized)
-					{
-						const bool bNeedCopy = !Instance->Guides.DeformedResource->bIsInitialized; 
-						Instance->Guides.DeformedResource->Allocate(GraphBuilder, EHairResourceLoadingType::Sync); 
-						if (bNeedCopy) 
-						{ 
-							AddCopyHairStrandsPositionPass(GraphBuilder, ShaderMap, *Instance->Guides.RestResource, *Instance->Guides.DeformedResource); 
-						}
-					}
-				}
-			}
-
-			if (GeometryType == EHairGeometryType::Meshes)
-			{
-				FHairGroupInstance::FMeshes::FLOD& InstanceLOD = Instance->Meshes.LODs[IntLODIndex];
-
-				if (InstanceLOD.DeformedResource)				{ InstanceLOD.DeformedResource->Allocate(GraphBuilder, LoadingType, ResourceStatus); }
-				#if RHI_RAYTRACING
-				if (InstanceLOD.RaytracingResource)				{ InstanceLOD.RaytracingResource->Allocate(GraphBuilder, LoadingType, ResourceStatus);}
-				#endif
-
-				InstanceLOD.InitVertexFactory();
-			}
-			else if (GeometryType == EHairGeometryType::Cards)
-			{
-				FHairGroupInstance::FCards::FLOD& InstanceLOD = Instance->Cards.LODs[IntLODIndex];
-
-				if (InstanceLOD.InterpolationResource)			{ InstanceLOD.InterpolationResource->Allocate(GraphBuilder, LoadingType, ResourceStatus); }
-				if (InstanceLOD.DeformedResource)				{ InstanceLOD.DeformedResource->Allocate(GraphBuilder, LoadingType, ResourceStatus); }
-				if (InstanceLOD.Guides.RestRootResource)		{ InstanceLOD.Guides.RestRootResource->Allocate(GraphBuilder, LoadingType, ResourceStatus, MeshLODIndex); }
-				if (InstanceLOD.Guides.RestResource)			{ InstanceLOD.Guides.RestResource->Allocate(GraphBuilder, LoadingType, ResourceStatus); }
-				if (InstanceLOD.Guides.DeformedRootResource)	{ InstanceLOD.Guides.DeformedRootResource->Allocate(GraphBuilder, LoadingType, ResourceStatus, MeshLODIndex); }
-				if (InstanceLOD.Guides.DeformedResource)		{ InstanceLOD.Guides.DeformedResource->Allocate(GraphBuilder, LoadingType, ResourceStatus); }
-				if (InstanceLOD.Guides.InterpolationResource)	{ InstanceLOD.Guides.InterpolationResource->Allocate(GraphBuilder, LoadingType, ResourceStatus); }
-				#if RHI_RAYTRACING
-				if (InstanceLOD.RaytracingResource)				{ InstanceLOD.RaytracingResource->Allocate(GraphBuilder, LoadingType, ResourceStatus);}
-				#endif
-				InstanceLOD.InitVertexFactory();
-			}
-			else if (GeometryType == EHairGeometryType::Strands)
-			{
-				check(Instance->HairGroupPublicData);
-
-				if (Instance->Strands.RestRootResource)			{ Instance->Strands.RestRootResource->Allocate(GraphBuilder, EHairResourceLoadingType::Async, ResourceStatus, RequestedCurveCount, RequestedPointCount, MeshLODIndex); }
-				if (Instance->Strands.RestResource)				{ Instance->Strands.RestResource->Allocate(GraphBuilder, EHairResourceLoadingType::Async, ResourceStatus, RequestedCurveCount, RequestedPointCount); }
-				if (Instance->Strands.ClusterResource)			{ Instance->Strands.ClusterResource->Allocate(GraphBuilder, EHairResourceLoadingType::Async, ResourceStatus, RequestedCurveCount, RequestedPointCount); }
-				if (Instance->Strands.InterpolationResource)	{ Instance->Strands.InterpolationResource->Allocate(GraphBuilder, EHairResourceLoadingType::Async, ResourceStatus, RequestedCurveCount, RequestedPointCount); }
-				if (Instance->Strands.CullingResource)			{ Instance->Strands.CullingResource->Allocate(GraphBuilder, EHairResourceLoadingType::Async, ResourceStatus, RequestedCurveCount, RequestedPointCount); }
-
-				if (Instance->Strands.DeformedRootResource)		{ Instance->Strands.DeformedRootResource->Allocate(GraphBuilder, EHairResourceLoadingType::Async, ResourceStatus, RequestedCurveCount, RequestedPointCount, MeshLODIndex); }
-				if (Instance->Strands.DeformedResource)			{ Instance->Strands.DeformedResource->Allocate(GraphBuilder, EHairResourceLoadingType::Async, ResourceStatus, RequestedCurveCount, RequestedPointCount); }
-				#if RHI_RAYTRACING
-				if (bHasPathTracingView)						{ AllocateRaytracingResources(Instance); }
-				if (Instance->Strands.RenRaytracingResource)	{ Instance->Strands.RenRaytracingResource->Allocate(GraphBuilder, EHairResourceLoadingType::Async, ResourceStatus, RequestedCurveCount, RequestedPointCount); }
-				#endif
-				Instance->Strands.VertexFactory->InitResources(GraphBuilder.RHICmdList);
-
-				// Early initialization, so that when filtering MeshBatch block in SceneVisiblity, we can use this value to know 
-				// if the hair instance is visible or not (i.e., HairLengthScale > 0)
-				Instance->HairGroupPublicData->VFInput.Strands.Common.LengthScale = Instance->Strands.Modifier.HairLengthScale;
-			}
-
-			// Only switch LOD if the data are ready to be used
-			const bool bIsLODDataReady = !ResourceStatus.HasStatus(EHairResourceStatus::EStatus::Loading) || (bStreamingEnabled && ResourceStatus.AvailableCurveCount > 0);
-			if (bIsLODDataReady)
-			{
-				const uint32 IntPrevLODIndex = FMath::FloorToInt(PrevLODIndex); 
-				const EHairBindingType CurrBindingType = Instance->HairGroupPublicData->GetBindingType(IntLODIndex);
-				const EHairBindingType PrevBindingType = Instance->HairGroupPublicData->GetBindingType(IntPrevLODIndex);
-
-				Instance->HairGroupPublicData->SetLODVisibility(bIsVisible);
-				Instance->HairGroupPublicData->SetLODIndex(LODIndex);
-				Instance->HairGroupPublicData->SetLODBias(0);
-				Instance->HairGroupPublicData->SetMeshLODIndex(MeshLODIndex);
-				Instance->HairGroupPublicData->VFInput.GeometryType = GeometryType;
-				Instance->HairGroupPublicData->VFInput.BindingType = CurrBindingType;
-				Instance->HairGroupPublicData->VFInput.bHasLODSwitch = IntPrevLODIndex != IntLODIndex;
-				Instance->HairGroupPublicData->VFInput.bHasLODSwitchBindingType = CurrBindingType != PrevBindingType;
-				Instance->GeometryType = GeometryType;
-				Instance->BindingType = CurrBindingType;
-				Instance->Guides.bIsSimulationEnable = Instance->HairGroupPublicData->IsSimulationEnable(IntLODIndex);
-				Instance->Guides.bHasGlobalInterpolation = Instance->HairGroupPublicData->IsGlobalInterpolationEnable(IntLODIndex);
-				Instance->Guides.bIsDeformationEnable = Instance->HairGroupPublicData->bIsDeformationEnable;
-				Instance->Guides.bIsSimulationCacheEnable = Instance->HairGroupPublicData->bIsSimulationCacheEnable;
-				Instance->Strands.bCullingEnable = bCullingEnable;
-			}
+	};
 
-			// Ensure that if the binding type is set to Skinning (or has RBF enabled), the skel. mesh is valid and support skin cache.
-			if ((Instance->Guides.bHasGlobalInterpolation || Instance->BindingType == EHairBindingType::Skinning) && MeshLODIndex < 0)
-			{
-				return false;
-			}
+	// Find unique groom assets
+	TArray<FInstanceData> InstanceDatas;
+	InstanceDatas.Reserve(Instances.Num());
+	for (FHairStrandsInstance* AbstractInstance : Instances)
+	{
+		FHairGroupInstance* Instance = static_cast<FHairGroupInstance*>(AbstractInstance);
+		check(Instance);
+		check(Instance->HairGroupPublicData);
 
-			// If requested LOD's resources are not ready yet, and the previous LOD was invalid or if the MeshLODIndex has changed 
-			// (which would required extra data loading) change the geometry to be invalid, so that it don't get processed this frame.
-			const bool bIsLODValid = PrevLODIndex >= 0;
-			const bool bHasMeshLODChanged = PrevMeshLODIndex != MeshLODIndex && (Instance->Guides.bHasGlobalInterpolation || Instance->BindingType == EHairBindingType::Skinning);
-			if (!bIsLODDataReady && (!bIsLODValid || bHasMeshLODChanged))
-			{
-				return false;
-			}
+		FInstanceData& InstanceData = InstanceDatas.AddDefaulted_GetRef();
+		InstanceData.Instance 				= Instance;
+		InstanceData.HairGroupPublicData 	= Instance->HairGroupPublicData;
+		InstanceData.PrevHairLODIndex 		= Instance->HairGroupPublicData->GetLODIndex();
+		InstanceData.PrevMeshLODIndex 		= Instance->HairGroupPublicData->GetMeshLODIndex();
+		InstanceData.HairLOD 				= ComputeHairLODIndex(Instance, Views);
+		InstanceData.MeshLODIndex 			= GetMeshLODIndex(GraphBuilder, ShaderMap, Scene, Instance, InstanceData.MeshLODLocalToWorld);
+	}
 
-			// Adapt the number curve/point based on available curves/points
-			if (GeometryType == EHairGeometryType::Strands && bStreamingEnabled)
-			{
-				if (!bIsLODDataReady)
-				{
-					return false;
-				}
-				check(bIsLODDataReady);
-				check(Instance->Strands.RestResource);
-
-				// Adapt CurveCount/PointCount/CoverageScale based on what is actually available
-				const uint32 EffectiveCurveCount = FMath::Min(ResourceStatus.AvailableCurveCount, Instance->HairGroupPublicData->ContinuousLODCurveCount);
-				Instance->HairGroupPublicData->ContinuousLODCurveCount = EffectiveCurveCount;
-				Instance->HairGroupPublicData->ContinuousLODPointCount = EffectiveCurveCount > 0 ? Instance->Strands.Data->Header.CurveToPointCount[EffectiveCurveCount - 1] : 0;
-				Instance->HairGroupPublicData->ContinuousLODCoverageScale = ComputeActiveCurveCoverageScale(EffectiveCurveCount, Instance->HairGroupPublicData->RestCurveCount);
-				check(Instance->HairGroupPublicData->ContinuousLODPointCount <= Instance->HairGroupPublicData->RestPointCount);
-			}
-			return true;
-		};
+	// LOD selection & load resources
+	for (FInstanceData& InstanceData : InstanceDatas)
+	{
+		// 0. Initial LOD index picking	
+		// If continuous LOD is enabled, we bypass all other type of geometric representation, and only use LOD0
+		//const FHairLOD HairLOD = ComputeHairLODIndex(InstanceData.Instance, Views);
+		ApplyHairLOD(InstanceData.HairLOD, InstanceData.Instance);
 
 		// 1. Try to find an available LOD
-		bool bFoundValidLOD = SelectValidLOD(GraphBuilder, LODIndex);
+		bool bFoundValidLOD = SelectValidLOD(
+			GraphBuilder,
+			ShaderMap,
+			ShaderPlatform,
+			InstanceData.Instance,
+			InstanceData.HairLOD.HairLODIndex,
+			InstanceData.PrevHairLODIndex,
+			InstanceData.HairLOD.HairLODCount,
+			InstanceData.MeshLODIndex,
+			InstanceData.PrevMeshLODIndex,
+			bHasPathTracingView);
 
 		// 2. If no valid LOD are found, try to find a coarser LOD which is:
 		// * loaded 
@@ -2674,13 +2756,23 @@ static void RunHairLODSelection(
 		// * does not require simulation (as simulation requires LOD to be selected on the game thread i.e., Predicted/Forced, and so we can override this LOD here)
 		if (!bFoundValidLOD)
 		{
-			for (int32 FallbackLODIndex = FMath::Clamp(FMath::FloorToInt(LODIndex + 1), 0, LODCount - 1); FallbackLODIndex < LODCount; ++FallbackLODIndex)
+			for (uint32 FallbackHairLODIndex = FMath::Clamp(FMath::FloorToInt(InstanceData.HairLOD.HairLODIndex + 1), 0, InstanceData.HairLOD.HairLODCount - 1); FallbackHairLODIndex < InstanceData.HairLOD.HairLODCount; ++FallbackHairLODIndex)
 			{
-				if (Instance->HairGroupPublicData->GetBindingType(FallbackLODIndex) == EHairBindingType::Rigid &&
-					!Instance->HairGroupPublicData->IsSimulationEnable(FallbackLODIndex) &&
-					!Instance->HairGroupPublicData->IsGlobalInterpolationEnable(FallbackLODIndex))
+				if ( InstanceData.HairGroupPublicData->GetBindingType(FallbackHairLODIndex) == EHairBindingType::Rigid &&
+					!InstanceData.HairGroupPublicData->IsSimulationEnable(FallbackHairLODIndex) &&
+					!InstanceData.HairGroupPublicData->IsGlobalInterpolationEnable(FallbackHairLODIndex))
 				{
-					bFoundValidLOD = SelectValidLOD(GraphBuilder, FallbackLODIndex);
+					bFoundValidLOD = SelectValidLOD(
+						GraphBuilder,
+						ShaderMap,
+						ShaderPlatform,
+						InstanceData.Instance,
+						FallbackHairLODIndex,
+						InstanceData.PrevHairLODIndex,
+						InstanceData.HairLOD.HairLODCount,
+						InstanceData.MeshLODIndex,
+						InstanceData.PrevMeshLODIndex,
+						bHasPathTracingView);
 					break;
 				}
 			}
@@ -2689,17 +2781,17 @@ static void RunHairLODSelection(
 		// 3. If no LOD are valid, then mark the instance as invalid for this frame
 		if (!bFoundValidLOD)
 		{
-			Instance->GeometryType = EHairGeometryType::NoneGeometry;
-			Instance->BindingType = EHairBindingType::NoneBinding;
+			InstanceData.Instance->GeometryType = EHairGeometryType::NoneGeometry;
+			InstanceData.Instance->BindingType = EHairBindingType::NoneBinding;
 		}
 
 		// Update the local-to-world transform based on the binding type 
 		if (GetHairSwapBufferType() != EHairBufferSwapType::Tick)
 		{
-			Instance->Debug.SkinningPreviousLocalToWorld = Instance->Debug.SkinningCurrentLocalToWorld;
-			Instance->Debug.SkinningCurrentLocalToWorld = MeshLODLocalToWorld;
+			InstanceData.Instance->Debug.SkinningPreviousLocalToWorld = InstanceData.Instance->Debug.SkinningCurrentLocalToWorld;
+			InstanceData.Instance->Debug.SkinningCurrentLocalToWorld  = InstanceData.MeshLODLocalToWorld;
 		}
-		Instance->LocalToWorld = Instance->GetCurrentLocalToWorld();
+		InstanceData.Instance->LocalToWorld = InstanceData.Instance->GetCurrentLocalToWorld();
 	}
 }
 
