Fix hair AutoLOD to work in combination of other geometry types.

diff --git a/Engine/Plugins/Runtime/HairStrands/Source/HairStrandsCore/Private/GroomManager.cpp b/Engine/Plugins/Runtime/HairStrands/Source/HairStrandsCore/Private/GroomManager.cpp
--- a/Engine/Plugins/Runtime/HairStrands/Source/HairStrandsCore/Private/GroomManager.cpp
+++ b/Engine/Plugins/Runtime/HairStrands/Source/HairStrandsCore/Private/GroomManager.cpp
@@ -2416,14 +2416,36 @@ static FHairLOD ComputeHairLODIndex(const FHairGroupInstance* Instance, const TA
 	Out.ContinuousLODBounds = SphereBound;
 	Out.ContinuousLODCoverageScale = 1.f;
 
+	// Extract the min/max LOD with strands. 
+	// This assumes that these LODs are contiguous
+	int32 MinLODIndexWithStrands = INDEX_NONE;
+	int32 MaxLODIndexWithStrands = INDEX_NONE;
+	{
+		int32 LODIndex = 0;
+		const TArray<EHairGeometryType>& GeometryTypes = Instance->HairGroupPublicData->GetLODGeometryTypes();
+		for (EHairGeometryType Type : GeometryTypes)
+		{
+			if (Type == EHairGeometryType::Strands)
+			{
+				if (MinLODIndexWithStrands == INDEX_NONE)
+				{
+					MinLODIndexWithStrands = LODIndex;
+				}
+				MaxLODIndexWithStrands = LODIndex;
+			}
+			++LODIndex;
+		}
+	}
+
 	// Auto LOD
+	const bool bNeedAutoLOD = MaxLODIndexWithStrands != INDEX_NONE && int32(Out.HairLODIndex) <= MaxLODIndexWithStrands;
 	if (Instance->Strands.ClusterResource && Instance->Strands.Data)
 	{
 		uint32 EffectiveCurveCount = 0;
-		if (Instance->HairGroupPublicData->bAutoLOD || IsHairStrandsForceAutoLODEnabled())
+		if (bNeedAutoLOD && (Instance->HairGroupPublicData->bAutoLOD || IsHairStrandsForceAutoLODEnabled()))
 		{
 			EffectiveCurveCount = ComputeActiveCurveCount(Out.ContinuousLODScreenSize, Instance->HairGroupPublicData->RestCurveCount, Instance->HairGroupPublicData->ClusterCount);
-			Out.HairLODIndex = 0;
+			Out.HairLODIndex = MinLODIndexWithStrands;
 		}
 		else
 		{
