Use correct primitive type during Groom component PSO precaching

diff --git a/Engine/Plugins/Runtime/HairStrands/Source/HairStrandsCore/Private/GroomComponent.cpp b/Engine/Plugins/Runtime/HairStrands/Source/HairStrandsCore/Private/GroomComponent.cpp
--- a/Engine/Plugins/Runtime/HairStrands/Source/HairStrandsCore/Private/GroomComponent.cpp
+++ b/Engine/Plugins/Runtime/HairStrands/Source/HairStrandsCore/Private/GroomComponent.cpp
@@ -3121,6 +3121,12 @@ void UGroomComponent::CollectPSOPrecacheData(const FPSOPrecacheParams& BasePreca
 		UMaterialInterface* MaterialInterface = GetMaterial(GetMaterialIndexWithFallback(VFsPerMaterial.MaterialIndex), VFsPerMaterial.HairGeometryType, true);
 		if (MaterialInterface)
 		{
+#if USE_HAIR_TRIANGLE_STRIP
+			PrecachePSOParams.PrimitiveType = VFsPerMaterial.HairGeometryType == EHairGeometryType::Strands ? PT_TriangleStrip : PT_TriangleList;
+#else
+			PrecachePSOParams.PrimitiveType = PT_TriangleList;
+#endif
+
 			FComponentPSOPrecacheParams& ComponentParams = OutParams[OutParams.AddDefaulted()];
 			ComponentParams.MaterialInterface = MaterialInterface;
 			ComponentParams.VertexFactoryDataList = VFsPerMaterial.VertexFactoryDataList;
