Split GroomAsset platform data into data and resources.

This clean up the derive data building path, and will easy fixing instance streaming.

diff --git a/Engine/Plugins/Runtime/HairStrands/Source/HairStrandsCore/Private/Deformer/DeformerDataInterfaceGroomExec.cpp b/Engine/Plugins/Runtime/HairStrands/Source/HairStrandsCore/Private/Deformer/DeformerDataInterfaceGroomExec.cpp
--- a/Engine/Plugins/Runtime/HairStrands/Source/HairStrandsCore/Private/Deformer/DeformerDataInterfaceGroomExec.cpp
+++ b/Engine/Plugins/Runtime/HairStrands/Source/HairStrandsCore/Private/Deformer/DeformerDataInterfaceGroomExec.cpp
@@ -120,8 +120,8 @@ int32 FOptimusGroomExecDataProviderProxy::GetDispatchThreadCount(TArray<FIntVect
 	{		
 		if (FHairGroupInstance* Instance = GroomComponent->GetGroupInstance(InvocationIndex))
 		{
-			const int32 NumControlPoints = Instance->Strands.Data->GetNumPoints();
-			const int32 NumCurves = Instance->Strands.Data->GetNumCurves();
+			const int32 NumControlPoints = Instance->Strands.GetData().GetNumPoints();
+			const int32 NumCurves = Instance->Strands.GetData().GetNumCurves();
 			const int32 NumThreads = Domain == EOptimusGroomExecDomain::ControlPoint ? NumControlPoints : NumCurves;
 			ThreadCounts.Add(FIntVector(NumThreads, 1, 1));
 		}
@@ -149,8 +149,8 @@ void FOptimusGroomExecDataProviderProxy::GatherDispatchData(FDispatchData const&
 	for (int32 InvocationIndex = 0; InvocationIndex < ParameterArray.Num(); ++InvocationIndex)
 	{
 		FHairGroupInstance* Instance = GroomComponent->GetGroupInstance(InvocationIndex);
-		const int32 NumControlPoints = Instance->Strands.Data->GetNumPoints();
-		const int32 NumCurves = Instance->Strands.Data->GetNumCurves();
+		const int32 NumControlPoints = Instance->Strands.GetData().GetNumPoints();
+		const int32 NumCurves = Instance->Strands.GetData().GetNumCurves();
 		const int32 NumThreads = Domain == EOptimusGroomExecDomain::ControlPoint ? NumControlPoints : NumCurves;
 
 		FParameters& Parameters = ParameterArray[InvocationIndex];
diff --git a/Engine/Plugins/Runtime/HairStrands/Source/HairStrandsCore/Private/Deformer/DeformerDataInterfaceGroomGuide.cpp b/Engine/Plugins/Runtime/HairStrands/Source/HairStrandsCore/Private/Deformer/DeformerDataInterfaceGroomGuide.cpp
--- a/Engine/Plugins/Runtime/HairStrands/Source/HairStrandsCore/Private/Deformer/DeformerDataInterfaceGroomGuide.cpp
+++ b/Engine/Plugins/Runtime/HairStrands/Source/HairStrandsCore/Private/Deformer/DeformerDataInterfaceGroomGuide.cpp
@@ -197,8 +197,8 @@ void FOptimusGroomGuideDataProviderProxy::GatherDispatchData(FDispatchData const
 				// Reuse the strands properties for guides, and apply guides values
 				const FHairGroupPublicData::FVertexFactoryInput VFInput = ComputeHairStrandsVertexInputData(Instance, EGroomViewMode::None);
 				Parameters.Common 				= VFInput.Strands.Common;
-				Parameters.Common.PointCount 	= Instance->Guides.Data ? Instance->Guides.Data->Header.PointCount : 0;
-				Parameters.Common.CurveCount 	= Instance->Guides.Data ? Instance->Guides.Data->Header.CurveCount : 0;
+				Parameters.Common.PointCount 	= Instance->Guides.IsValid() ? Instance->Guides.GetData().Header.PointCount : 0;
+				Parameters.Common.CurveCount 	= Instance->Guides.IsValid() ? Instance->Guides.GetData().Header.CurveCount : 0;
 				Parameters.Resources 			= Resources[InvocationIndex];
 			}
 		}
diff --git a/Engine/Plugins/Runtime/HairStrands/Source/HairStrandsCore/Private/Deformer/DeformerDataInterfaceGroomWrite.cpp b/Engine/Plugins/Runtime/HairStrands/Source/HairStrandsCore/Private/Deformer/DeformerDataInterfaceGroomWrite.cpp
--- a/Engine/Plugins/Runtime/HairStrands/Source/HairStrandsCore/Private/Deformer/DeformerDataInterfaceGroomWrite.cpp
+++ b/Engine/Plugins/Runtime/HairStrands/Source/HairStrandsCore/Private/Deformer/DeformerDataInterfaceGroomWrite.cpp
@@ -195,8 +195,8 @@ void FOptimusGroomWriteDataProviderProxy::AllocateResources(FRDGBuilder& GraphBu
 		if (GroomInstance)
 		{
 			// Allocate required buffers
-			const int32 NumControlPoints = GroomInstance->Strands.Data->GetNumPoints();
-			const int32 NumCurves = GroomInstance->Strands.Data->GetNumCurves();
+			const int32 NumControlPoints = GroomInstance->Strands.GetData().GetNumPoints();
+			const int32 NumCurves = GroomInstance->Strands.GetData().GetNumCurves();
 
 			FResources& R = Resources.AddDefaulted_GetRef();
 
diff --git a/Engine/Plugins/Runtime/HairStrands/Source/HairStrandsCore/Private/Deformer/DeformerGroomComponentSource.cpp b/Engine/Plugins/Runtime/HairStrands/Source/HairStrandsCore/Private/Deformer/DeformerGroomComponentSource.cpp
--- a/Engine/Plugins/Runtime/HairStrands/Source/HairStrandsCore/Private/Deformer/DeformerGroomComponentSource.cpp
+++ b/Engine/Plugins/Runtime/HairStrands/Source/HairStrandsCore/Private/Deformer/DeformerGroomComponentSource.cpp
@@ -83,8 +83,8 @@ bool UOptimusGroomComponentSource::GetComponentElementCountsForExecutionDomain(
 		{
 			if (const FHairGroupInstance* Instance = GroomComponent->GetGroupInstance(InvocationIndex))
 			{
-				const int32 NumControlPoints = Instance->Strands.Data->GetNumPoints();
-				const int32 NumCurves = Instance->Strands.Data->GetNumCurves();
+				const int32 NumControlPoints = Instance->Strands.GetData().GetNumPoints();
+				const int32 NumCurves = Instance->Strands.GetData().GetNumCurves();
 				const int32 NumThreads = InDomainName == Domains::ControlPoint ? NumControlPoints : NumCurves;
 				OutInvocationElementCounts.Add(NumThreads);
 			}
diff --git a/Engine/Plugins/Runtime/HairStrands/Source/HairStrandsCore/Private/GroomAsset.cpp b/Engine/Plugins/Runtime/HairStrands/Source/HairStrandsCore/Private/GroomAsset.cpp
--- a/Engine/Plugins/Runtime/HairStrands/Source/HairStrandsCore/Private/GroomAsset.cpp
+++ b/Engine/Plugins/Runtime/HairStrands/Source/HairStrandsCore/Private/GroomAsset.cpp
@@ -125,24 +125,26 @@ static FAutoConsoleVariableRef CVarHairStrandsDump_All(
 
 uint32 FHairGroupPlatformData::FStrands::GetDataSize() const
 {
-	uint32 Total = 0;
-	Total += FBaseWithInterpolation::GetDataSize();
-	uint32 CTotal = 0;
+	uint32 ATotal = 0;
 	uint32 BTotal = 0;
+	uint32 CTotal = 0;
+	ATotal += ::GetDataSize(BulkData);
+	ATotal += ::GetDataSize(InterpolationBulkData);
 	BTotal += ClusterBulkData.Data.CurveToClusterIds.IsBulkDataLoaded() ? ClusterBulkData.Data.CurveToClusterIds.GetBulkDataSize() : 0;
 	CTotal += ClusterBulkData.Data.PackedClusterInfos.IsBulkDataLoaded()? ClusterBulkData.Data.PackedClusterInfos.GetBulkDataSize() : 0;
 	CTotal += ClusterBulkData.Data.PointLODs.IsBulkDataLoaded()	? ClusterBulkData.Data.PointLODs.GetBulkDataSize() : 0;
 
-	return BTotal + Total + CTotal;
+
+	return ATotal + BTotal + CTotal;
 }
 
-FGroomAssetMemoryStats FGroomAssetMemoryStats::Get(const FHairGroupPlatformData& In)
+FGroomAssetMemoryStats FGroomAssetMemoryStats::Get(const FHairGroupPlatformData& InData, const FHairGroupResources& In)
 {
 	FGroomAssetMemoryStats Out;
-	Out.CPU.Guides  = In.Guides.GetDataSize();
-	Out.CPU.Strands = In.Strands.GetDataSize();
-	Out.CPU.Cards   = In.Cards.GetDataSize();
-	Out.CPU.Meshes  = In.Meshes.GetDataSize();
+	Out.CPU.Guides  = InData.Guides.GetDataSize();
+	Out.CPU.Strands = InData.Strands.GetDataSize();
+	Out.CPU.Cards   = InData.Cards.GetDataSize();
+	Out.CPU.Meshes  = InData.Meshes.GetDataSize();
 
 	Out.GPU.Guides  = In.Guides.GetResourcesSize();
 	Out.GPU.Strands = In.Strands.GetResourcesSize();
@@ -220,7 +222,7 @@ void DumpLoadedGroomAssets(IConsoleVariable* InCVarPakTesterEnabled)
 			const uint32 GroupCount = AssetIt->GetHairGroupsPlatformData().Num();
 			for (uint32 GroupIt = 0; GroupIt < GroupCount; ++GroupIt)
 			{
-				const FGroomAssetMemoryStats Group = FGroomAssetMemoryStats::Get(AssetIt->GetHairGroupsPlatformData()[GroupIt]);
+				const FGroomAssetMemoryStats Group = FGroomAssetMemoryStats::Get(AssetIt->GetHairGroupsPlatformData()[GroupIt], AssetIt->GetHairGroupsResources()[GroupIt]);
 				Total.Accumulate(Group);
 				const uint32 LODCount = AssetIt->GetHairGroupsLOD()[GroupIt].LODs.Num();
 				if (bDetails)
@@ -410,12 +412,13 @@ void UpdateHairStrandsVerbosity(IConsoleVariable* InCVarVerbosity)
 /////////////////////////////////////////////////////////////////////////////////////////
 
 template<typename ResourceType>
-static void InitAtlasTexture(ResourceType* InResource, const FHairGroupCardsTextures& In)
+static void InitAtlasTexture(ResourceType* InResource, const FHairGroupCardsTextures& In, const bool bInvertUV)
 {
 	if (InResource == nullptr)
 	{
 		return;
 	}
+	InResource->bInvertUV = bInvertUV;
 
 	const TArray<UTexture2D*>& InTextures = In.Textures;
 	const uint32 InLayoutIndex = uint32(In.Layout);
@@ -455,34 +458,36 @@ static void InitAtlasTexture(ResourceType* InResource, const FHairGroupCardsText
 }
 
 template<typename T>
-static const T* GetSourceDescription(const TArray<T>& InHairGroups, uint32 GroupIndex, uint32 LODIndex, int32& SourceIndex)
+static const T* GetSourceDescription(const TArray<T>& InHairGroups, uint32 GroupIndex, uint32 LODIndex, int32* SourceIndex=nullptr)
 {
-	SourceIndex = 0;
+	int32 LocalSourceIndex = 0;
 	for (const T& SourceDesc : InHairGroups)
 	{
 		if (SourceDesc.GroupIndex == GroupIndex && SourceDesc.LODIndex == LODIndex)
 		{
+			if (SourceIndex) { *SourceIndex = LocalSourceIndex; }
 			return &SourceDesc;
 		}
-		++SourceIndex;
+		++LocalSourceIndex;
 	}
-	SourceIndex = -1;
+	if (SourceIndex) { *SourceIndex = -1; }
 	return nullptr;
 }
 
 template<typename T>
-static T* GetSourceDescription(TArray<T>& InHairGroups, uint32 GroupIndex, uint32 LODIndex, int32& SourceIndex)
+static T* GetSourceDescription(TArray<T>& InHairGroups, uint32 GroupIndex, uint32 LODIndex, int32* SourceIndex=nullptr)
 {
-	SourceIndex = 0;
+	int32 LocalSourceIndex = 0;
 	for (T& SourceDesc : InHairGroups)
 	{
 		if (SourceDesc.GroupIndex == GroupIndex && SourceDesc.LODIndex == LODIndex)
 		{
+			if (SourceIndex) { *SourceIndex = LocalSourceIndex; }
 			return &SourceDesc;
 		}
-		++SourceIndex;
+		++LocalSourceIndex;
 	}
-	SourceIndex = -1;
+	if (SourceIndex) { *SourceIndex = -1; }
 	return nullptr;
 }
 
@@ -498,7 +503,7 @@ bool IsHairStrandsAssetLoadingEnable()
 // * Clusters (bulk) data
 //
 // Guides, strands, and interpolation data are all transient
-static bool BuildHairGroup(
+static bool BuildHairGroup_Strands(
 	const int32 GroupIndex,
 	const bool bNeedInterpolationData,
 	const FHairDescriptionGroups& HairDescriptionGroups, 
@@ -538,33 +543,6 @@ static bool BuildHairGroup(
 	return bIsValid;
 }
 
-// Build: 
-// * Clusters (bulk) data
-//
-// Guides and strands data are all transient
-static bool BuildHairGroupCluster(
-	const int32 GroupIndex,
-	const FHairDescriptionGroups& InHairDescriptionGroups,
-	const TArray<FHairGroupsInterpolation>& InHairGroupsInterpolation,
-	const TArray<FHairGroupsLOD>& InHairGroupsLOD,
-	TArray<FHairGroupInfoWithVisibility>& OutHairGroupsInfo,
-	TArray<FHairGroupPlatformData>& OutHairGroupsData)
-{
-	const bool bIsValid = InHairDescriptionGroups.IsValid();
-	if (bIsValid)
-	{
-		const FHairDescriptionGroup& HairGroup = InHairDescriptionGroups.HairGroups[GroupIndex];
-		check(GroupIndex <= InHairDescriptionGroups.HairGroups.Num());
-		check(GroupIndex == HairGroup.Info.GroupID);
-
-		FHairStrandsDatas StrandsData;
-		FHairStrandsDatas GuidesData;
-		FGroomBuilder::BuildData(HairGroup, InHairGroupsInterpolation[GroupIndex], OutHairGroupsInfo[GroupIndex], StrandsData, GuidesData);
-		FGroomBuilder::BuildClusterBulkData(StrandsData, InHairDescriptionGroups.Bounds.SphereRadius, InHairGroupsLOD[GroupIndex], OutHairGroupsData[GroupIndex].Strands.ClusterBulkData);
-	}
-	return bIsValid;
-}
-
 /////////////////////////////////////////////////////////////////////////////////////////
 
 uint8 UGroomAsset::GenerateClassStripFlags(FArchive& Ar)
@@ -1016,7 +994,7 @@ void UGroomAsset::UpdateResource()
 		AllChangeType = AllChangeType | ChangeType;
 
 		check(GroupIndex < uint32(GetNumHairGroups()));
-		FHairGroupPlatformData& GroupData = GetHairGroupsPlatformData()[GroupIndex];
+		FHairGroupResources& GroupData = GetHairGroupsResources()[GroupIndex];
 		if (GroupData.Guides.RestResource)
 		{
 			InternalUpdateResource(GroupData.Guides.RestResource);
@@ -1025,33 +1003,11 @@ void UGroomAsset::UpdateResource()
 		if (IsHairStrandsEnabled(EHairStrandsShaderType::Strands))
 		{
 			InternalUpdateResource(GroupData.Strands.RestResource);
-			if (GroupData.Strands.InterpolationResource)
-			{
-				InternalUpdateResource(GroupData.Strands.InterpolationResource);
-			}
-
-			if ((ChangeType & GroomChangeType_LOD) && GroupData.Strands.HasValidData() && CanRebuildFromDescription())
-			{
-				const FHairDescriptionGroups& LocalHairDescriptionGroups = GetHairDescriptionGroups();
-
-				// Force rebuilding the LOD data as the LOD settings has changed
-				BuildHairGroupCluster(
-					GroupIndex,
-					LocalHairDescriptionGroups,
-					GetHairGroupsInterpolation(),
-					GetHairGroupsLOD(),
-					GetHairGroupsInfo(),
-					GetHairGroupsPlatformData());
-
-				GroupData.Strands.ClusterResource = new FHairStrandsClusterResource(GroupData.Strands.ClusterBulkData, FHairResourceName(GetFName(), GroupIndex), GetAssetPathName());
-			}
-			else
-			{
-				InternalUpdateResource(GroupData.Strands.ClusterResource);
-			}
+			InternalUpdateResource(GroupData.Strands.InterpolationResource);
+			InternalUpdateResource(GroupData.Strands.ClusterResource);
 
 			#if RHI_RAYTRACING
-			if (IsHairRayTracingEnabled() && GroupData.Strands.RaytracingResource)
+			if (IsHairRayTracingEnabled())
 			{
 				// Release raytracing resources, so that it will be recreate by the component.
 				InternalReleaseResource(GroupData.Strands.RaytracingResource);
@@ -1105,20 +1061,18 @@ void UGroomAsset::ReleaseResource()
 
 void UGroomAsset::ReleaseGuidesResource(uint32 GroupIndex)
 {
-	if (GroupIndex >= uint32(GetHairGroupsPlatformData().Num())) return;
-	FHairGroupPlatformData& GroupData = GetHairGroupsPlatformData()[GroupIndex];
-	if (GroupData.Guides.IsValid())
+	if (GetHairGroupsResources().IsValidIndex(GroupIndex))
 	{
+		FHairGroupResources& GroupData = GetHairGroupsResources()[GroupIndex];
 		InternalReleaseResource(GroupData.Guides.RestResource);
 	}
 }
 
 void UGroomAsset::ReleaseStrandsResource(uint32 GroupIndex)
 {
-	if (GroupIndex >= uint32(GetHairGroupsPlatformData().Num())) return;
-	FHairGroupPlatformData& GroupData = GetHairGroupsPlatformData()[GroupIndex];
-	if (GroupData.Strands.IsValid())
+	if (GetHairGroupsResources().IsValidIndex(GroupIndex))
 	{
+		FHairGroupResources& GroupData = GetHairGroupsResources()[GroupIndex];
 		InternalReleaseResource(GroupData.Strands.RestResource);
 		InternalReleaseResource(GroupData.Strands.ClusterResource);
 		InternalReleaseResource(GroupData.Strands.InterpolationResource);
@@ -1130,30 +1084,34 @@ void UGroomAsset::ReleaseStrandsResource(uint32 GroupIndex)
 
 void UGroomAsset::ReleaseCardsResource(uint32 GroupIndex)
 {
-	if (GroupIndex >= uint32(GetHairGroupsPlatformData().Num())) return;
-	FHairGroupPlatformData& GroupData = GetHairGroupsPlatformData()[GroupIndex];
-	for (FHairGroupPlatformData::FCards::FLOD& LOD : GroupData.Cards.LODs)
+	if (GetHairGroupsResources().IsValidIndex(GroupIndex))
 	{
-		InternalReleaseResource(LOD.RestResource);
-		InternalReleaseResource(LOD.InterpolationResource);
-		InternalReleaseResource(LOD.Guides.RestResource);
-		InternalReleaseResource(LOD.Guides.InterpolationResource);
-		#if RHI_RAYTRACING
-		InternalReleaseResource(LOD.RaytracingResource);
-		#endif
+		FHairGroupResources& GroupData = GetHairGroupsResources()[GroupIndex];
+		for (FHairGroupResources::FCards::FLOD& LOD : GroupData.Cards.LODs)
+		{
+			InternalReleaseResource(LOD.RestResource);
+			InternalReleaseResource(LOD.InterpolationResource);
+			InternalReleaseResource(LOD.GuideRestResource);
+			InternalReleaseResource(LOD.GuideInterpolationResource);
+			#if RHI_RAYTRACING
+			InternalReleaseResource(LOD.RaytracingResource);
+			#endif
+		}
 	}
 }
 
 void UGroomAsset::ReleaseMeshesResource(uint32 GroupIndex)
 {
-	if (GroupIndex >= uint32(GetHairGroupsPlatformData().Num())) return;
-	FHairGroupPlatformData& GroupData = GetHairGroupsPlatformData()[GroupIndex];
-	for (FHairGroupPlatformData::FMeshes::FLOD& LOD : GroupData.Meshes.LODs)
+	if (GetHairGroupsResources().IsValidIndex(GroupIndex))
 	{
-		InternalReleaseResource(LOD.RestResource);
-		#if RHI_RAYTRACING
-		InternalReleaseResource(LOD.RaytracingResource);
-		#endif
+		FHairGroupResources& GroupData = GetHairGroupsResources()[GroupIndex];
+		for (FHairGroupResources::FMeshes::FLOD& LOD : GroupData.Meshes.LODs)
+		{
+			InternalReleaseResource(LOD.RestResource);
+			#if RHI_RAYTRACING
+			InternalReleaseResource(LOD.RaytracingResource);
+			#endif
+		}
 	}
 }
 
@@ -1417,36 +1375,6 @@ void UGroomAsset::PostLoad()
 
 	check(GetNumHairGroups() > 0);
 
-	// Build hair strands if needed
-#if WITH_EDITORONLY_DATA
-	if (!bSucceed && IsHairStrandsEnabled(EHairStrandsShaderType::Strands) && CanRebuildFromDescription())
-	{
-		const FHairDescriptionGroups& LocalHairDescriptionGroups = GetHairDescriptionGroups();
-
-		const int32 GroupCount = GetHairGroupsInterpolation().Num();
-		check(LocalHairDescriptionGroups.HairGroups.Num() == GroupCount);
-		for (int32 GroupIndex = 0; GroupIndex < GroupCount; ++GroupIndex)
-		{
-			const bool bNeedInterpoldationData = NeedsInterpolationData(GroupIndex);
-			const bool bNeedToBuildData = 
-				(GetHairGroupsPlatformData()[GroupIndex].Guides.BulkData.GetNumCurves() == 0 ||
-				 (bNeedInterpoldationData && GetHairGroupsPlatformData()[GroupIndex].Strands.InterpolationBulkData.GetPointCount() == 0)) &&
-				 GetHairGroupsPlatformData()[GroupIndex].Strands.BulkData.GetNumCurves() > 0; // Empty groom has no data to build
-			if (bNeedToBuildData)
-			{
-				BuildHairGroup(
-					GroupIndex,
-					bNeedInterpoldationData,
-					LocalHairDescriptionGroups,
-					GetHairGroupsInterpolation(),
-					GetHairGroupsLOD(),
-					GetHairGroupsInfo(),
-					GetHairGroupsPlatformData());
-			}
-		}
-	}
-#endif
-
 	// Convert legacy version
 	{
 		bool bNeedSaving = false;
@@ -1470,10 +1398,14 @@ void UGroomAsset::PostLoad()
 		}
 	}
 
+	// * When running with the editor, InitResource is called in CacheDerivedDatas
+	// * When running without the editor, InitResource is explicitely called here
+#if !WITH_EDITOR
 	if (!IsTemplate() && IsHairStrandsAssetLoadingEnable())
 	{
 		InitResources();
 	}
+#endif
 
 	check(AreGroupsValid());
 
@@ -1502,84 +1434,6 @@ void UGroomAsset::PostLoad()
 	}
 }
 
-void UGroomAsset::PreSave(const class ITargetPlatform* TargetPlatform)
-{
-	PRAGMA_DISABLE_DEPRECATION_WARNINGS;
-	Super::PreSave(TargetPlatform);
-	PRAGMA_ENABLE_DEPRECATION_WARNINGS;
-}
-
-void UGroomAsset::PreSave(FObjectPreSaveContext ObjectSaveContext)
-{
-#if WITH_EDITORONLY_DATA
-	check(AreGroupsValid());
-
-	const uint32 GroupCount = GetNumHairGroups();
-	uint32 ChangeType = 0;
-
-	for (uint32 GroupIt = 0; GroupIt < GroupCount; ++GroupIt)
-	{
-		if (!(CachedHairGroupsInterpolation[GroupIt] == GetHairGroupsInterpolation()[GroupIt]))
-		{
-			ChangeType = ChangeType | GroomChangeType_Interpolation;
-			break;
-		}
-	}
-
-	for (uint32 GroupIt = 0; GroupIt < GroupCount; ++GroupIt)
-	{
-		if (!(CachedHairGroupsLOD[GroupIt] == GetHairGroupsLOD()[GroupIt]))
-		{
-			ChangeType = ChangeType | GroomChangeType_LOD;
-			break;
-		}
-	}
-
-	if ((ChangeType & GroomChangeType_Interpolation) || (ChangeType & GroomChangeType_LOD))
-	{
-		if (CanRebuildFromDescription())
-		{
-			const FHairDescriptionGroups &LocalHairDescriptionGroups = GetHairDescriptionGroups();
-
-			FGroomComponentRecreateRenderStateContext RecreateRenderContext(this);
-			for (uint32 GroupIt = 0; GroupIt < GroupCount; ++GroupIt)
-			{
-				const bool bHasInterpolationChanged = !(CachedHairGroupsInterpolation[GroupIt] == GetHairGroupsInterpolation()[GroupIt]);
-				const bool bHasLODChanged = !(CachedHairGroupsLOD[GroupIt] == GetHairGroupsLOD()[GroupIt]);
-				const bool bNeedInterpoldationData = NeedsInterpolationData(GroupIt);
-
-				if (bHasInterpolationChanged)
-				{
-					BuildHairGroup(
-						GroupIt,
-						bNeedInterpoldationData,
-						LocalHairDescriptionGroups,
-						GetHairGroupsInterpolation(),
-						GetHairGroupsLOD(),
-						GetHairGroupsInfo(),
-						GetHairGroupsPlatformData());
-				}
-				else if (bHasLODChanged)
-				{
-					BuildHairGroupCluster(
-						GroupIt,
-						LocalHairDescriptionGroups,
-						GetHairGroupsInterpolation(),
-						GetHairGroupsLOD(),
-						GetHairGroupsInfo(),
-						GetHairGroupsPlatformData());
-				}
-			}
-		}
-		InitResources();
-	}
-
-	UpdateHairGroupsInfo();
-	UpdateCachedSettings();
-#endif
-	Super::PreSave(ObjectSaveContext);
-}
-
 void UGroomAsset::BeginDestroy()
 {
 	ReleaseResource();
@@ -1836,9 +1690,9 @@ void InternalSerializeCards(FArchive& Ar, UObject* Owner, TArray<FHairGroupPlatf
 	{
 		LODs.SetNum(Count);
 	}
-	for (uint32 MeshIt = 0; MeshIt < Count; ++MeshIt)
+	for (uint32 CardIt = 0; CardIt < Count; ++CardIt)
 	{
-		InternalSerializeCard(Ar, Owner, LODs[MeshIt]);
+		InternalSerializeCard(Ar, Owner, LODs[CardIt]);
 	}
 }
 
@@ -1864,22 +1718,21 @@ static void InternalSerializeCard(FArchive& Ar, UObject* Owner, FHairGroupPlatfo
 	{
 		CardLODData.BulkData.Serialize(Ar);
 		CardLODData.InterpolationBulkData.Serialize(Ar);
-
-		CardLODData.Guides.BulkData.Serialize(Ar, Owner);
-		CardLODData.Guides.InterpolationBulkData.Serialize(Ar, Owner);
+		CardLODData.GuideBulkData.Serialize(Ar, Owner);
+		CardLODData.GuideInterpolationBulkData.Serialize(Ar, Owner);
 	}
 	else
 	{
 		// LOD has been marked to be cooked out so serialize empty data
 		FHairCardsBulkData NoCardsBulkData;
-		NoCardsBulkData.Serialize(Ar);
-
 		FHairCardsInterpolationBulkData NoInterpolationBulkData;
-		NoInterpolationBulkData.Serialize(Ar);
+		FHairStrandsBulkData NoGuideBulkData;
+		FHairStrandsInterpolationBulkData NoGuideInteprolationBulkData;
 
-		FHairGroupPlatformData::FBaseWithInterpolation NoGuideData;
-		NoGuideData.BulkData.Serialize(Ar, Owner);
-		NoGuideData.InterpolationBulkData.Serialize(Ar, Owner);
+		NoCardsBulkData.Serialize(Ar);
+		NoInterpolationBulkData.Serialize(Ar);
+		NoGuideBulkData.Serialize(Ar, Owner);
+		NoGuideInteprolationBulkData.Serialize(Ar, Owner);
 	}
 }
 
@@ -2057,6 +1910,7 @@ void UGroomAsset::SetNumGroup(uint32 InGroupCount, bool bResetGroupData, bool bR
 	if (InGroupCount != GetHairGroupsPlatformData().Num())
 	{
 		GetHairGroupsPlatformData().Init(FHairGroupPlatformData(), InGroupCount);
+		GetHairGroupsResources().Init(FHairGroupResources(), InGroupCount);
 	}
 
 	if (InGroupCount != GetHairGroupsPhysics().Num())
@@ -2334,14 +2188,13 @@ bool UGroomAsset::IsFullyCached()
 		const uint32 LODCount = GetHairGroupsLOD()[GroupIndex].LODs.Num();
 		for (uint32 LODIt = 0; LODIt < LODCount; ++LODIt)
 		{
-			int32 SourceIt = 0;
 			// GetSourceDescription will cross-check the LOD settings with the cards/meshes settings to see if they would produce any data
-			if (const FHairGroupsCardsSourceDescription* Desc = GetSourceDescription(GetHairGroupsCards(), GroupIndex, LODIt, SourceIt))
+			if (const FHairGroupsCardsSourceDescription* Desc = GetSourceDescription(GetHairGroupsCards(), GroupIndex, LODIt))
 			{
 				Desc->HasMeshChanged(); // this query will trigger a load of the mesh dependency, which has to be done in the game thread
 				bHasCardsLODs = true;
 			}
-			if (const FHairGroupsMeshesSourceDescription* Desc = GetSourceDescription(GetHairGroupsMeshes(), GroupIndex, LODIt, SourceIt))
+			if (const FHairGroupsMeshesSourceDescription* Desc = GetSourceDescription(GetHairGroupsMeshes(), GroupIndex, LODIt))
 			{
 				Desc->HasMeshChanged(); // this query will trigger a load of the mesh dependency, which has to be done in the game thread
 				bHasMeshesLODs = true;
@@ -2468,21 +2321,24 @@ bool UGroomAsset::CacheDerivedDatas()
 	FGroomComponentRecreateRenderStateContext RecreateContext(bIsGameThread ? this : nullptr);
 
 	const uint32 GroupCount = GetHairGroupsInterpolation().Num();
+	bool bSucceed = true;
 	for (uint32 GroupIndex = 0; GroupIndex < GroupCount; ++GroupIndex)
 	{
-		bool bSucceed = CacheDerivedData(GroupIndex);
-		if (!bSucceed)
-			return false;
+		bSucceed = bSucceed && CacheDerivedData(GroupIndex);
 	}
-	UpdateCachedSettings();
-	UpdateHairGroupsInfo();
 
-	if (bIsGameThread)
+	if (bSucceed)
 	{
-		ReleaseResource();
-		InitResources();
+		UpdateCachedSettings();
+		UpdateHairGroupsInfo();
+	
+		if (bIsGameThread)
+		{
+			ReleaseResource();
+			InitResources();
+		}
 	}
-	return true;
+	return bSucceed;
 }
 
 bool UGroomAsset::CacheDerivedData(uint32 GroupIndex)
@@ -2495,7 +2351,7 @@ bool UGroomAsset::CacheDerivedData(uint32 GroupIndex)
 
 	const uint32 GroupCount = GetHairGroupsInterpolation().Num();
 	check(GroupIndex < GroupCount);
-	if (GroupIndex >= GroupCount)
+	if (!GetHairGroupsInterpolation().IsValidIndex(GroupIndex))
 	{
 		return false;
 	}
@@ -2592,7 +2448,7 @@ bool UGroomAsset::CacheStrandsData(uint32 GroupIndex, FString& OutDerivedDataKey
 		}
 
 		// Build groom data with the new build settings
-		bSuccess = BuildHairGroup(
+		bSuccess = BuildHairGroup_Strands(
 			GroupIndex,
 			NeedsInterpolationData(GroupIndex),
 			LocalHairDescriptionGroups,
@@ -2656,18 +2512,17 @@ bool UGroomAsset::CacheStrandsData(uint32 GroupIndex, FString& OutDerivedDataKey
 	return bSuccess;
 }
 
-bool UGroomAsset::HasValidMeshesData(uint32 GroupIndex) const
+bool UGroomAsset::HasValidData_Cards(uint32 GroupIndex) const
 {
-	const FHairGroupPlatformData& HairGroupData = GetHairGroupsPlatformData()[GroupIndex];
 	const FHairGroupsLOD& GroupsLOD = GetHairGroupsLOD()[GroupIndex];
 	for (int32 LODIt = 0; LODIt < GroupsLOD.LODs.Num(); ++LODIt)
 	{
-		if (GroupsLOD.LODs[LODIt].GeometryType == EGroomGeometryType::Meshes)
+		// Note: We don't condition/filter resource allocation based on GeometryType == EGroomGeometryType::Cards, because 
+		// geometry type can be switched at runtime (between cards/strands)
 		{
-			int32 SourceIt = 0;
-			if (const FHairGroupsMeshesSourceDescription* Desc = GetSourceDescription(GetHairGroupsMeshes(), GroupIndex, LODIt, SourceIt))
+			if (const FHairGroupsCardsSourceDescription* Desc = GetSourceDescription(GetHairGroupsCards(), GroupIndex, LODIt))
 			{
-				if (Desc->ImportedMesh)
+				if (Desc->GetMesh())
 				{
 					return true;
 				}
@@ -2677,19 +2532,16 @@ bool UGroomAsset::HasValidMeshesData(uint32 GroupIndex) const
 	return false;
 }
 
-bool UGroomAsset::HasValidCardsData(uint32 GroupIndex) const
+bool UGroomAsset::HasValidData_Meshes(uint32 GroupIndex) const
 {
-	const FHairGroupPlatformData& HairGroupData = GetHairGroupsPlatformData()[GroupIndex];
 	const FHairGroupsLOD& GroupsLOD = GetHairGroupsLOD()[GroupIndex];
 	for (int32 LODIt = 0; LODIt < GroupsLOD.LODs.Num(); ++LODIt)
 	{
-		// Note: We don't condition/filter resource allocation based on GeometryType == EGroomGeometryType::Cards, because 
-		// geometry type can be switched at runtime (between cards/strands)
+		if (GroupsLOD.LODs[LODIt].GeometryType == EGroomGeometryType::Meshes)
 		{
-			int32 SourceIt = 0;
-			if (const FHairGroupsCardsSourceDescription* Desc = GetSourceDescription(GetHairGroupsCards(), GroupIndex, LODIt, SourceIt))
+			if (const FHairGroupsMeshesSourceDescription* Desc = GetSourceDescription(GetHairGroupsMeshes(), GroupIndex, LODIt))
 			{
-				if (Desc->GetMesh())
+				if (Desc->ImportedMesh)
 				{
 					return true;
 				}
@@ -2702,35 +2554,25 @@ bool UGroomAsset::HasValidCardsData(uint32 GroupIndex) const
 // Return true if the meshes data have changed
 bool UGroomAsset::CacheCardsData(uint32 GroupIndex, const FString& StrandsKey)
 {
-	const bool bHasValidCardsData = HasValidCardsData(GroupIndex);
-	if (!bHasValidCardsData)
+	// We should never query the DDC for something that we won't be able to build. So check that first.
+	// The goal is to avoid paying the price of a guaranteed failure on high-latency DDC backends every run.
+	const bool bHasValidSourceData = HasValidData_Cards(GroupIndex);
+	if (!bHasValidSourceData)
 	{
 		const bool bNeedReset = !CardsDerivedDataKey[GroupIndex].IsEmpty();
 		CardsDerivedDataKey[GroupIndex].Empty();
+
+		// Initialized LODs
+		const uint32 LODCount = GetHairGroupsLOD()[GroupIndex].LODs.Num();
+		GetHairGroupsPlatformData()[GroupIndex].Cards.LODs.SetNum(LODCount);
+
 		return bNeedReset;
 	}
 
 	const FString KeySuffix = GroomDerivedDataCacheUtils::BuildCardsDerivedDataKeySuffix(GroupIndex, GetHairGroupsLOD()[GroupIndex].LODs, GetHairGroupsCards());
 	const FString DerivedDataKey = StrandsKey + KeySuffix;
 
-	FHairGroupPlatformData& HairGroupData = GetHairGroupsPlatformData()[GroupIndex];
-
-	// We should never query the DDC for something that we won't be able to build. So check that first.
-	// The goal is to avoid paying the price of a guaranteed failure on high-latency DDC backends every run.
-	bool bDataCanBeBuilt = false;
-
-	const FHairGroupsLOD& GroupsLOD = GetHairGroupsLOD()[GroupIndex];
-	for (int32 LODIt = 0; LODIt < GroupsLOD.LODs.Num(); ++LODIt)
-	{
-		int32 SourceIt = 0;
-		if (FHairGroupsCardsSourceDescription* Desc = GetSourceDescription(GetHairGroupsCards(), GroupIndex, LODIt, SourceIt))
-		{
-			// Note: We don't condition/filter resource allocation based on GeometryType == EGroomGeometryType::Cards, because 
-			// geometry type can be switched at runtime (between cards/strands)
-			bDataCanBeBuilt |= Desc->ImportedMesh != nullptr;
-		}
-	}
-
+	// Query DDC
 	using namespace UE::DerivedData;
 	const FCacheKey Key = ConvertLegacyCacheKey(DerivedDataKey);
 	const FSharedString Name = MakeStringView(GetPathName());
@@ -2744,6 +2586,9 @@ bool UGroomAsset::CacheCardsData(uint32 GroupIndex, const FString& StrandsKey)
 		Owner.Wait();
 	}
 
+	// * If data is available, populate LOD data
+	// * Otherwise build the LOD data
+	FHairGroupPlatformData& HairGroupData = GetHairGroupsPlatformData()[GroupIndex];
 	if (Data)
 	{
 		UE_CLOG(IsHairStrandsDDCLogEnable(), LogHairStrands, Log, TEXT("[Groom/DDC] Cards - Found (Groom:%s Group:%d)."), *GetName(), GroupIndex);
@@ -2755,7 +2600,8 @@ bool UGroomAsset::CacheCardsData(uint32 GroupIndex, const FString& StrandsKey)
 	{
 		UE_CLOG(IsHairStrandsDDCLogEnable(), LogHairStrands, Log, TEXT("[Groom/DDC] Cards - Not found (Groom:%s Group:%d)."), *GetName(), GroupIndex);
 
-		if (!BuildCardsData(GroupIndex))
+		const bool bSucceed = BuildHairGroup_Cards(GroupIndex);
+		if (!bSucceed)
 		{
 			CardsDerivedDataKey[GroupIndex].Empty();
 			return false;
@@ -2769,29 +2615,11 @@ bool UGroomAsset::CacheCardsData(uint32 GroupIndex, const FString& StrandsKey)
 		AsyncOwner.KeepAlive();
 	}
 
-	// Handle the case where the cards data is already cached in the DDC
-	// Need to populate the strands data with it
-	FHairGroupPlatformData& GroupData = GetHairGroupsPlatformData()[GroupIndex];
-	if (StrandsDerivedDataKey[GroupIndex].IsEmpty() && GroupData.Strands.BulkData.GetNumPoints() == 0)
-	{
-		for (int32 LODIt = 0; LODIt < GroupData.Cards.LODs.Num(); ++LODIt)
-		{
-			int32 SourceIt = 0;
-			if (FHairGroupsCardsSourceDescription* Desc = GetSourceDescription(GetHairGroupsCards(), GroupIndex, LODIt, SourceIt))
-			{
-				FHairGroupPlatformData::FCards::FLOD& LOD = GroupData.Cards.LODs[LODIt];
-				GroupData.Strands.BulkData = LOD.Guides.BulkData;
-				GroupData.Guides.BulkData = LOD.Guides.BulkData;
-				break;
-			}
-		}
-	}
-
 	// Update the imported mesh DDC keys
+	FHairGroupPlatformData& GroupData = GetHairGroupsPlatformData()[GroupIndex];
 	for (int32 LODIt = 0; LODIt < GroupData.Cards.LODs.Num(); ++LODIt)
 	{
-		int32 SourceIt = 0;
-		if (FHairGroupsCardsSourceDescription* Desc = GetSourceDescription(GetHairGroupsCards(), GroupIndex, LODIt, SourceIt))
+		if (FHairGroupsCardsSourceDescription* Desc = GetSourceDescription(GetHairGroupsCards(), GroupIndex, LODIt))
 		{
 			Desc->UpdateMeshKey();
 		}
@@ -2805,36 +2633,25 @@ bool UGroomAsset::CacheCardsData(uint32 GroupIndex, const FString& StrandsKey)
 // Return true if the meshes data have changed
 bool UGroomAsset::CacheMeshesData(uint32 GroupIndex)
 {
-	const bool bHasValidMeshesData = UGroomAsset::HasValidMeshesData(GroupIndex);
-	if (!bHasValidMeshesData)
+	// We should never query the DDC for something that we won't be able to build. So check that first.
+	// The goal is to avoid paying the price of a guaranteed failure on high-latency DDC backends every run.
+	const bool bHasValidSourceData = UGroomAsset::HasValidData_Meshes(GroupIndex);
+	if (!bHasValidSourceData)
 	{
 		const bool bNeedReset = !MeshesDerivedDataKey[GroupIndex].IsEmpty();
 		MeshesDerivedDataKey[GroupIndex].Empty();
+
+		// Initialized LODs
+		const uint32 LODCount = GetHairGroupsLOD()[GroupIndex].LODs.Num();
+		GetHairGroupsPlatformData()[GroupIndex].Meshes.LODs.SetNum(LODCount);
+
 		return bNeedReset;
 	}
 
 	const FString KeySuffix = GroomDerivedDataCacheUtils::BuildMeshesDerivedDataKeySuffix(GroupIndex, GetHairGroupsLOD()[GroupIndex].LODs, GetHairGroupsMeshes());
 	const FString DerivedDataKey = GroomDerivedDataCacheUtils::BuildGroomDerivedDataKey(KeySuffix);
 
-	FHairGroupPlatformData& HairGroupData = GetHairGroupsPlatformData()[GroupIndex];
-
-	// We should never query the DDC for something that we won't be able to build. So check that first.
-	// The goal is to avoid paying the price of a guaranteed failure on high-latency DDC backends every run.
-	bool bDataCanBeBuilt = false;
-
-	const FHairGroupsLOD& GroupsLOD = GetHairGroupsLOD()[GroupIndex];
-	for (int32 LODIt = 0; LODIt < GroupsLOD.LODs.Num(); ++LODIt)
-	{
-		int32 SourceIt = 0;
-		if (FHairGroupsMeshesSourceDescription* Desc = GetSourceDescription(GetHairGroupsMeshes(), GroupIndex, LODIt, SourceIt))
-		{
-			if (GroupsLOD.LODs[LODIt].GeometryType == EGroomGeometryType::Meshes)
-			{
-				bDataCanBeBuilt |= Desc->ImportedMesh != nullptr;
-			}
-		}
-	}
-
+	// Query DDC
 	using namespace UE::DerivedData;
 	const FCacheKey Key = ConvertLegacyCacheKey(DerivedDataKey);
 	const FSharedString Name = MakeStringView(GetPathName());
@@ -2848,8 +2665,11 @@ bool UGroomAsset::CacheMeshesData(uint32 GroupIndex)
 		Owner.Wait();
 	}
 
+	// * If data is available, populate LOD data
+	// * Otherwise build the LOD data
+	FHairGroupPlatformData& HairGroupData = GetHairGroupsPlatformData()[GroupIndex];
 	TArray<uint8> DerivedData;
-	if (bDataCanBeBuilt && GetDerivedDataCacheRef().GetSynchronous(*DerivedDataKey, DerivedData, GetPathName()))
+	if (GetDerivedDataCacheRef().GetSynchronous(*DerivedDataKey, DerivedData, GetPathName()))
 	{
 		UE_CLOG(IsHairStrandsDDCLogEnable(), LogHairStrands, Log, TEXT("[Groom/DDC] Meshes - Found (Groom:%s Group:%d)."), *GetName(), GroupIndex);
 
@@ -2860,7 +2680,8 @@ bool UGroomAsset::CacheMeshesData(uint32 GroupIndex)
 	{
 		UE_CLOG(IsHairStrandsDDCLogEnable(), LogHairStrands, Log, TEXT("[Groom/DDC] Meshes - Not found (Groom:%s Group:%d)."), *GetName(), GroupIndex);
 
-		if (!BuildMeshesData(GroupIndex))
+		const bool bSucceed = BuildHairGroup_Meshes(GroupIndex);
+		if (!bSucceed)
 		{
 			MeshesDerivedDataKey[GroupIndex].Empty();
 			return false;
@@ -2951,11 +2772,9 @@ static bool InternalImportCardGeometry(
 	return FHairCardsBuilder::ImportGeometry(CardsMesh, TempHairGuidesData, TempHairStrandsData, TempHairStrandsVoxelData, InDesc->GuideType == EHairCardsGuideType::Generated, OutBulkData, OutGuidesData, OutInterpolationBulkData);
 }
 
-bool UGroomAsset::BuildCardsData(uint32 GroupIndex)
+bool UGroomAsset::HasChanged_Cards(uint32 GroupIndex, TArray<bool>& OutIsValid) const 
 {
-	LLM_SCOPE_BYTAG(Groom);
-
-	if (!IsHairStrandsEnabled(EHairStrandsShaderType::Cards) || GetHairGroupsCards().Num() == 0 || !CanRebuildFromDescription())
+	if (!IsHairStrandsEnabled(EHairStrandsShaderType::Cards) || (GetHairGroupsCards().Num() == 0 && CachedHairGroupsCards.Num() == 0)|| !CanRebuildFromDescription())
 	{
 		return false;
 	}
@@ -2963,89 +2782,60 @@ bool UGroomAsset::BuildCardsData(uint32 GroupIndex)
 	bool bHasChanged = GetHairGroupsCards().Num() != CachedHairGroupsCards.Num();
 
 	check(GroupIndex < uint32(GetNumHairGroups()));
-	FHairGroupPlatformData& GroupData = GetHairGroupsPlatformData()[GroupIndex];
+	const FHairGroupPlatformData& GroupData = GetHairGroupsPlatformData()[GroupIndex];
 
 	// The settings might have been previously cached without the data having been built
 	const uint32 LODCount = GetHairGroupsLOD()[GroupIndex].LODs.Num();
-	TArray<bool> bIsAlreadyBuilt;
-	bIsAlreadyBuilt.SetNum(LODCount);
+	OutIsValid.Init(false, LODCount);
 	for (uint32 LODIt = 0; LODIt < LODCount; ++LODIt)
 	{
 		int32 SourceIt = 0;
-		if (const FHairGroupsCardsSourceDescription* Desc = GetSourceDescription(GetHairGroupsCards(), GroupIndex, LODIt, SourceIt))
+		if (const FHairGroupsCardsSourceDescription* Desc = GetSourceDescription(GetHairGroupsCards(), GroupIndex, LODIt, &SourceIt))
 		{
-			bIsAlreadyBuilt[LODIt] = GroupData.Strands.HasValidData() && Desc->CardsInfo.NumCardVertices > 0 && !Desc->HasMeshChanged();
-			bHasChanged |= !bIsAlreadyBuilt[LODIt];
-		}
-	}
-
-	if (!bHasChanged)
-	{
-		for (uint32 SourceIt = 0, SourceCount = GetHairGroupsCards().Num(); SourceIt < SourceCount; ++SourceIt)
-		{
-			const bool bEquals = CachedHairGroupsCards[SourceIt] == GetHairGroupsCards()[SourceIt];
-			if (!bEquals)
+			bool bIsLODValid = GroupData.Cards.LODs.IsValidIndex(LODIt) && GroupData.Cards.LODs[LODIt].HasValidData() && !Desc->HasMeshChanged();
+			if (const FHairGroupsCardsSourceDescription* CachedDesc = (CachedHairGroupsCards.IsValidIndex(SourceIt) ? &CachedHairGroupsCards[SourceIt] : nullptr))
 			{
-				bHasChanged = true;
-			}
-			else
-			{
-				// Reload/Update texture resources when the groom asset is saved, so that the textures 
-				// content is up to date with what has been saved.
-				if (IsInGameThread())
-				{
-					for (UTexture2D* Tex : GetHairGroupsCards()[SourceIt].Textures.Textures)
-					{
-						if (Tex != nullptr)
-						{
-							Tex->UpdateResource();
-						}
-					}
-				}
+				bIsLODValid = bIsLODValid && *CachedDesc == *Desc;
 			}
+			OutIsValid[LODIt] = bIsLODValid;
 		}
+
+		bHasChanged |= !OutIsValid[LODIt];
 	}
 
-	if (!bHasChanged)
+	return bHasChanged;
+}
+
+bool UGroomAsset::BuildHairGroup_Cards(uint32 GroupIndex)
+{
+	LLM_SCOPE_BYTAG(Groom);
+
+	TArray<bool> bIsAlreadyBuilt;
+	if (!HasChanged_Cards(GroupIndex, bIsAlreadyBuilt))
 	{
 		return false;
 	}
 
+	const uint32 LODCount = GetHairGroupsLOD()[GroupIndex].LODs.Num();
+	FHairGroupPlatformData& GroupData = GetHairGroupsPlatformData()[GroupIndex];
+
 	bool bDataBuilt = false;
 	GroupData.Cards.LODs.SetNum(LODCount);
-	const FHairGroupsLOD& GroupsLOD = GetHairGroupsLOD()[GroupIndex];
 	for (uint32 LODIt = 0; LODIt < LODCount; ++LODIt)
 	{
-		int32 SourceIt = 0;
-		if (FHairGroupsCardsSourceDescription* Desc = GetSourceDescription(GetHairGroupsCards(), GroupIndex, LODIt, SourceIt))
+		if (FHairGroupsCardsSourceDescription* Desc = GetSourceDescription(GetHairGroupsCards(), GroupIndex, LODIt))
 		{
-			// #hair_todo: add finer culling granularity to rebuild only what is necessary
-			const FHairGroupsCardsSourceDescription* CachedDesc = SourceIt < CachedHairGroupsCards.Num() ? &CachedHairGroupsCards[SourceIt] : nullptr;
-			const bool bLODHasChanged = CachedDesc == nullptr || !(*CachedDesc == *Desc);
-			if (!bLODHasChanged && bIsAlreadyBuilt[LODIt])
+			if (bIsAlreadyBuilt[LODIt])
 			{
-				bDataBuilt |= bIsAlreadyBuilt[LODIt];
+				bDataBuilt = true;
 				continue;
 			}
 
-			if (!IsInGameThread())
-			{
-				// Build needs to execute from the game thread
-				bRetryLoadFromGameThread = true;
-				return false;
-			}
-
 			FHairGroupPlatformData::FCards::FLOD& LOD = GroupData.Cards.LODs[LODIt];
 			LOD.BulkData.Reset();
 
-			// 0. Release geometry resources
-			InternalReleaseResource(LOD.RestResource);
-			InternalReleaseResource(LOD.InterpolationResource);
-			InternalReleaseResource(LOD.Guides.RestResource);
-			InternalReleaseResource(LOD.Guides.InterpolationResource);
-
 			// 1. Load geometry data, if any
-			bool bInitResources = false;
+			bool bSuccess = false;
 			FHairStrandsDatas LODGuidesData;
 			if (UStaticMesh* CardsMesh = Desc->ImportedMesh)
 			{
@@ -3053,36 +2843,23 @@ bool UGroomAsset::BuildCardsData(uint32 GroupIndex)
 
 				const FHairDescriptionGroups& LocalHairDescriptionGroups = GetHairDescriptionGroups();
 				const TArray<FHairGroupsInterpolation>& LocalInterpolationGroups= GetHairGroupsInterpolation();
-				bInitResources = InternalImportCardGeometry(Desc, LocalHairDescriptionGroups, LocalInterpolationGroups, LOD.BulkData, LODGuidesData, LOD.InterpolationBulkData);
-				if (!bInitResources)
+				bSuccess = InternalImportCardGeometry(Desc, LocalHairDescriptionGroups, LocalInterpolationGroups, LOD.BulkData, LODGuidesData, LOD.InterpolationBulkData);
+				if (!bSuccess)
 				{
 					UE_LOG(LogHairStrands, Warning, TEXT("Failed to import cards from %s for Group %d LOD %d."), *CardsMesh->GetName(), GroupIndex, LODIt);
 					LOD.BulkData.Reset();
 				}
 			}
 
-			bDataBuilt |= bInitResources;
+			bDataBuilt |= bSuccess;
 
 			// Clear the cards stats
 			Desc->CardsInfo.NumCardVertices = 0;
 			Desc->CardsInfo.NumCards = 0;
 
 			// 2. Init geometry & texture resources, and generate interpolation data/resources
-			if (bInitResources)
+			if (bSuccess)
 			{
-				FHairResourceName ResourceName(GetFName(), GroupIndex, LODIt);
-				const FName OwnerName = GetAssetPathName(LODIt);
-
-				LOD.RestResource = new FHairCardsRestResource(LOD.BulkData, ResourceName, OwnerName);
-				BeginInitResource(LOD.RestResource); // Immediate allocation, as needed for the vertex factory, input stream building
-
-				// 2.1 Load atlas textures
-				InitAtlasTexture(LOD.RestResource, Desc->Textures);
-				LOD.RestResource->bInvertUV = Desc->bInvertUV;
-				
-				// 2.2 Load interoplatino resources
-				LOD.InterpolationResource = new FHairCardsInterpolationResource(LOD.InterpolationBulkData, ResourceName, OwnerName);
-
 				// Create own interpolation settings for cards.
 				// Force closest guides as this is the most relevant matching metric for cards, due to their coarse geometry
 				FHairInterpolationSettings CardsInterpolationSettings = GetHairGroupsInterpolation()[GroupIndex].InterpolationSettings;
@@ -3092,19 +2869,11 @@ bool UGroomAsset::BuildCardsData(uint32 GroupIndex)
 				CardsInterpolationSettings.InterpolationDistance = EHairInterpolationWeight::Parametric;
 				CardsInterpolationSettings.InterpolationQuality = EHairInterpolationQuality::High;
 
-				// There could be no strands data when importing cards into an empty groom so get them from the card guides
-				bool bCopyRenderData = false;
-				if (GroupData.Guides.BulkData.GetNumPoints() == 0)
-				{
-					// The RenderData is filled out by BuildData
-					bCopyRenderData = true;
-				}
-
 				FHairGroupInfo DummyGroupInfo;
 
 				// 1. Build data & bulk data for the cards guides
 				FGroomBuilder::BuildData(LODGuidesData);
-				FGroomBuilder::BuildBulkData(DummyGroupInfo, LODGuidesData, LOD.Guides.BulkData);
+				FGroomBuilder::BuildBulkData(DummyGroupInfo, LODGuidesData, LOD.GuideBulkData);
 
 				// 2. (Re)Build data for the sim guides (since there are transient/no-cached)
 				const FHairDescriptionGroups& LocalHairDescriptionGroups = GetHairDescriptionGroups();
@@ -3116,25 +2885,14 @@ bool UGroomAsset::BuildCardsData(uint32 GroupIndex)
 				// 3. Compute interpolation data & bulk data between the sim guides and the cards guides
 				FHairStrandsInterpolationDatas LODGuideInterpolationData;
 				FGroomBuilder::BuildInterplationData(DummyGroupInfo, LODGuidesData, SimGuidesData, CardsInterpolationSettings, LODGuideInterpolationData);
-				FGroomBuilder::BuildInterplationBulkData(SimGuidesData, LODGuideInterpolationData, LOD.Guides.InterpolationBulkData);
-
-				if (bCopyRenderData)
-				{
-					GroupData.Strands.BulkData = LOD.Guides.BulkData;
-					GroupData.Guides.BulkData  = LOD.Guides.BulkData;
-				}
-
-				LOD.Guides.RestResource = new FHairStrandsRestResource(LOD.Guides.BulkData, EHairStrandsResourcesType::Cards, ResourceName, OwnerName);
-
-				LOD.Guides.InterpolationResource = new FHairStrandsInterpolationResource(LOD.Guides.InterpolationBulkData, ResourceName, OwnerName);
+				FGroomBuilder::BuildInterplationBulkData(SimGuidesData, LODGuideInterpolationData, LOD.GuideInterpolationBulkData);
 
 				// Update card stats to display
 				Desc->CardsInfo.NumCardVertices = LOD.BulkData.GetNumVertices();
-				Desc->CardsInfo.NumCards = LOD.Guides.BulkData.GetNumCurves();// LOD.Data.Cards.IndexOffsets.Num();
+				Desc->CardsInfo.NumCards 		= LOD.GuideBulkData.GetNumCurves();
 			}
 		}
 	}
-
 	return bDataBuilt;
 }
 
@@ -3148,15 +2906,6 @@ bool UGroomAsset::BuildCardsData()
 
 	if (bDataChanged)
 	{
-		InitGuideResources();
-
-		// When building cards in an empty groom, the strands resources need to be initialized once
-		FHairGroupPlatformData& GroupData = GetHairGroupsPlatformData()[0];
-		if (!GroupData.Strands.ClusterResource)
-		{
-			InitStrandsResources();
-		}
-
 		InitCardsResources();
 	}
 
@@ -3179,83 +2928,57 @@ bool UGroomAsset::BuildMeshesData()
 	return bDataChanged;
 }
 
-bool UGroomAsset::BuildMeshesData(uint32 GroupIndex)
+bool UGroomAsset::HasChanged_Meshes(uint32 GroupIndex, TArray<bool>& OutIsValid) const 
 {
-	if (!IsHairStrandsEnabled(EHairStrandsShaderType::Meshes) || GetHairGroupsMeshes().Num() == 0)
+	if (!IsHairStrandsEnabled(EHairStrandsShaderType::Meshes) || (GetHairGroupsMeshes().Num() == 0 && CachedHairGroupsMeshes.Num() == 0))
 	{
 		return false;
 	}
 
 	bool bHasChanged = GetHairGroupsMeshes().Num() != CachedHairGroupsMeshes.Num();
 
-	FHairGroupPlatformData& GroupData = GetHairGroupsPlatformData()[GroupIndex];
+	const FHairGroupPlatformData& GroupData = GetHairGroupsPlatformData()[GroupIndex];
 
 	// The settings might have been previously cached without the data having been built
 	const uint32 LODCount = GetHairGroupsLOD()[GroupIndex].LODs.Num();
-	TArray<bool> bIsAlreadyBuilt;
-	bIsAlreadyBuilt.SetNum(LODCount);
+	OutIsValid.Init(false, LODCount);
 	for (uint32 LODIt = 0; LODIt < LODCount; ++LODIt)
 	{
 		int32 SourceIt = 0;
-		if (const FHairGroupsMeshesSourceDescription* Desc = GetSourceDescription(GetHairGroupsMeshes(), GroupIndex, LODIt, SourceIt))
+		if (const FHairGroupsMeshesSourceDescription* Desc = GetSourceDescription(GetHairGroupsMeshes(), GroupIndex, LODIt, &SourceIt))
 		{
-			bIsAlreadyBuilt[LODIt] = GroupData.Meshes.LODs.IsValidIndex(LODIt) && GroupData.Meshes.LODs[LODIt].BulkData.GetNumVertices() > 0 && !Desc->HasMeshChanged();
-			bHasChanged |= !bIsAlreadyBuilt[LODIt];
-		}
-	}
+			bool bIsLODValid = GroupData.Meshes.LODs.IsValidIndex(LODIt) && GroupData.Meshes.LODs[LODIt].HasValidData() && !Desc->HasMeshChanged();
 
-	if (!bHasChanged)
-	{
-		for (uint32 SourceIt = 0, SourceCount = GetHairGroupsMeshes().Num(); SourceIt < SourceCount; ++SourceIt)
-		{
-			const bool bEquals = CachedHairGroupsMeshes[SourceIt] == GetHairGroupsMeshes()[SourceIt];
-			if (!bEquals)
+			if (const FHairGroupsMeshesSourceDescription* CachedDesc = CachedHairGroupsMeshes.IsValidIndex(SourceIt) ? &CachedHairGroupsMeshes[SourceIt] : nullptr)
 			{
-				bHasChanged = true;
-				break;
+				bIsLODValid = bIsLODValid && *CachedDesc == *Desc;
 			}
-		}
-	}
 
-	if (!bHasChanged)
-	{
-		for (uint32 SourceIt = 0, SourceCount = GetHairGroupsMeshes().Num(); SourceIt < SourceCount; ++SourceIt)
-		{
-			const bool bEquals = CachedHairGroupsMeshes[SourceIt] == GetHairGroupsMeshes()[SourceIt];
-			if (!bEquals)
-			{
-				bHasChanged = true;
-			}
-			else
-			{
-				// Reload/Update texture resources when the groom asset is saved, so that the textures 
-				// content is up to date with what has been saved.
-				if (IsInGameThread())
-				{
-					for (UTexture2D* Texture : GetHairGroupsMeshes()[SourceIt].Textures.Textures)
-					{
-						if (Texture != nullptr)
-						{
-							Texture->UpdateResource();
-						}
-					}
-				}
-			}
+			OutIsValid[LODIt] = bIsLODValid;
 		}
-	}
+		bHasChanged |= !OutIsValid[LODIt];
+	}	
+	
+	return bHasChanged;
+}
 
-	if (!bHasChanged)
+bool UGroomAsset::BuildHairGroup_Meshes(uint32 GroupIndex)
+{
+	TArray<bool> bIsAlreadyBuilt;
+	if (!HasChanged_Meshes(GroupIndex, bIsAlreadyBuilt))
 	{
 		return false;
 	}
 
+	FHairGroupPlatformData& GroupData = GetHairGroupsPlatformData()[GroupIndex];
+	const uint32 LODCount = GetHairGroupsLOD()[GroupIndex].LODs.Num();
+
 	bool bDataBuilt = false;
 	GroupData.Meshes.LODs.SetNum(LODCount);
-	const FHairGroupsLOD& GroupsLOD = GetHairGroupsLOD()[GroupIndex];
 	for (uint32 LODIt = 0; LODIt < LODCount; ++LODIt)
 	{
 		int32 SourceIt = 0;
-		if (const FHairGroupsMeshesSourceDescription* Desc = GetSourceDescription(GetHairGroupsMeshes(), GroupIndex, LODIt, SourceIt))
+		if (const FHairGroupsMeshesSourceDescription* Desc = GetSourceDescription(GetHairGroupsMeshes(), GroupIndex, LODIt, &SourceIt))
 		{
 			const FHairGroupsMeshesSourceDescription* CachedDesc = SourceIt < CachedHairGroupsMeshes.Num() ? &CachedHairGroupsMeshes[SourceIt] : nullptr;
 			const bool bLODHasChanged = CachedDesc == nullptr || !(*CachedDesc == *Desc);
@@ -3296,8 +3019,7 @@ bool UGroomAsset::BuildMeshesData(uint32 GroupIndex)
 	// Update the imported mesh DDC keys
 	for (int32 LODIt = 0; LODIt < GroupData.Meshes.LODs.Num(); ++LODIt)
 	{
-		int32 SourceIt = 0;
-		if (FHairGroupsMeshesSourceDescription* Desc = GetSourceDescription(GetHairGroupsMeshes(), GroupIndex, LODIt, SourceIt))
+		if (FHairGroupsMeshesSourceDescription* Desc = GetSourceDescription(GetHairGroupsMeshes(), GroupIndex, LODIt))
 		{
 			Desc->UpdateMeshKey();
 		}
@@ -3331,16 +3053,17 @@ void UGroomAsset::InitGuideResources()
 
 FHairStrandsRestResource* UGroomAsset::AllocateGuidesResources(uint32 GroupIndex)
 {
-	if (GroupIndex < uint32(GetNumHairGroups()))
+	if (GetHairGroupsPlatformData().IsValidIndex(GroupIndex))
 	{
 		FHairGroupPlatformData& GroupData = GetHairGroupsPlatformData()[GroupIndex];
+		FHairGroupResources& GroupResources = GetHairGroupsResources()[GroupIndex];
 		if (GroupData.Guides.HasValidData())
 		{
-			if (GroupData.Guides.RestResource == nullptr)
+			if (GroupResources.Guides.RestResource == nullptr)
 			{
-				GroupData.Guides.RestResource = new FHairStrandsRestResource(GroupData.Guides.BulkData, EHairStrandsResourcesType::Guides, FHairResourceName(GetFName(), GroupIndex), GetAssetPathName());
+				GroupResources.Guides.RestResource = new FHairStrandsRestResource(GroupData.Guides.BulkData, EHairStrandsResourcesType::Guides, FHairResourceName(GetFName(), GroupIndex), GetAssetPathName());
 			}
-			return GroupData.Guides.RestResource;
+			return GroupResources.Guides.RestResource;
 		}
 	}
 	return nullptr;
@@ -3348,15 +3071,16 @@ FHairStrandsRestResource* UGroomAsset::AllocateGuidesResources(uint32 GroupIndex
 
 FHairStrandsInterpolationResource* UGroomAsset::AllocateInterpolationResources(uint32 GroupIndex)
 {
-	if (GroupIndex < uint32(GetNumHairGroups()))
+	if (GetHairGroupsPlatformData().IsValidIndex(GroupIndex))
 	{
 		FHairGroupPlatformData& GroupData = GetHairGroupsPlatformData()[GroupIndex];
-		check(GroupData.Guides.IsValid());
-		if (GroupData.Strands.InterpolationResource == nullptr)
+		FHairGroupResources& GroupResources = GetHairGroupsResources()[GroupIndex];
+		check(GroupData.Guides.HasValidData());
+		if (GroupResources.Strands.InterpolationResource == nullptr)
 		{
-			GroupData.Strands.InterpolationResource = new FHairStrandsInterpolationResource(GroupData.Strands.InterpolationBulkData, FHairResourceName(GetFName(), GroupIndex), GetAssetPathName());
+			GroupResources.Strands.InterpolationResource = new FHairStrandsInterpolationResource(GroupData.Strands.InterpolationBulkData, FHairResourceName(GetFName(), GroupIndex), GetAssetPathName());
 		}
-		return GroupData.Strands.InterpolationResource;
+		return GroupResources.Strands.InterpolationResource;
 	}
 	return nullptr;
 }
@@ -3364,15 +3088,17 @@ FHairStrandsInterpolationResource* UGroomAsset::AllocateInterpolationResources(u
 #if RHI_RAYTRACING
 FHairStrandsRaytracingResource* UGroomAsset::AllocateCardsRaytracingResources(uint32 GroupIndex, uint32 LODIndex)
 {
-	if (IsRayTracingAllowed() && GroupIndex < uint32(GetNumHairGroups()) && LODIndex < uint32(GetHairGroupsLOD()[GroupIndex].LODs.Num()))
+	if (IsRayTracingAllowed() && GetHairGroupsLOD().IsValidIndex(GroupIndex) && GetHairGroupsLOD()[GroupIndex].LODs.IsValidIndex(LODIndex))
 	{
 		FHairGroupPlatformData& GroupData = GetHairGroupsPlatformData()[GroupIndex];
-		FHairGroupPlatformData::FCards::FLOD& LOD = GroupData.Cards.LODs[LODIndex];
-		check(LOD.BulkData.IsValid());
+		FHairGroupResources& GroupResource = GetHairGroupsResources()[GroupIndex];
+		FHairGroupPlatformData::FCards::FLOD& LODData = GroupData.Cards.LODs[LODIndex];
+		FHairGroupResources::FCards::FLOD& LOD = GroupResource.Cards.LODs[LODIndex];
+		check(LODData.HasValidData());
 
 		if (LOD.RaytracingResource == nullptr)
 		{
-			LOD.RaytracingResource = new FHairStrandsRaytracingResource(LOD.BulkData, FHairResourceName(GetFName(), GroupIndex, LODIndex), GetAssetPathName(LODIndex));
+			LOD.RaytracingResource = new FHairStrandsRaytracingResource(LODData.BulkData, FHairResourceName(GetFName(), GroupIndex, LODIndex), GetAssetPathName(LODIndex));
 		}
 		return LOD.RaytracingResource;
 	}
@@ -3381,15 +3107,17 @@ FHairStrandsRaytracingResource* UGroomAsset::AllocateCardsRaytracingResources(ui
 
 FHairStrandsRaytracingResource* UGroomAsset::AllocateMeshesRaytracingResources(uint32 GroupIndex, uint32 LODIndex)
 {
-	if (IsRayTracingAllowed() && GroupIndex < uint32(GetNumHairGroups()) && LODIndex < uint32(GetHairGroupsLOD()[GroupIndex].LODs.Num()))
+	if (IsRayTracingAllowed() && GetHairGroupsLOD().IsValidIndex(GroupIndex) && GetHairGroupsLOD()[GroupIndex].LODs.IsValidIndex(LODIndex))
 	{
 		FHairGroupPlatformData& GroupData = GetHairGroupsPlatformData()[GroupIndex];
-		FHairGroupPlatformData::FMeshes::FLOD& LOD = GroupData.Meshes.LODs[LODIndex];
-		check(LOD.BulkData.IsValid());
+		FHairGroupResources& GroupResource = GetHairGroupsResources()[GroupIndex];
 
+		FHairGroupPlatformData::FMeshes::FLOD& LODData = GroupData.Meshes.LODs[LODIndex];
+		FHairGroupResources::FMeshes::FLOD& LOD = GroupResource.Meshes.LODs[LODIndex];
+		check(LODData.HasValidData());
 		if (LOD.RaytracingResource == nullptr)
 		{
-			LOD.RaytracingResource = new FHairStrandsRaytracingResource(LOD.BulkData, FHairResourceName(GetFName(), GroupIndex, LODIndex), GetAssetPathName(LODIndex));
+			LOD.RaytracingResource = new FHairStrandsRaytracingResource(LODData.BulkData, FHairResourceName(GetFName(), GroupIndex, LODIndex), GetAssetPathName(LODIndex));
 		}
 		return LOD.RaytracingResource;
 	}
@@ -3398,16 +3126,17 @@ FHairStrandsRaytracingResource* UGroomAsset::AllocateMeshesRaytracingResources(u
 
 FHairStrandsRaytracingResource* UGroomAsset::AllocateStrandsRaytracingResources(uint32 GroupIndex)
 {
-	if (IsRayTracingAllowed() && GroupIndex < uint32(GetNumHairGroups()))
+	if (IsRayTracingAllowed() && GetHairGroupsPlatformData().IsValidIndex(GroupIndex))
 	{
 		FHairGroupPlatformData& GroupData = GetHairGroupsPlatformData()[GroupIndex];
+		FHairGroupResources& GroupResource = GetHairGroupsResources()[GroupIndex];
 		check(GroupData.Strands.HasValidData());
 
-		if (GroupData.Strands.RaytracingResource == nullptr)
+		if (GroupResource.Strands.RaytracingResource == nullptr)
 		{
-			GroupData.Strands.RaytracingResource = new FHairStrandsRaytracingResource(GroupData.Strands.BulkData, FHairResourceName(GetFName(), GroupIndex), GetAssetPathName());
+			GroupResource.Strands.RaytracingResource = new FHairStrandsRaytracingResource(GroupData.Strands.BulkData, FHairResourceName(GetFName(), GroupIndex), GetAssetPathName());
 		}
-		return GroupData.Strands.RaytracingResource;
+		return GroupResource.Strands.RaytracingResource;
 	}
 	return nullptr;
 }
@@ -3432,16 +3161,13 @@ void UGroomAsset::InitStrandsResources()
 	for (uint32 GroupIndex = 0, GroupCount = GetNumHairGroups(); GroupIndex < GroupCount; ++GroupIndex)
 	{
 		FHairGroupPlatformData& GroupData = GetHairGroupsPlatformData()[GroupIndex];
+		FHairGroupResources& GroupResource = GetHairGroupsResources()[GroupIndex];
 
 		if (GroupData.Strands.HasValidData())
 		{
 			FHairResourceName ResourceName(GetFName(), GroupIndex);
-			GroupData.Strands.RestResource = new FHairStrandsRestResource(GroupData.Strands.BulkData, EHairStrandsResourcesType::Strands, ResourceName, OwnerName);
-
-			if (GroupData.Strands.ClusterBulkData.IsValid())
-			{
-				GroupData.Strands.ClusterResource = new FHairStrandsClusterResource(GroupData.Strands.ClusterBulkData, ResourceName, OwnerName);
-			}
+			GroupResource.Strands.RestResource = new FHairStrandsRestResource(GroupData.Strands.BulkData, EHairStrandsResourcesType::Strands, ResourceName, OwnerName);
+			GroupResource.Strands.ClusterResource = new FHairStrandsClusterResource(GroupData.Strands.ClusterBulkData, ResourceName, OwnerName);
 
 			// Interpolation are lazy allocated, as these resources are only used when RBF/simulation is enabled by a groom component
 			// AllocateInterpolationResources(GroupIndex);
@@ -3459,30 +3185,20 @@ static void InitCardsTextureResources(UGroomAsset* GroomAsset)
 	for (uint32 GroupIndex = 0, GroupCount = GroomAsset->GetNumHairGroups(); GroupIndex < GroupCount; ++GroupIndex)
 	{
 		FHairGroupPlatformData& GroupData = GroomAsset->GetHairGroupsPlatformData()[GroupIndex];
+		FHairGroupResources& GroupResource = GroomAsset->GetHairGroupsResources()[GroupIndex];
 
 		const uint32 LODCount = GroomAsset->GetHairGroupsLOD()[GroupIndex].LODs.Num();
 		GroupData.Cards.LODs.SetNum(LODCount);
 
 		for (uint32 LODIt = 0; LODIt < LODCount; ++LODIt)
 		{
-			FHairGroupPlatformData::FCards::FLOD& LOD = GroupData.Cards.LODs[LODIt];
-
-			int32 SourceIt = 0;
-			const FHairGroupsCardsSourceDescription* Desc = GetSourceDescription(GroomAsset->GetHairGroupsCards(), GroupIndex, LODIt, SourceIt);
-			if (!Desc)
-			{
-				continue;
-			}
-
-			if (LOD.HasValidData())
+			if (const FHairGroupsCardsSourceDescription* Desc = GetSourceDescription(GroomAsset->GetHairGroupsCards(), GroupIndex, LODIt))
 			{
-				if (Desc)
+				FHairGroupPlatformData::FCards::FLOD& LODData = GroupData.Cards.LODs[LODIt];
+				FHairGroupResources::FCards::FLOD& LOD = GroupResource.Cards.LODs[LODIt];
+				if (LODData.HasValidData())
 				{
-					InitAtlasTexture(LOD.RestResource, Desc->Textures);
-					if (LOD.RestResource)
-					{
-						LOD.RestResource->bInvertUV = Desc->bInvertUV; // Should fix procedural texture so that this does not happen
-					}
+					InitAtlasTexture(LOD.RestResource, Desc->Textures, Desc->bInvertUV);
 				}
 			}
 		}
@@ -3493,57 +3209,66 @@ void UGroomAsset::InitCardsResources()
 {
 	LLM_SCOPE_BYTAG(Groom);
 
-	if (!IsHairStrandsEnabled(EHairStrandsShaderType::Cards) || GetHairGroupsCards().Num() == 0)
-	{
-		return;
-	}
-
 	for (uint32 GroupIndex = 0, GroupCount = GetNumHairGroups(); GroupIndex < GroupCount; ++GroupIndex)
 	{
+		check(GetHairGroupsPlatformData().IsValidIndex(GroupIndex));
+
+		// * Release resource if any
+		// * Resize the LODs to match the group LOD count
+		// Even if the groom does not have cards/meshes data, LOD datas/resources count needs to match
+		ReleaseCardsResource(GroupIndex);
+
 		FHairGroupPlatformData& GroupData = GetHairGroupsPlatformData()[GroupIndex];
+		FHairGroupResources& GroupResource = HairGroupsResources[GroupIndex];
 
 		const uint32 LODCount = GetHairGroupsLOD()[GroupIndex].LODs.Num();
-		GroupData.Cards.LODs.SetNum(LODCount);
+		check(uint32(GroupData.Cards.LODs.Num()) == LODCount);
+		GroupResource.Cards.LODs.SetNum(LODCount);
 
-		for (uint32 LODIt = 0; LODIt < LODCount; ++LODIt)
+		// If the group does not have cards, skip resource creation
+		if (!IsHairStrandsEnabled(EHairStrandsShaderType::Cards) || GetHairGroupsCards().Num() == 0)
 		{
-			FHairGroupPlatformData::FCards::FLOD& LOD = GroupData.Cards.LODs[LODIt];
+			continue;
+		}
 
-			int32 SourceIt = 0;
-			const FHairGroupsCardsSourceDescription* Desc = GetSourceDescription(GetHairGroupsCards(), GroupIndex, LODIt, SourceIt);
-			if (!Desc)
-			{
-				continue;
-			}
+		for (uint32 LODIt = 0; LODIt < LODCount; ++LODIt)
+		{
+			FHairGroupPlatformData::FCards::FLOD& LODData = GroupData.Cards.LODs[LODIt];
+			FHairGroupResources::FCards::FLOD& LOD = GroupResource.Cards.LODs[LODIt];
 
-			if (LOD.RestResource == nullptr &&			// don't initialize again if they were previously initialized during the BuildCardsGeometry
-				LOD.HasValidData())
+			if (const FHairGroupsCardsSourceDescription* Desc = GetSourceDescription(GetHairGroupsCards(), GroupIndex, LODIt))
 			{
-				FHairResourceName ResourceName(GetFName(), GroupIndex, LODIt);
-				const FName OwnerName = GetAssetPathName(LODIt);
-
-				LOD.RestResource = new FHairCardsRestResource(LOD.BulkData, ResourceName, OwnerName);
-				BeginInitResource(LOD.RestResource); // Immediate allocation, as needed for the vertex factory, input stream building
-
-				LOD.InterpolationResource = new FHairCardsInterpolationResource(LOD.InterpolationBulkData, ResourceName, OwnerName);
-
-				LOD.Guides.RestResource = new FHairStrandsRestResource(LOD.Guides.BulkData, EHairStrandsResourcesType::Cards, ResourceName, OwnerName);
-
-				LOD.Guides.InterpolationResource = new FHairStrandsInterpolationResource(LOD.Guides.InterpolationBulkData, ResourceName, OwnerName);
-
-				if (Desc)
+				const bool bNeedInitialization = LOD.RestResource == nullptr && LODData.HasValidData();
+				if (bNeedInitialization)
 				{
-					InitAtlasTexture(LOD.RestResource, Desc->Textures);
-					LOD.RestResource->bInvertUV = Desc->bInvertUV; // Should fix procedural texture so that this does not happen
+					FHairResourceName ResourceName(GetFName(), GroupIndex, LODIt);
+					const FName OwnerName = GetAssetPathName(LODIt);
+	
+					LOD.RestResource 				= new FHairCardsRestResource(LODData.BulkData, ResourceName, OwnerName);
+					LOD.InterpolationResource		= new FHairCardsInterpolationResource(LODData.InterpolationBulkData, ResourceName, OwnerName);
+					LOD.GuideRestResource			= new FHairStrandsRestResource(LODData.GuideBulkData, EHairStrandsResourcesType::Cards, ResourceName, OwnerName);
+					LOD.GuideInterpolationResource	= new FHairStrandsInterpolationResource(LODData.GuideInterpolationBulkData, ResourceName, OwnerName);
+
+					// Resource allocation
+					// Immediate allocation, as needed for the vertex factory, input stream building
+					BeginInitResource(LOD.RestResource); 
+					BeginInitResource(LOD.InterpolationResource);
+					BeginInitResource(LOD.GuideRestResource);
+					BeginInitResource(LOD.GuideInterpolationResource);
+	
+					InitAtlasTexture(LOD.RestResource, Desc->Textures, Desc->bInvertUV);
 				}
 			}
 
-			if (Desc)
-			{
+			// Update groom cards stats (visible in groom editor)
+			#if WITH_EDITOR
+			if (FHairGroupsCardsSourceDescription* Desc = GetSourceDescription(GetHairGroupsCards(), GroupIndex, LODIt))
+			{	
 				// Update card stats to display
-				Desc->CardsInfo.NumCardVertices = LOD.BulkData.GetNumVertices();
-				Desc->CardsInfo.NumCards = LOD.Guides.BulkData.GetNumCurves();// LOD.Data.Cards.GetNumCards();
+				Desc->CardsInfo.NumCardVertices = LODData.BulkData.GetNumVertices();
+				Desc->CardsInfo.NumCards 		= LODData.GuideBulkData.GetNumCurves();
 			}
+			#endif
 		}
 	}
 }
@@ -3552,35 +3277,39 @@ void UGroomAsset::InitMeshesResources()
 {
 	LLM_SCOPE_BYTAG(Groom);
 
-	if (!IsHairStrandsEnabled(EHairStrandsShaderType::Meshes))
-	{
-		return;
-	}
-
 	for (uint32 GroupIndex = 0, GroupCount = GetNumHairGroups(); GroupIndex < GroupCount; ++GroupIndex)
 	{
+		check(GetHairGroupsPlatformData().IsValidIndex(GroupIndex));
+
+		// * Release resource if any
+		// * Resize the LODs to match the group LOD count
+		// Even if the groom does not have cards/meshes data, LOD datas/resources count needs to match
+		ReleaseMeshesResource(GroupIndex);
+
 		FHairGroupPlatformData& GroupData = GetHairGroupsPlatformData()[GroupIndex];
+		FHairGroupResources& GroupResource = GetHairGroupsResources()[GroupIndex];
 
-		const uint32 LODCount = GroupData.Meshes.LODs.Num();
-		for (uint32 LODIt = 0; LODIt < LODCount; ++LODIt)
+		const uint32 LODCount = GetHairGroupsLOD()[GroupIndex].LODs.Num();
+		check(uint32(GroupData.Meshes.LODs.Num()) == LODCount);
+		GroupResource.Meshes.LODs.SetNum(LODCount);
+
+		if (!IsHairStrandsEnabled(EHairStrandsShaderType::Meshes) || GetHairGroupsMeshes().Num() == 0)
 		{
-			int32 SourceIt = 0;
-			const FHairGroupsMeshesSourceDescription* Desc = GetSourceDescription(GetHairGroupsMeshes(), GroupIndex, LODIt, SourceIt);
-			if (!Desc)
-			{
-				continue;
-			}
-			FHairGroupPlatformData::FMeshes::FLOD& LOD = GroupData.Meshes.LODs[LODIt];
-			InternalReleaseResource(LOD.RestResource);
+			continue;
+		}
 
-			if (LOD.HasValidData())
+		for (uint32 LODIt = 0; LODIt < LODCount; ++LODIt)
+		{
+			if (const FHairGroupsMeshesSourceDescription* Desc = GetSourceDescription(GetHairGroupsMeshes(), GroupIndex, LODIt))
 			{
-				LOD.RestResource = new FHairMeshesRestResource(LOD.BulkData, FHairResourceName(GetFName(), GroupIndex, LODIt), GetAssetPathName(LODIt));
-				BeginInitResource(LOD.RestResource); // Immediate allocation, as needed for the vertex factory, input stream building
-
-				if (Desc)
+				FHairGroupPlatformData::FMeshes::FLOD& LODData = GroupData.Meshes.LODs[LODIt];
+				FHairGroupResources::FMeshes::FLOD& LOD = GroupResource.Meshes.LODs[LODIt];
+	
+				if (LODData.HasValidData())
 				{
-					InitAtlasTexture(LOD.RestResource, Desc->Textures);
+					LOD.RestResource = new FHairMeshesRestResource(LODData.BulkData, FHairResourceName(GetFName(), GroupIndex, LODIt), GetAssetPathName(LODIt));
+					BeginInitResource(LOD.RestResource); // Immediate allocation, as needed for the vertex factory, input stream building
+					InitAtlasTexture(LOD.RestResource, Desc->Textures, false /*bInvertUV*/);
 				}
 			}
 		}
@@ -3787,24 +3516,7 @@ void UGroomAsset::StripLODs(const TArray<int32>& LODsToKeep, bool bRebuildResour
 		--LODsTOKeepIndex;
 	}
 
-	// Rebuild the LOD data
-	if (bRebuildResources && CanRebuildFromDescription())
-	{
-		const FHairDescriptionGroups& LocalHairDescriptionGroups = GetHairDescriptionGroups();
-		for (int32 GroupIt = 0; GroupIt < GroupCount; ++GroupIt)
-		{
-			FHairGroupPlatformData& GroupData = GetHairGroupsPlatformData()[GroupIt];
-			InternalReleaseResource(GroupData.Strands.ClusterResource);
-
-			BuildHairGroupCluster(
-				GroupIt,
-				LocalHairDescriptionGroups,
-				GetHairGroupsInterpolation(),
-				GetHairGroupsLOD(),
-				GetHairGroupsInfo(),
-				GetHairGroupsPlatformData());
-		}
-	}
+	CacheDerivedDatas();
 }
 #endif // WITH_EDITORONLY_DATA
 
@@ -3840,10 +3552,11 @@ void UGroomAsset::CreateDebugData()
 
 		if (GroupData.Debug.Data.IsValid())
 		{
-			GroupData.Debug.Resource = new FHairStrandsDebugDatas::FResources();
+			FHairGroupResources& GroupResource = GetHairGroupsResources()[GroupIndex];
+			GroupResource.Debug.Resource = new FHairStrandsDebugResources();
 
 			FHairStrandsDebugDatas* InData = &GroupData.Debug.Data;
-			FHairStrandsDebugDatas::FResources* InResource = GroupData.Debug.Resource;
+			FHairStrandsDebugResources* InResource = GroupResource.Debug.Resource;
 			ENQUEUE_RENDER_COMMAND(HairStrandsDebugResourceCommand)(
 				[InData, InResource](FRHICommandListImmediate& RHICmdList)
 				{
@@ -3949,8 +3662,7 @@ bool UGroomAsset::GetHairCardsGuidesDatas(
 		return false;
 	}
 
-	int32 SourceIt = 0;
-	if (FHairGroupsCardsSourceDescription* Desc = GetSourceDescription(GetHairGroupsCards(), GroupIndex, LODIndex, SourceIt))
+	if (FHairGroupsCardsSourceDescription* Desc = GetSourceDescription(GetHairGroupsCards(), GroupIndex, LODIndex))
 	{
 		// 1. Load geometry data, if any
 		if (UStaticMesh* CardsMesh = Desc->ImportedMesh)
@@ -3980,18 +3692,23 @@ void UGroomAsset::GetResourceSizeEx(FResourceSizeEx& CumulativeResourceSize)
 	CumulativeResourceSize.AddDedicatedSystemMemoryBytes(GetHairGroupsMaterials().GetAllocatedSize());
 	CumulativeResourceSize.AddDedicatedSystemMemoryBytes(GetHairGroupsPlatformData().GetAllocatedSize());
 	
+	// Datas
 	for (const FHairGroupPlatformData & GroupData : GetHairGroupsPlatformData())
 	{		
-		CumulativeResourceSize.AddDedicatedVideoMemoryBytes(GroupData.Guides.GetResourcesSize());
-		CumulativeResourceSize.AddDedicatedVideoMemoryBytes(GroupData.Strands.GetResourcesSize());
-		CumulativeResourceSize.AddDedicatedVideoMemoryBytes(GroupData.Cards.GetResourcesSize());
-		CumulativeResourceSize.AddDedicatedVideoMemoryBytes(GroupData.Meshes.GetResourcesSize());
-
 		CumulativeResourceSize.AddDedicatedSystemMemoryBytes(GroupData.Guides.GetDataSize());
 		CumulativeResourceSize.AddDedicatedSystemMemoryBytes(GroupData.Strands.GetDataSize());
 		CumulativeResourceSize.AddDedicatedSystemMemoryBytes(GroupData.Cards.GetDataSize());
 		CumulativeResourceSize.AddDedicatedSystemMemoryBytes(GroupData.Meshes.GetDataSize());
 	}
+
+	// Resources
+	for (const FHairGroupResources & GroupResource : GetHairGroupsResources())
+	{		
+		CumulativeResourceSize.AddDedicatedVideoMemoryBytes(GroupResource.Guides.GetResourcesSize());
+		CumulativeResourceSize.AddDedicatedVideoMemoryBytes(GroupResource.Strands.GetResourcesSize());
+		CumulativeResourceSize.AddDedicatedVideoMemoryBytes(GroupResource.Cards.GetResourcesSize());
+		CumulativeResourceSize.AddDedicatedVideoMemoryBytes(GroupResource.Meshes.GetResourcesSize());
+	}
 }
 
 FName UGroomAsset::GetAssetPathName(int32 LODIndex)
@@ -4082,3 +3799,13 @@ DEFINE_GROOM_MEMBER_ACCESSOR(FPerPlatformBool, DisableBelowMinLodStripping);
 DEFINE_GROOM_MEMBER_ACCESSOR_ARRAY(TArray<float>, EffectiveLODBias);
 DEFINE_GROOM_MEMBER_ACCESSOR_ARRAY(TArray<FHairGroupPlatformData>, HairGroupsPlatformData);
 DEFINE_GROOM_MEMBER_ACCESSOR_ARRAY(TArray<FHairGroupInfoWithVisibility>, HairGroupsInfo);
+
+const TArray<FHairGroupResources>& UGroomAsset::GetHairGroupsResources() const
+{
+	return HairGroupsResources;
+}
+
+TArray<FHairGroupResources>& UGroomAsset::GetHairGroupsResources()
+{
+	return HairGroupsResources;
+}
\ No newline at end of file
diff --git a/Engine/Plugins/Runtime/HairStrands/Source/HairStrandsCore/Private/GroomBuilder.cpp b/Engine/Plugins/Runtime/HairStrands/Source/HairStrandsCore/Private/GroomBuilder.cpp
--- a/Engine/Plugins/Runtime/HairStrands/Source/HairStrandsCore/Private/GroomBuilder.cpp
+++ b/Engine/Plugins/Runtime/HairStrands/Source/HairStrandsCore/Private/GroomBuilder.cpp
@@ -3419,7 +3419,7 @@ static FRDGBufferRef ReadPositions(
 	FRDGBufferRef OutPositions = nullptr;
 	if (const FHairGroupInstance* Instance = Component->GetGroupInstance(GroupIt))
 	{
-		const uint32 NumVertices = Instance->Strands.Data->GetNumPoints();
+		const uint32 NumVertices = Instance->Strands.GetData().GetNumPoints();
 
 		FRDGBufferSRVRef InPositions = nullptr;
 		FRDGBufferSRVRef InPositionOffset = nullptr;
diff --git a/Engine/Plugins/Runtime/HairStrands/Source/HairStrandsCore/Private/GroomComponent.cpp b/Engine/Plugins/Runtime/HairStrands/Source/HairStrandsCore/Private/GroomComponent.cpp
--- a/Engine/Plugins/Runtime/HairStrands/Source/HairStrandsCore/Private/GroomComponent.cpp
+++ b/Engine/Plugins/Runtime/HairStrands/Source/HairStrandsCore/Private/GroomComponent.cpp
@@ -2167,16 +2167,16 @@ UMaterialInterface* UGroomComponent::GetMaterial(int32 ElementIndex) const
 
 FHairStrandsRestResource* UGroomComponent::GetGuideStrandsRestResource(uint32 GroupIndex)
 {
-	if (!GroomAsset || GroupIndex >= uint32(GroomAsset->GetNumHairGroups()))
+	if (!GroomAsset || !GroomAsset->GetHairGroupsResources().IsValidIndex(GroupIndex))
 	{
 		return nullptr;
 	}
 
-	if (!GroomAsset->GetHairGroupsPlatformData()[GroupIndex].Guides.IsValid())
+	if (!GroomAsset->GetHairGroupsResources()[GroupIndex].Guides.IsValid())
 	{
 		return nullptr;
 	}
-	return GroomAsset->GetHairGroupsPlatformData()[GroupIndex].Guides.RestResource;
+	return GroomAsset->GetHairGroupsResources()[GroupIndex].Guides.RestResource;
 }
 
 FHairStrandsDeformedResource* UGroomComponent::GetGuideStrandsDeformedResource(uint32 GroupIndex)
@@ -2667,6 +2667,7 @@ void UGroomComponent::InitResources(bool bIsBindingReloading)
 		HairGroupInstance->LocalToWorld							= HairGroupInstance->GetCurrentLocalToWorld();
 
 		FHairGroupPlatformData& GroupData = GroomAsset->GetHairGroupsPlatformData()[GroupIt];
+		FHairGroupResources& GroupResources = GroomAsset->GetHairGroupsResources()[GroupIt];
 
 		const EHairInterpolationType HairInterpolationType = ToHairInterpolationType(GroomAsset->GetHairInterpolationType());
 
@@ -2732,8 +2733,6 @@ void UGroomComponent::InitResources(bool bIsBindingReloading)
 		const bool bNeedGuides = (GroupData.Guides.HasValidData() && (bHasNeedSimulation[GroupIt] || bHasNeedGlobalDeformation[GroupIt] || bHasNeedDeformation[GroupIt])) || HairGroupInstance->HairGroupPublicData->bIsSimulationCacheEnable;
 		if (bNeedGuides)
 		{
-			HairGroupInstance->Guides.Data = &GroupData.Guides.BulkData;
-
 			if (LocalBindingAsset)
 			{
 				check(RegisteredMeshComponent);
@@ -2750,7 +2749,7 @@ void UGroomComponent::InitResources(bool bIsBindingReloading)
 
 			// Lazy allocation of the guide resources
 			HairGroupInstance->Guides.RestResource = GroomAsset->AllocateGuidesResources(GroupIt);
-			check(GroupData.Guides.RestResource);
+			check(GroupResources.Guides.RestResource);
 
 			// If guides are allocated, deformed resources are always needs since they are either used with simulation, or RBF deformation. Both are dynamics, and require deformed positions
 			HairGroupInstance->Guides.DeformedResource = new FHairStrandsDeformedResource(GroupData.Guides.BulkData, EHairStrandsResourcesType::Guides, ResourceName, OwnerName);
@@ -2767,11 +2766,10 @@ void UGroomComponent::InitResources(bool bIsBindingReloading)
 		HairGroupInstance->Strands.Modifier = GetGroomGroupsDesc(GroomAsset, this, GroupIt);
 
 		// Strands data/resources
-		if (bNeedStrandsData && GroupData.Strands.IsValid())
+		if (bNeedStrandsData && GroupData.Strands.HasValidData())
 		{
 			check(GroupIt < GroomGroupsDesc.Num());
 
-			HairGroupInstance->Strands.Data = &GroupData.Strands.BulkData;
 			HairGroupInstance->HairGroupPublicData->RestPointCount = GroupData.Strands.BulkData.GetNumPoints();
 			HairGroupInstance->HairGroupPublicData->RestCurveCount = GroupData.Strands.BulkData.GetNumCurves();
 			HairGroupInstance->HairGroupPublicData->ClusterCount = 0u;
@@ -2873,19 +2871,19 @@ void UGroomComponent::InitResources(bool bIsBindingReloading)
 				HairGroupInstance->Strands.DeformedRootResource = new FHairStrandsDeformedRootResource(HairGroupInstance->Strands.RestRootResource, EHairStrandsResourcesType::Strands, ResourceName, OwnerName);
 			}
 
-			HairGroupInstance->Strands.RestResource = GroupData.Strands.RestResource;
+			HairGroupInstance->Strands.RestResource = GroupResources.Strands.RestResource;
 			if (bNeedDynamicResources)
 			{
 				HairGroupInstance->Strands.DeformedResource = new FHairStrandsDeformedResource(GroupData.Strands.BulkData, EHairStrandsResourcesType::Strands, ResourceName, OwnerName);
 			} 
 
 			// An empty groom doesn't have a ClusterResource
-			HairGroupInstance->Strands.ClusterResource = GroupData.Strands.ClusterResource;
+			HairGroupInstance->Strands.ClusterResource = GroupResources.Strands.ClusterResource;
 			if (HairGroupInstance->Strands.ClusterResource)
 			{
 				// This codes assumes strands LOD are contigus and the highest (i.e., 0...x). Change this code to something more robust
 				check(HairGroupInstance->HairGroupPublicData);
-				const int32 StrandsLODCount = GroupData.Strands.ClusterResource->BulkData.Header.LODInfos.Num();
+				const int32 StrandsLODCount = GroupResources.Strands.ClusterResource->BulkData.Header.LODInfos.Num();
 				const TArray<float>& LODScreenSizes = HairGroupInstance->HairGroupPublicData->GetLODScreenSizes();
 				const TArray<bool>& LODVisibilities = HairGroupInstance->HairGroupPublicData->GetLODVisibilities();
 				check(StrandsLODCount <= LODScreenSizes.Num());
@@ -2905,11 +2903,15 @@ void UGroomComponent::InitResources(bool bIsBindingReloading)
 		}
 
 		// Cards resources
-		for (FHairGroupPlatformData::FCards::FLOD& LOD : GroupData.Cards.LODs)
+		const uint32 CardLODCount = GroupData.Cards.LODs.Num();
+		for (uint32 CardLODIt=0;CardLODIt<CardLODCount;++CardLODIt)
 		{
+			FHairGroupPlatformData::FCards::FLOD& LODData = GroupData.Cards.LODs[CardLODIt];
+			FHairGroupResources::FCards::FLOD& LODResource = GroupResources.Cards.LODs[CardLODIt];
+
 			const uint32 CardsLODIndex = HairGroupInstance->Cards.LODs.Num();
 			FHairGroupInstance::FCards::FLOD& InstanceLOD = HairGroupInstance->Cards.LODs.AddDefaulted_GetRef();
-			if (LOD.IsValid())
+			if (LODData.HasValidData())
 			{
 				FHairResourceName LODResourceName(GetFName(), GroupIt, CardsLODIndex);
 				const FName LODOwnerName = GroomAsset->GetAssetPathName(CardsLODIndex);
@@ -2927,13 +2929,12 @@ void UGroomComponent::InitResources(bool bIsBindingReloading)
 					check(BindingType == EHairBindingType::Skinning);
 				}
 
-				InstanceLOD.Data = &LOD.BulkData;
-				InstanceLOD.RestResource = LOD.RestResource;
-				InstanceLOD.InterpolationResource = LOD.InterpolationResource;
+				InstanceLOD.RestResource = LODResource.RestResource;
+				InstanceLOD.InterpolationResource = LODResource.InterpolationResource;
 
 				if (bNeedDeformedPositions)
 				{
-					InstanceLOD.DeformedResource = new FHairCardsDeformedResource(LOD.BulkData, false, LODResourceName, LODOwnerName);
+					InstanceLOD.DeformedResource = new FHairCardsDeformedResource(LODData.BulkData, false, LODResourceName, LODOwnerName);
 				}
 
 				#if RHI_RAYTRACING
@@ -2942,7 +2943,7 @@ void UGroomComponent::InitResources(bool bIsBindingReloading)
 					if (bNeedDeformedPositions)
 					{
 						// Allocate dynamic raytracing resources (owned by the groom component/instance)
-						InstanceLOD.RaytracingResource = new FHairStrandsRaytracingResource(*InstanceLOD.Data, LODResourceName, LODOwnerName);
+						InstanceLOD.RaytracingResource = new FHairStrandsRaytracingResource(InstanceLOD.GetData(), LODResourceName, LODOwnerName);
 						InstanceLOD.RaytracingResourceOwned = true;
 					}
 					else
@@ -2958,8 +2959,7 @@ void UGroomComponent::InitResources(bool bIsBindingReloading)
 				// Strands data/resources
 				if (bNeedDeformedPositions)
 				{
-					InstanceLOD.Guides.Data = &LOD.Guides.BulkData;
-					InstanceLOD.Guides.InterpolationResource = LOD.Guides.InterpolationResource;
+					InstanceLOD.Guides.InterpolationResource = LODResource.GuideInterpolationResource;
 
 					if (bNeedRootData)
 					{
@@ -2974,9 +2974,9 @@ void UGroomComponent::InitResources(bool bIsBindingReloading)
 						InstanceLOD.Guides.DeformedRootResource = new FHairStrandsDeformedRootResource(InstanceLOD.Guides.RestRootResource, EHairStrandsResourcesType::Cards, LODResourceName, LODOwnerName);
 					}
 
-					InstanceLOD.Guides.RestResource = LOD.Guides.RestResource;
+					InstanceLOD.Guides.RestResource = LODResource.GuideRestResource;
 					{
-						InstanceLOD.Guides.DeformedResource = new FHairStrandsDeformedResource(LOD.Guides.BulkData, EHairStrandsResourcesType::Cards, LODResourceName, LODOwnerName);
+						InstanceLOD.Guides.DeformedResource = new FHairStrandsDeformedResource(LODData.GuideBulkData, EHairStrandsResourcesType::Cards, LODResourceName, LODOwnerName);
 					}
 
 					InstanceLOD.Guides.HairInterpolationType = HairInterpolationType;
@@ -2985,11 +2985,15 @@ void UGroomComponent::InitResources(bool bIsBindingReloading)
 		}
 
 		// Meshes resources
-		for (FHairGroupPlatformData::FMeshes::FLOD& LOD : GroupData.Meshes.LODs)
+		const uint32 MeshLODCount = GroupData.Meshes.LODs.Num();
+		for (uint32 MeshLODIt=0;MeshLODIt<MeshLODCount;++MeshLODIt)
 		{
+			FHairGroupPlatformData::FMeshes::FLOD& LODData = GroupData.Meshes.LODs[MeshLODIt];
+			FHairGroupResources::FMeshes::FLOD& LODResource = GroupResources.Meshes.LODs[MeshLODIt];
+
 			const int32 MeshLODIndex = HairGroupInstance->Meshes.LODs.Num();
 			FHairGroupInstance::FMeshes::FLOD& InstanceLOD = HairGroupInstance->Meshes.LODs.AddDefaulted_GetRef();
-			if (LOD.IsValid())
+			if (LODData.IsValid())
 			{
 				FHairResourceName LODResourceName(GetFName(), GroupIt, MeshLODIndex);
 				const FName LODOwnerName = GroomAsset->GetAssetPathName(MeshLODIndex);
@@ -2998,11 +3002,10 @@ void UGroomComponent::InitResources(bool bIsBindingReloading)
 				const bool bHasGlobalDeformation   = HairGroupInstance->HairGroupPublicData->IsGlobalInterpolationEnable(MeshLODIndex);
 				const bool bNeedDeformedPositions  = bHasGlobalDeformation && BindingType == EHairBindingType::Skinning;
 
-				InstanceLOD.Data = &LOD.BulkData;
-				InstanceLOD.RestResource = LOD.RestResource;
+				InstanceLOD.RestResource = LODResource.RestResource;
 				if (bNeedDeformedPositions)
 				{
-					InstanceLOD.DeformedResource = new FHairMeshesDeformedResource(LOD.BulkData, true, LODResourceName, LODOwnerName);
+					InstanceLOD.DeformedResource = new FHairMeshesDeformedResource(LODData.BulkData, true, LODResourceName, LODOwnerName);
 				}
 
 				#if RHI_RAYTRACING
@@ -3011,7 +3014,7 @@ void UGroomComponent::InitResources(bool bIsBindingReloading)
 					if (bNeedDeformedPositions)
 					{
 						// Allocate dynamic raytracing resources (owned by the groom component/instance)
-						InstanceLOD.RaytracingResource = new FHairStrandsRaytracingResource(*InstanceLOD.Data, LODResourceName, LODOwnerName);
+						InstanceLOD.RaytracingResource = new FHairStrandsRaytracingResource(LODData.BulkData, LODResourceName, LODOwnerName);
 						InstanceLOD.RaytracingResourceOwned = true;
 					}
 					else
diff --git a/Engine/Plugins/Runtime/HairStrands/Source/HairStrandsCore/Private/GroomDebug.cpp b/Engine/Plugins/Runtime/HairStrands/Source/HairStrandsCore/Private/GroomDebug.cpp
--- a/Engine/Plugins/Runtime/HairStrands/Source/HairStrandsCore/Private/GroomDebug.cpp
+++ b/Engine/Plugins/Runtime/HairStrands/Source/HairStrandsCore/Private/GroomDebug.cpp
@@ -487,7 +487,7 @@ static void AddDrawDebugStrandsCVsPass(
 	Parameters->ViewUniformBuffer = View.ViewUniformBuffer;
 	Parameters->HairStrandsVF = Instance->Strands.UniformBuffer;
 	Parameters->LocalToWorld = FMatrix44f(Instance->LocalToWorld.ToMatrixWithScale());		// LWC_TODO: Precision loss // TODO change this to Uniform buffer parameters..
-	Parameters->MaxVertexCount = Instance->Strands.Data->GetNumPoints();
+	Parameters->MaxVertexCount = Instance->Strands.GetData().GetNumPoints();
 	Parameters->ColorTexture = GraphBuilder.CreateUAV(ColorTexture);
 	Parameters->DepthTexture = DepthTexture;
 	Parameters->LinearSampler = TStaticSamplerState<SF_Bilinear, AM_Clamp, AM_Clamp, AM_Clamp>::GetRHI();
@@ -579,7 +579,7 @@ static void AddDrawDebugCardsGuidesPass(
 	}
 
 	const FHairGroupInstance::FCards::FLOD& LOD = Instance->Cards.LODs[HairLODIndex];
-	if (!LOD.Guides.Data)
+	if (!LOD.Guides.RestResource)
 	{
 		return;
 	}
@@ -820,9 +820,8 @@ static void AddHairDebugPrintInstancePass(
 		case EHairGeometryType::Strands:
 			if (Instance->Strands.IsValid())
 			{
-				check(Instance->Strands.Data);
-				D.Data0.Z = Instance->Strands.Data->GetNumCurves(); // Change this later on for having dynamic value
-				D.Data0.W = Instance->Strands.Data->GetNumPoints(); // Change this later on for having dynamic value
+				D.Data0.Z = Instance->Strands.GetData().GetNumCurves(); // Change this later on for having dynamic value
+				D.Data0.W = Instance->Strands.GetData().GetNumPoints(); // Change this later on for having dynamic value
 				const int32 MeshLODIndex = Instance->HairGroupPublicData->MeshLODIndex;
 				if (MeshLODIndex>=0 && Instance->Strands.RestRootResource)
 				{
@@ -845,7 +844,7 @@ static void AddHairDebugPrintInstancePass(
 					D.Data2.Y = Instance->HairGroupPublicData->GetActiveStrandsPointCount();
 					D.Data2.Z = Instance->HairGroupPublicData->GetActiveStrandsCurveCount();
 
-					D.Data2.W = Instance->Strands.Data->Header.ImportedAttributes;
+					D.Data2.W = Instance->Strands.GetData().Header.ImportedAttributes;
 				}
 				
 				D.Data3 = FUintVector4(0);
@@ -868,8 +867,8 @@ static void AddHairDebugPrintInstancePass(
 		case EHairGeometryType::Cards:
 			if (Instance->Cards.IsValid(IntLODIndex))
 			{
-				D.Data0.Z = Instance->Cards.LODs[IntLODIndex].Guides.IsValid() ? Instance->Cards.LODs[IntLODIndex].Guides.Data->GetNumCurves() : 0;
-				D.Data0.W = Instance->Cards.LODs[IntLODIndex].Data->GetNumVertices();
+				D.Data0.Z = Instance->Cards.LODs[IntLODIndex].Guides.IsValid() ? Instance->Cards.LODs[IntLODIndex].Guides.GetData().GetNumCurves() : 0;
+				D.Data0.W = Instance->Cards.LODs[IntLODIndex].GetData().GetNumVertices();
 
 				D.Data2 = FUintVector4(0);
 				D.Data2.X |= VisibilityType == EHairInstanceVisibilityType::CardsOrMeshesPrimaryView ? 0x1u : 0u;
@@ -887,7 +886,7 @@ static void AddHairDebugPrintInstancePass(
 			if (Instance->Meshes.IsValid(IntLODIndex))
 			{
 				D.Data0.Z = 0;
-				D.Data0.W = Instance->Meshes.LODs[IntLODIndex].Data->GetNumVertices();
+				D.Data0.W = Instance->Meshes.LODs[IntLODIndex].GetData().GetNumVertices();
 
 				D.Data2 = FUintVector4(0);
 				D.Data2.X |= VisibilityType == EHairInstanceVisibilityType::CardsOrMeshesPrimaryView ? 0x1u : 0u;
@@ -1113,14 +1112,15 @@ static void AddHairDebugPrintMemoryPass(
 				for (uint32 GroupIt = 0; GroupIt < GroupCount; ++GroupIt)
 				{					
 					const FHairGroupPlatformData& Data = AssetIt->GetHairGroupsPlatformData()[GroupIt];
-					const FGroomAssetMemoryStats MemoryStats = FGroomAssetMemoryStats::Get(Data);
+					const FHairGroupResources& Resources = AssetIt->GetHairGroupsResources()[GroupIt];
+					const FGroomAssetMemoryStats MemoryStats = FGroomAssetMemoryStats::Get(Data, Resources);
 
 					GroomNames.Add(AssetIt->GetName());
 
 					FInfos& D = GroomBuffer.AddDefaulted_GetRef();
 					D.Data0.X = GroupIt;
 					D.Data0.Y = GroupCount;
-					D.Data0.Z = Data.Strands.RestResource ? Data.Strands.RestResource->MaxAvailableCurveCount : 0;
+					D.Data0.Z = Resources.Strands.RestResource ? Resources.Strands.RestResource->MaxAvailableCurveCount : 0;
 					D.Data0.W = Data.Strands.BulkData.IsValid() ? Data.Strands.BulkData.GetNumCurves() : 0;
 			
 					D.Data1.X = MemoryStats.CPU.Guides;
diff --git a/Engine/Plugins/Runtime/HairStrands/Source/HairStrandsCore/Private/GroomManager.cpp b/Engine/Plugins/Runtime/HairStrands/Source/HairStrandsCore/Private/GroomManager.cpp
--- a/Engine/Plugins/Runtime/HairStrands/Source/HairStrandsCore/Private/GroomManager.cpp
+++ b/Engine/Plugins/Runtime/HairStrands/Source/HairStrandsCore/Private/GroomManager.cpp
@@ -1180,7 +1180,7 @@ static void RunHairStrandsInterpolation_Strands(
 
 		if (InstanceData.MeshLODIndex >= 0 && (InstanceData.bSimulationEnable || InstanceData.bDeformationEnable))
 		{
-			check(Instance->Strands.HasValidData());
+			check(Instance->Strands.IsValid());
 		}
 	}
 
@@ -1508,14 +1508,14 @@ static void RunHairStrandsInterpolation_Strands(
 				// Special case for debug mode were the attribute buffer is patch with some custom data to show hair properties (strands belonging to the same cluster, ...)
 				if (InstanceData.Instance->Strands.DebugCurveAttributeBuffer.Buffer == nullptr)
 				{
-					CreateHairStrandsDebugAttributeBuffer(GraphBuilder, &InstanceData.Instance->Strands.DebugCurveAttributeBuffer, InstanceData.Instance->Strands.Data->GetCurveAttributeSizeInBytes(), FName(InstanceData.Instance->Debug.MeshComponentName));
+					CreateHairStrandsDebugAttributeBuffer(GraphBuilder, &InstanceData.Instance->Strands.DebugCurveAttributeBuffer, InstanceData.Instance->Strands.GetData().GetCurveAttributeSizeInBytes(), FName(InstanceData.Instance->Debug.MeshComponentName));
 				}
 				FRDGImportedBuffer OutRenCurveAttributeBuffer = Register(GraphBuilder, InstanceData.Instance->Strands.DebugCurveAttributeBuffer, ERDGImportedBufferFlags::CreateUAV);
 			
 				const bool bValidGuide = InstanceData.bNeedDeformation && (InstanceData.bSimulationEnable || InstanceData.bGlobalDeformationEnable || InstanceData.bDeformationEnable || InstanceData.Instance->Guides.bIsSimulationCacheEnable);// || (WITH_EDITOR && PatchMode == EHairPatchAttribute::GuideInflucence);
 				const bool bUseSingleGuide	= InstanceData.bNeedDeformation && bValidGuide && InstanceData.Instance->Strands.InterpolationResource->UseSingleGuide();
 
-				check(InstanceData.Instance->Strands.Data);
+				check(InstanceData.Instance->Strands.IsValid());
 				AddPatchAttributePass(
 					GraphBuilder,
 					ShaderMap,
@@ -1523,7 +1523,7 @@ static void RunHairStrandsInterpolation_Strands(
 					PatchMode,
 					bValidGuide,
 					bUseSingleGuide,
-					*InstanceData.Instance->Strands.Data,
+					InstanceData.Instance->Strands.GetData(),
 					Register(GraphBuilder, InstanceData.Instance->Strands.RestResource->CurveAttributeBuffer, ERDGImportedBufferFlags::CreateSRV).Buffer,
 					RegisterAsSRV(GraphBuilder, InstanceData.Instance->Strands.RestResource->CurveBuffer),
 					RegisterAsSRV(GraphBuilder, InstanceData.Instance->Strands.ClusterResource->CurveToClusterIdBuffer),
@@ -1806,7 +1806,7 @@ static void RunHairStrandsInterpolation_Cards(
 
 				if (InstanceData.bNeedDeformation)
 				{
-					check(InstanceData.CardInstance->Guides.Data);
+					check(InstanceData.CardInstance->Guides.RestResource);
 				}
 				if (InstanceData.bSimulationEnable || InstanceData.bDeformationEnable)
 				{
@@ -2186,7 +2186,7 @@ static void AllocateRaytracingResources(FHairGroupInstance* Instance)
 {
 	if (IsHairRayTracingEnabled() && !Instance->Strands.RenRaytracingResource)
 	{
-		check(Instance->Strands.Data);
+		check(Instance->Strands.IsValid());
 
 #if RHI_ENABLE_RESOURCE_INFO
 		FName OwnerName(FString::Printf(TEXT("%s [LOD%d]"), *Instance->Debug.MeshComponentName, Instance->Debug.MeshLODIndex));
@@ -2196,7 +2196,7 @@ static void AllocateRaytracingResources(FHairGroupInstance* Instance)
 
 		// Allocate dynamic raytracing resources (owned by the groom component/instance)
 		FHairResourceName ResourceName(FName(Instance->Debug.GroomAssetName), Instance->Debug.GroupIndex);
-		Instance->Strands.RenRaytracingResource		 = new FHairStrandsRaytracingResource(*Instance->Strands.Data, ResourceName, OwnerName);
+		Instance->Strands.RenRaytracingResource		 = new FHairStrandsRaytracingResource(Instance->Strands.GetData(), ResourceName, OwnerName);
 		Instance->Strands.RenRaytracingResourceOwned = true;
 	}
 	Instance->Strands.ViewRayTracingMask |= EHairViewRayTracingMask::PathTracing;
@@ -2269,7 +2269,7 @@ void AddHairStreamingRequest(FHairGroupInstance* Instance, int32 InLODIndex)
 
 		// Lazy allocation of resources
 		// Note: Allocation will only be done if the resources is not initialized yet. Guides deformed position are also initialized from the Rest position at creation time.
-		if (Instance->Guides.Data && bLODNeedsGuides)
+		if (Instance->Guides.RestResource && bLODNeedsGuides)
 		{
 			if (Instance->Guides.RestRootResource)			{ Instance->Guides.RestRootResource->StreamInData(MeshLODIndex);}
 			if (Instance->Guides.RestResource)				{ Instance->Guides.RestResource->StreamInData(); }
@@ -2384,7 +2384,7 @@ static FHairLOD ComputeHairLODIndex(const FHairGroupInstance* Instance, const TA
 	for (const FSceneView* View : Views)
 	{
 		const FVector3d BoundScale = Instance->LocalToWorld.GetScale3D();
-		const FVector3f BoundExtent = Instance->Strands.Data ? FVector3f(Instance->Strands.Data->Header.BoundingBox.GetExtent()) : FVector3f(0,0,0);
+		const FVector3f BoundExtent = Instance->Strands.IsValid() ? FVector3f(Instance->Strands.GetData().Header.BoundingBox.GetExtent()) : FVector3f(0,0,0);
 		const float BoundRadius = FMath::Max3(BoundExtent.X, BoundExtent.Y, BoundExtent.Z) * FMath::Max3(BoundScale.X, BoundScale.Y, BoundScale.Z);
 
 		const float ScreenSize_RestBound = FMath::Clamp(ComputeBoundsScreenSize(FVector4(SphereBound.Center, 1), BoundRadius, *View), 0.f, 1.0f);
@@ -2439,7 +2439,7 @@ static FHairLOD ComputeHairLODIndex(const FHairGroupInstance* Instance, const TA
 
 	// Auto LOD
 	const bool bNeedAutoLOD = MaxLODIndexWithStrands != INDEX_NONE && int32(Out.HairLODIndex) <= MaxLODIndexWithStrands;
-	if (Instance->Strands.ClusterResource && Instance->Strands.Data)
+	if (Instance->Strands.ClusterResource && Instance->Strands.RestResource)
 	{
 		uint32 EffectiveCurveCount = 0;
 		if (bNeedAutoLOD && (Instance->HairGroupPublicData->bAutoLOD || IsHairStrandsForceAutoLODEnabled()))
@@ -2451,10 +2451,10 @@ static FHairLOD ComputeHairLODIndex(const FHairGroupInstance* Instance, const TA
 		{
 			EffectiveCurveCount = Instance->Strands.ClusterResource->BulkData.GetCurveCount(Out.HairLODIndex);
 		}
-		check(EffectiveCurveCount <= uint32(Instance->Strands.Data->Header.CurveToPointCount.Num()));
+		check(EffectiveCurveCount <= uint32(Instance->Strands.GetData().Header.CurveToPointCount.Num()));
 
 		Out.ContinuousLODCurveCount = EffectiveCurveCount;
-		Out.ContinuousLODPointCount = EffectiveCurveCount > 0 ? Instance->Strands.Data->Header.CurveToPointCount[EffectiveCurveCount - 1] : 0;
+		Out.ContinuousLODPointCount = EffectiveCurveCount > 0 ? Instance->Strands.GetData().Header.CurveToPointCount[EffectiveCurveCount - 1] : 0;
 		Out.ContinuousLODCoverageScale = ComputeActiveCurveCoverageScale(EffectiveCurveCount, Instance->HairGroupPublicData->RestCurveCount);
 	}
 
@@ -2537,7 +2537,7 @@ static bool SelectValidLOD(
 	{
 		// Round Curve/Point request to curve 'page'
 		RequestedCurveCount = GetRoundedCurveCount(Instance->HairGroupPublicData->ContinuousLODCurveCount, Instance->HairGroupPublicData->RestCurveCount);
-		RequestedPointCount = RequestedCurveCount > 0 ? Instance->Strands.Data->Header.CurveToPointCount[RequestedCurveCount - 1] : 0;
+		RequestedPointCount = RequestedCurveCount > 0 ? Instance->Strands.GetData().Header.CurveToPointCount[RequestedCurveCount - 1] : 0;
 	}
 
 	const bool bSimulationEnable			= Instance->HairGroupPublicData->IsSimulationEnable(IntHairLODIndex);
@@ -2553,7 +2553,7 @@ static bool SelectValidLOD(
 
 	// Lazy allocation of resources
 	// Note: Allocation will only be done if the resources is not initialized yet. Guides deformed position are also initialized from the Rest position at creation time.
-	if (Instance->Guides.Data && bLODNeedsGuides)
+	if (Instance->Guides.IsValid() && bLODNeedsGuides)
 	{
 		if (Instance->Guides.RestRootResource)			{ Instance->Guides.RestRootResource->Allocate(GraphBuilder, LoadingType, ResourceStatus, MeshLODIndex); }
 		if (Instance->Guides.RestResource)				{ Instance->Guides.RestResource->Allocate(GraphBuilder, LoadingType, ResourceStatus); }
@@ -2671,12 +2671,12 @@ static bool SelectValidLOD(
 			return false;
 		}
 		check(bIsLODDataReady);
-		check(Instance->Strands.RestResource);
+		check(Instance->Strands.IsValid());
 
 		// Adapt CurveCount/PointCount/CoverageScale based on what is actually available
 		const uint32 EffectiveCurveCount = FMath::Min(ResourceStatus.AvailableCurveCount, Instance->HairGroupPublicData->ContinuousLODCurveCount);
 		Instance->HairGroupPublicData->ContinuousLODCurveCount = EffectiveCurveCount;
-		Instance->HairGroupPublicData->ContinuousLODPointCount = EffectiveCurveCount > 0 ? Instance->Strands.Data->Header.CurveToPointCount[EffectiveCurveCount - 1] : 0;
+		Instance->HairGroupPublicData->ContinuousLODPointCount = EffectiveCurveCount > 0 ? Instance->Strands.GetData().Header.CurveToPointCount[EffectiveCurveCount - 1] : 0;
 		Instance->HairGroupPublicData->ContinuousLODCoverageScale = ComputeActiveCurveCoverageScale(EffectiveCurveCount, Instance->HairGroupPublicData->RestCurveCount);
 		check(Instance->HairGroupPublicData->ContinuousLODPointCount <= Instance->HairGroupPublicData->RestPointCount);
 	}
diff --git a/Engine/Plugins/Runtime/HairStrands/Source/HairStrandsCore/Private/GroomResources.cpp b/Engine/Plugins/Runtime/HairStrands/Source/HairStrandsCore/Private/GroomResources.cpp
--- a/Engine/Plugins/Runtime/HairStrands/Source/HairStrandsCore/Private/GroomResources.cpp
+++ b/Engine/Plugins/Runtime/HairStrands/Source/HairStrandsCore/Private/GroomResources.cpp
@@ -1795,13 +1795,11 @@ void CreateHairStrandsDebugDatas(
 	check(Out.VoxelData.Num()>0);
 }
 
-void CreateHairStrandsDebugResources(FRDGBuilder& GraphBuilder, const FHairStrandsDebugDatas* In, FHairStrandsDebugDatas::FResources* Out)
+void CreateHairStrandsDebugResources(FRDGBuilder& GraphBuilder, const FHairStrandsDebugDatas* In, FHairStrandsDebugResources* Out)
 {
 	check(In);
 	check(Out);
 
-	Out->VoxelDescription = In->VoxelDescription;
-
 	FRDGBufferRef VoxelOffsetAndCount = CreateStructuredBuffer(
 		GraphBuilder,
 		TEXT("HairStrandsDebug_VoxelOffsetAndCount"),
@@ -1818,6 +1816,7 @@ void CreateHairStrandsDebugResources(FRDGBuilder& GraphBuilder, const FHairStran
 		In->VoxelData.GetData(),
 		sizeof(FHairStrandsDebugDatas::FVoxel) * In->VoxelData.Num());
 
+	Out->VoxelDescription = In->VoxelDescription;
 	Out->VoxelOffsetAndCount = ConvertToExternalAccessBuffer(GraphBuilder, VoxelOffsetAndCount);
 	Out->VoxelData = ConvertToExternalAccessBuffer(GraphBuilder, VoxelData);
 }
diff --git a/Engine/Plugins/Runtime/HairStrands/Source/HairStrandsCore/Private/GroomTextureBuilder.cpp b/Engine/Plugins/Runtime/HairStrands/Source/HairStrandsCore/Private/GroomTextureBuilder.cpp
--- a/Engine/Plugins/Runtime/HairStrands/Source/HairStrandsCore/Private/GroomTextureBuilder.cpp
+++ b/Engine/Plugins/Runtime/HairStrands/Source/HairStrandsCore/Private/GroomTextureBuilder.cpp
@@ -1023,11 +1023,11 @@ static bool TraceTextures(
 			continue;
 		}
 
-		FHairGroupPlatformData& GroupData = InInfo.GroomAsset->GetHairGroupsPlatformData()[GroupIndex];
+		FHairGroupResources& GroupResource = InInfo.GroomAsset->GetHairGroupsResources()[GroupIndex];
 		FHairGroupsRendering& RenderingData = InInfo.GroomAsset->GetHairGroupsRendering()[GroupIndex];
 
-		FRDGBufferRef VoxelOffsetAndCount = GraphBuilder.RegisterExternalBuffer(GroupData.Debug.Resource->VoxelOffsetAndCount);
-		FRDGBufferRef VoxelData = GraphBuilder.RegisterExternalBuffer(GroupData.Debug.Resource->VoxelData);
+		FRDGBufferRef VoxelOffsetAndCount = GraphBuilder.RegisterExternalBuffer(GroupResource.Debug.Resource->VoxelOffsetAndCount);
+		FRDGBufferRef VoxelData = GraphBuilder.RegisterExternalBuffer(GroupResource.Debug.Resource->VoxelData);
 
 		{
 			FRHIShaderResourceView* PositionBuffer = nullptr;
@@ -1091,12 +1091,12 @@ static bool TraceTextures(
 			// Ensure the rest resources are loaded when rendering the strands textures
 			// Use EHairResourceLoadingType::Async, as the we can't use ::Sync because the resource is expected to be loading in async
 			// If a streaming request is in-flight, ensure it is done and allocate the resource
-			GroupData.Strands.RestResource->Allocate(GraphBuilder, EHairResourceLoadingType::Async);
+			GroupResource.Strands.RestResource->Allocate(GraphBuilder, EHairResourceLoadingType::Async);
 		#if WITH_EDITORONLY_DATA
-			if (GroupData.Strands.RestResource->StreamingRequest.DDCRequestOwner)
+			if (GroupResource.Strands.RestResource->StreamingRequest.DDCRequestOwner)
 			{
-				GroupData.Strands.RestResource->StreamingRequest.DDCRequestOwner->Wait();
-				GroupData.Strands.RestResource->Allocate(GraphBuilder, EHairResourceLoadingType::Async);
+				GroupResource.Strands.RestResource->StreamingRequest.DDCRequestOwner->Wait();
+				GroupResource.Strands.RestResource->Allocate(GraphBuilder, EHairResourceLoadingType::Async);
 			}
 		#endif	
 
@@ -1104,20 +1104,20 @@ static bool TraceTextures(
 			FHairStrandsInstanceRawParameters Instance;
 			Instance.Common.GroupIndex = GroupIndex;
 			Instance.Common.GroupCount = GroupCount;
-			Instance.Common.PointCount = GroupData.Strands.RestResource->BulkData.GetNumPoints();
-			Instance.Common.CurveCount = GroupData.Strands.RestResource->BulkData.GetNumCurves();
+			Instance.Common.PointCount = GroupResource.Strands.RestResource->BulkData.GetNumPoints();
+			Instance.Common.CurveCount = GroupResource.Strands.RestResource->BulkData.GetNumCurves();
 			Instance.Common.Radius = RenderingData.GeometrySettings.HairWidth * 0.5f;
 			Instance.Common.RootScale = 1.0f;
 			Instance.Common.TipScale = 1.0f;
 			Instance.Common.LengthScale = 1.0f;
-			Instance.Common.Length = GroupData.Strands.RestResource->BulkData.Header.MaxLength;
-			Instance.Common.PositionOffset = FVector3f(GroupData.Strands.RestResource->GetPositionOffset());
-			GetHairStrandsAttributeParameter(GroupData.Strands.RestResource->BulkData, Instance.Common.Attributes);
+			Instance.Common.Length = GroupResource.Strands.RestResource->BulkData.Header.MaxLength;
+			Instance.Common.PositionOffset = FVector3f(GroupResource.Strands.RestResource->GetPositionOffset());
+			GetHairStrandsAttributeParameter(GroupResource.Strands.RestResource->BulkData, Instance.Common.Attributes);
 
-			Instance.Resources.PositionBuffer = GroupData.Strands.RestResource->PositionBuffer.SRV;
-			Instance.Resources.CurveAttributeBuffer = GroupData.Strands.RestResource->CurveAttributeBuffer.SRV;
-			Instance.Resources.PointAttributeBuffer = GroupData.Strands.RestResource->PointAttributeBuffer.SRV;
-			Instance.Resources.PointToCurveBuffer = GroupData.Strands.RestResource->PointToCurveBuffer.SRV;
+			Instance.Resources.PositionBuffer = GroupResource.Strands.RestResource->PositionBuffer.SRV;
+			Instance.Resources.CurveAttributeBuffer = GroupResource.Strands.RestResource->CurveAttributeBuffer.SRV;
+			Instance.Resources.PointAttributeBuffer = GroupResource.Strands.RestResource->PointAttributeBuffer.SRV;
+			Instance.Resources.PointToCurveBuffer = GroupResource.Strands.RestResource->PointToCurveBuffer.SRV;
 
 			InternalGenerateHairStrandsTextures(
 				GraphBuilder,
@@ -1140,11 +1140,11 @@ static bool TraceTextures(
 				UVsBuffer,
 				TangentBuffer,
 
-				GroupData.Debug.Resource->VoxelDescription.VoxelMinBound,
-				GroupData.Debug.Resource->VoxelDescription.VoxelMaxBound,
-				GroupData.Debug.Resource->VoxelDescription.VoxelResolution,
-				GroupData.Debug.Resource->VoxelDescription.VoxelSize,
-				GroupData.Debug.Resource->VoxelDescription.MaxSegmentPerVoxel,
+				GroupResource.Debug.Resource->VoxelDescription.VoxelMinBound,
+				GroupResource.Debug.Resource->VoxelDescription.VoxelMaxBound,
+				GroupResource.Debug.Resource->VoxelDescription.VoxelResolution,
+				GroupResource.Debug.Resource->VoxelDescription.VoxelSize,
+				GroupResource.Debug.Resource->VoxelDescription.MaxSegmentPerVoxel,
 				VoxelOffsetAndCount,
 				VoxelData,
 				Instance,
diff --git a/Engine/Plugins/Runtime/HairStrands/Source/HairStrandsCore/Private/HairStrandsClusterCulling.cpp b/Engine/Plugins/Runtime/HairStrands/Source/HairStrandsCore/Private/HairStrandsClusterCulling.cpp
--- a/Engine/Plugins/Runtime/HairStrands/Source/HairStrandsCore/Private/HairStrandsClusterCulling.cpp
+++ b/Engine/Plugins/Runtime/HairStrands/Source/HairStrandsCore/Private/HairStrandsClusterCulling.cpp
@@ -196,7 +196,7 @@ void AddInstanceToClusterData(FHairGroupInstance* In, FHairStrandClusterData& Ou
 	// Initialize group cluster data for culling by the renderer
 	const int32 ClusterDataGroupIndex = Out.HairGroups.Num();
 	FHairStrandClusterData::FHairGroup& HairGroupCluster = Out.HairGroups.Emplace_GetRef();
-	HairGroupCluster.MaxPointPerCurve = In->Strands.Data ? In->Strands.Data->Header.MaxPointPerCurve : 0;
+	HairGroupCluster.MaxPointPerCurve = In->Strands.IsValid() ? In->Strands.GetData().Header.MaxPointPerCurve : 0;
 	HairGroupCluster.ClusterScale = In->HairGroupPublicData->ClusterScale;
 
 	HairGroupCluster.CurveBuffer = &In->Strands.RestResource->CurveBuffer;
diff --git a/Engine/Plugins/Runtime/HairStrands/Source/HairStrandsCore/Private/HairStrandsInterpolation.cpp b/Engine/Plugins/Runtime/HairStrands/Source/HairStrandsCore/Private/HairStrandsInterpolation.cpp
--- a/Engine/Plugins/Runtime/HairStrands/Source/HairStrandsCore/Private/HairStrandsInterpolation.cpp
+++ b/Engine/Plugins/Runtime/HairStrands/Source/HairStrandsCore/Private/HairStrandsInterpolation.cpp
@@ -826,7 +826,7 @@ void AddHairClusterAABBPass(
 	FRDGBufferUAVRef GroupAABBBUAV)
 {
 	// Clusters AABB are only update if the groom is deformed.
-	const FBoxSphereBounds& Bounds = Instance->Strands.Data->Header.BoundingBox;
+	const FBoxSphereBounds& Bounds = Instance->Strands.GetData().Header.BoundingBox;
 	FTransform InRenLocalToTranslatedWorld = Instance->LocalToWorld;
 	InRenLocalToTranslatedWorld.AddToTranslation(TranslatedWorldOffset);
 	const FBoxSphereBounds TransformedBounds = Bounds.TransformBy(InRenLocalToTranslatedWorld);
@@ -1723,11 +1723,11 @@ FHairGroupPublicData::FVertexFactoryInput InternalComputeHairStrandsVertexInputD
 
 	if (bUseGuideAttributeOffsets)
 	{
-		GetHairStrandsAttributeParameter(*Instance->Guides.Data, OutVFInput.Strands.Common.Attributes);
+		GetHairStrandsAttributeParameter(Instance->Guides.GetData(), OutVFInput.Strands.Common.Attributes);
 	}
 	else
 	{
-		GetHairStrandsAttributeParameter(*Instance->Strands.Data, OutVFInput.Strands.Common.Attributes);
+		GetHairStrandsAttributeParameter(Instance->Strands.GetData(), OutVFInput.Strands.Common.Attributes);
 	}
 
 	return OutVFInput;
diff --git a/Engine/Plugins/Runtime/HairStrands/Source/HairStrandsCore/Private/Niagara/NiagaraDataInterfaceHairStrands.cpp b/Engine/Plugins/Runtime/HairStrands/Source/HairStrandsCore/Private/Niagara/NiagaraDataInterfaceHairStrands.cpp
--- a/Engine/Plugins/Runtime/HairStrands/Source/HairStrandsCore/Private/Niagara/NiagaraDataInterfaceHairStrands.cpp
+++ b/Engine/Plugins/Runtime/HairStrands/Source/HairStrandsCore/Private/Niagara/NiagaraDataInterfaceHairStrands.cpp
@@ -732,7 +732,7 @@ void UNiagaraDataInterfaceHairStrands::ExtractDatasAndResources(
 		OutLocalToWorld = SystemInstance ? SystemInstance->GetWorldTransform() : FTransform::Identity;
 		if (OutGroupIndex < DefaultSource->GetNumHairGroups())
 		{
-			OutStrandsRestResource = DefaultSource->GetHairGroupsPlatformData()[OutGroupIndex].Guides.RestResource;
+			OutStrandsRestResource = DefaultSource->GetHairGroupsResources()[OutGroupIndex].Guides.RestResource;
 			OutGroomAsset = DefaultSource;
 		}
 	}
diff --git a/Engine/Plugins/Runtime/HairStrands/Source/HairStrandsCore/Public/GroomAsset.h b/Engine/Plugins/Runtime/HairStrands/Source/HairStrandsCore/Public/GroomAsset.h
--- a/Engine/Plugins/Runtime/HairStrands/Source/HairStrandsCore/Public/GroomAsset.h
+++ b/Engine/Plugins/Runtime/HairStrands/Source/HairStrandsCore/Public/GroomAsset.h
@@ -99,94 +99,176 @@ struct HAIRSTRANDSCORE_API FHairGroupsMaterial
 };
 
 
-/** Describe all data & resource for a groom asset's hair group */
-struct HAIRSTRANDSCORE_API FHairGroupPlatformData
+struct FHairGroupResources
 {
-	////////////////////////////////////////////////////////////////////////////
-	// Helper
-
-	struct FBase
+	struct FGuides
 	{
-		bool HasValidData() const		{ return BulkData.GetNumPoints() > 0;}
-		bool IsValid() const			{ return RestResource != nullptr; }
-		const FBox& GetBounds() const	{ return BulkData.GetBounds(); }
+		bool IsValid() const { return RestResource != nullptr; }
 
-		uint32 GetDataSize() const
+		/* Return the memory size for GPU resources */
+		uint32 GetResourcesSize() const
 		{
 			uint32 Total = 0;
-			Total += BulkData.Data.Positions.IsBulkDataLoaded() 		? BulkData.Data.Positions.GetBulkDataSize()   : 0;
-			Total += BulkData.Data.CurveAttributes.IsBulkDataLoaded()	? BulkData.Data.CurveAttributes.GetBulkDataSize() : 0;
-			Total += BulkData.Data.PointAttributes.IsBulkDataLoaded()	? BulkData.Data.PointAttributes.GetBulkDataSize() : 0;
-			Total += BulkData.Data.Curves.IsBulkDataLoaded() 			? BulkData.Data.Curves.GetBulkDataSize(): 0;
-			Total += BulkData.Data.PointToCurve.IsBulkDataLoaded()		? BulkData.Data.PointToCurve.GetBulkDataSize() : 0;
+			if (RestResource) Total += RestResource->GetResourcesSize();
 			return Total;
 		}
 
-		FHairStrandsBulkData				BulkData;
-		FHairStrandsRestResource*			RestResource = nullptr;
-	};
+		FHairStrandsRestResource* RestResource = nullptr;
 
-	struct FBaseWithInterpolation : FBase
+	} Guides;
+
+	struct FStrands
 	{
-		uint32 GetDataSize() const
+		bool IsValid() const { return RestResource != nullptr; }
+
+		/* Return the memory size for GPU resources */
+		uint32 GetResourcesSize() const
 		{
 			uint32 Total = 0;
-			Total += FBase::GetDataSize();
-			Total += InterpolationBulkData.Data.Interpolation.IsBulkDataLoaded()	? InterpolationBulkData.Data.Interpolation.GetBulkDataSize() : 0;
-			Total += InterpolationBulkData.Data.SimRootPointIndex.IsBulkDataLoaded()? InterpolationBulkData.Data.SimRootPointIndex.GetBulkDataSize() : 0;
+			if (RestResource) 			Total += RestResource->GetResourcesSize();
+			if (InterpolationResource) 	Total += InterpolationResource->GetResourcesSize();
+			if (ClusterResource) 		Total += ClusterResource->GetResourcesSize();
+			#if RHI_RAYTRACING
+			if (RaytracingResource) 	Total += RaytracingResource->GetResourcesSize();
+			#endif
 			return Total;
 		}
 
-		FHairStrandsInterpolationBulkData	InterpolationBulkData;
+		FHairStrandsRestResource*			RestResource = nullptr;
 		FHairStrandsInterpolationResource*	InterpolationResource = nullptr;
-	};
+		FHairStrandsClusterResource* 		ClusterResource = nullptr;
+		#if RHI_RAYTRACING
+		FHairStrandsRaytracingResource*		RaytracingResource = nullptr;
+		#endif
 
-	////////////////////////////////////////////////////////////////////////////
-	// Data
+		bool bIsCookedOut = false;
+	} Strands;
 
-	struct FGuides : FBase
+	struct FCards
 	{
 		/* Return the memory size for GPU resources */
 		uint32 GetResourcesSize() const
 		{
 			uint32 Total = 0;
-			if (RestResource) Total += RestResource->GetResourcesSize();
+			for (const FLOD& LOD : LODs)
+			{
+				Total += LOD.GetResourcesSize();
+			}
 			return Total;
 		}
 
-		uint32 GetDataSize() const
+		struct FLOD
 		{
-			uint32 Total = 0;
-			Total += FBase::GetDataSize();
-			return Total;
-		}
+			/* Return the memory size for GPU resources */
+			uint32 GetResourcesSize() const
+			{
+				uint32 Total = 0;
+				if (RestResource) 				Total += RestResource->GetResourcesSize();
+				if (InterpolationResource) 		Total += InterpolationResource->GetResourcesSize();
+				if (GuideRestResource) 			Total += GuideRestResource->GetResourcesSize();
+				if (GuideInterpolationResource) Total += GuideInterpolationResource->GetResourcesSize();
+				#if RHI_RAYTRACING
+				if (RaytracingResource) 		Total += RaytracingResource->GetResourcesSize();
+				#endif
+				return Total;
+			}
 
-	} Guides;
+			bool IsValid() const { return RestResource != nullptr; }
 
-	struct FStrands : FBaseWithInterpolation
-	{
+			FHairCardsRestResource*				RestResource = nullptr;
+			FHairCardsInterpolationResource*	InterpolationResource = nullptr;
+			FHairStrandsRestResource*			GuideRestResource = nullptr;
+			FHairStrandsInterpolationResource*	GuideInterpolationResource = nullptr;
+			#if RHI_RAYTRACING
+			FHairStrandsRaytracingResource*		RaytracingResource = nullptr;
+			#endif
+
+			bool bIsCookedOut = false;
+		};
+		TArray<FLOD> LODs;
+	} Cards;
+
+	struct FMeshes
+	{	
 		/* Return the memory size for GPU resources */
 		uint32 GetResourcesSize() const
 		{
 			uint32 Total = 0;
-			if (RestResource) Total += RestResource->GetResourcesSize();
-			if (InterpolationResource) Total += InterpolationResource->GetResourcesSize();
-			if (ClusterResource) Total += ClusterResource->GetResourcesSize();
-			#if RHI_RAYTRACING
-			if (RaytracingResource) Total += RaytracingResource->GetResourcesSize();
-			#endif
+			for (const FLOD& LOD : LODs)
+			{
+				Total += LOD.GetResourcesSize();
+			}
 			return Total;
 		}
 
-		uint32 GetDataSize() const;
+		struct FLOD
+		{
+			/* Return the memory size for GPU resources */
+			uint32 GetResourcesSize() const
+			{
+				uint32 Total = 0;
+				if (RestResource) 		Total += RestResource->GetResourcesSize();
+				#if RHI_RAYTRACING
+				if (RaytracingResource) Total += RaytracingResource->GetResourcesSize();
+				#endif
+				return Total;
+			}
 
-		FHairStrandsClusterBulkData	ClusterBulkData;
-		FHairStrandsClusterResource* ClusterResource = nullptr;
+			bool IsValid() const { return RestResource != nullptr; }
 
-		#if RHI_RAYTRACING
-		FHairStrandsRaytracingResource*		RaytracingResource = nullptr;
-		#endif
+			FHairMeshesRestResource* 		RestResource = nullptr;
+			#if RHI_RAYTRACING
+			FHairStrandsRaytracingResource* RaytracingResource = nullptr;
+			#endif
+			bool bIsCookedOut = false;
+		};
+		TArray<FLOD> LODs;
+	} Meshes;
+
+	struct FDebug
+	{
+		FHairStrandsDebugResources* Resource = nullptr;
+	} Debug;
+};
+
+FORCEINLINE uint32 GetDataSize(const FHairStrandsBulkData& BulkData)
+{
+	uint32 Total = 0;
+	Total += BulkData.Data.Positions.IsBulkDataLoaded() 		? BulkData.Data.Positions.GetBulkDataSize()   : 0;
+	Total += BulkData.Data.CurveAttributes.IsBulkDataLoaded()	? BulkData.Data.CurveAttributes.GetBulkDataSize() : 0;
+	Total += BulkData.Data.PointAttributes.IsBulkDataLoaded()	? BulkData.Data.PointAttributes.GetBulkDataSize() : 0;
+	Total += BulkData.Data.Curves.IsBulkDataLoaded() 			? BulkData.Data.Curves.GetBulkDataSize(): 0;
+	Total += BulkData.Data.PointToCurve.IsBulkDataLoaded()		? BulkData.Data.PointToCurve.GetBulkDataSize() : 0;
+	return Total;
+}
+
+FORCEINLINE uint32 GetDataSize(const FHairStrandsInterpolationBulkData& InterpolationBulkData) 	
+{
+	uint32 Total = 0;
+	Total += InterpolationBulkData.Data.Interpolation.IsBulkDataLoaded()	? InterpolationBulkData.Data.Interpolation.GetBulkDataSize() : 0;
+	Total += InterpolationBulkData.Data.SimRootPointIndex.IsBulkDataLoaded()? InterpolationBulkData.Data.SimRootPointIndex.GetBulkDataSize() : 0;
+	return Total;
+}
+
+struct FHairGroupPlatformData
+{
+	struct FGuides
+	{
+		bool HasValidData() const		{ return BulkData.GetNumPoints() > 0;}
+		const FBox& GetBounds() const	{ return BulkData.GetBounds(); }
+		uint32 GetDataSize() const		{ return ::GetDataSize(BulkData); }
+		FHairStrandsBulkData BulkData;
+	} Guides;
+
+	struct FStrands
+	{
+		bool HasValidData() const		{ return BulkData.GetNumPoints() > 0;}
+		const FBox& GetBounds() const	{ return BulkData.GetBounds(); }
 
+		uint32 GetDataSize() const;
+		FHairStrandsBulkData				BulkData;
+		FHairStrandsInterpolationBulkData	InterpolationBulkData;
+		FHairStrandsClusterBulkData			ClusterBulkData;
 		bool bIsCookedOut = false;
 	} Strands;
 
@@ -197,13 +279,13 @@ struct HAIRSTRANDSCORE_API FHairGroupPlatformData
 			for (const FLOD& LOD : LODs)
 			{
 				if (LOD.HasValidData())
-					return true;
+				return true;
 			}
 			return false;
 		}
 
-		bool HasValidData(uint32 LODIt) const { return LODIt < uint32(LODs.Num()) && LODs[LODIt].HasValidData(); }
-		bool IsValid(uint32 LODIt) const { return LODIt < uint32(LODs.Num()) && LODs[LODIt].IsValid(); }
+		bool HasValidData(uint32 LODIt) const 	{ return LODIt < uint32(LODs.Num()) && LODs[LODIt].HasValidData(); }
+		bool IsValid(uint32 LODIt) const 		{ return LODIt < uint32(LODs.Num()) && LODs[LODIt].IsValid(); }
 		FBox GetBounds() const
 		{
 			for (const FLOD& LOD : LODs)
@@ -213,17 +295,6 @@ struct HAIRSTRANDSCORE_API FHairGroupPlatformData
 			return FBox();
 		}
 
-		/* Return the memory size for GPU resources */
-		uint32 GetResourcesSize() const
-		{
-			uint32 Total = 0;
-			for (const FLOD& LOD : LODs)
-			{
-				Total += LOD.GetResourcesSize();
-			}
-			return Total;
-		}
-
 		uint32 GetDataSize() const
 		{
 			uint32 Total = 0;
@@ -236,18 +307,6 @@ struct HAIRSTRANDSCORE_API FHairGroupPlatformData
 
 		struct FLOD
 		{
-			/* Return the memory size for GPU resources */
-			uint32 GetResourcesSize() const
-			{
-				uint32 Total = 0;
-				if (RestResource) Total += RestResource->GetResourcesSize();
-				if (InterpolationResource) Total += InterpolationResource->GetResourcesSize();
-				#if RHI_RAYTRACING
-				if (RaytracingResource) Total += RaytracingResource->GetResourcesSize();
-				#endif
-				return Total;
-			}
-
 			uint32 GetDataSize() const
 			{
 				uint32 Total = 0;
@@ -255,29 +314,20 @@ struct HAIRSTRANDSCORE_API FHairGroupPlatformData
 				Total += BulkData.Normals.GetAllocatedSize();
 				Total += BulkData.UVs.GetAllocatedSize();
 				Total += BulkData.Indices.GetAllocatedSize();
-
 				Total += InterpolationBulkData.Interpolation.GetAllocatedSize();
-
-				Total += Guides.GetDataSize();
+				Total += ::GetDataSize(GuideBulkData);
+				Total += ::GetDataSize(GuideInterpolationBulkData);
 				return Total;
 			}
 
-			bool HasValidData() const { return BulkData.IsValid(); }
-			bool IsValid() const { return BulkData.IsValid() && RestResource != nullptr; }
+			bool HasValidData() const	{ return BulkData.IsValid(); }
+			bool IsValid() const 		{ return BulkData.IsValid(); }
 
 			// Main data & Resources
 			FHairCardsBulkData					BulkData;
-			FHairCardsRestResource*				RestResource = nullptr;
-
-			// Interpolation data & resources
 			FHairCardsInterpolationBulkData		InterpolationBulkData;
-			FHairCardsInterpolationResource*	InterpolationResource = nullptr;
-
-			FBaseWithInterpolation				Guides;
-
-			#if RHI_RAYTRACING
-			FHairStrandsRaytracingResource*		RaytracingResource = nullptr;
-			#endif
+			FHairStrandsBulkData				GuideBulkData;
+			FHairStrandsInterpolationBulkData	GuideInterpolationBulkData;
 
 			bool bIsCookedOut = false;
 		};
@@ -291,12 +341,12 @@ struct HAIRSTRANDSCORE_API FHairGroupPlatformData
 			for (const FLOD& LOD : LODs)
 			{
 				if (LOD.HasValidData())
-					return true;
+				return true;
 			}
 			return false;
 		}
-		bool HasValidData(uint32 LODIt) const { return LODIt < uint32(LODs.Num()) && LODs[LODIt].HasValidData(); }
-		bool IsValid(uint32 LODIt) const { return LODIt < uint32(LODs.Num()) && LODs[LODIt].IsValid(); }
+		bool HasValidData(uint32 LODIt) const 	{ return LODIt < uint32(LODs.Num()) && LODs[LODIt].HasValidData(); }
+		bool IsValid(uint32 LODIt) const 		{ return LODIt < uint32(LODs.Num()) && LODs[LODIt].IsValid(); }
 		FBox GetBounds() const
 		{
 			for (const FLOD& LOD : LODs)
@@ -305,17 +355,6 @@ struct HAIRSTRANDSCORE_API FHairGroupPlatformData
 			}
 			return FBox();
 		}
-		
-		/* Return the memory size for GPU resources */
-		uint32 GetResourcesSize() const
-		{
-			uint32 Total = 0;
-			for (const FLOD& LOD : LODs)
-			{
-				Total += LOD.GetResourcesSize();
-			}
-			return Total;
-		}
 
 		uint32 GetDataSize() const
 		{
@@ -329,17 +368,6 @@ struct HAIRSTRANDSCORE_API FHairGroupPlatformData
 
 		struct FLOD
 		{
-			/* Return the memory size for GPU resources */
-			uint32 GetResourcesSize() const
-			{
-				uint32 Total = 0;
-				if (RestResource) Total += RestResource->GetResourcesSize();
-				#if RHI_RAYTRACING
-				if (RaytracingResource) Total += RaytracingResource->GetResourcesSize();
-				#endif
-				return Total;
-			}
-
 			uint32 GetDataSize() const
 			{
 				uint32 Total = 0;
@@ -350,14 +378,10 @@ struct HAIRSTRANDSCORE_API FHairGroupPlatformData
 				return Total;
 			}
 
-			bool HasValidData() const { return BulkData.IsValid(); }
-			bool IsValid() const { return BulkData.IsValid() && RestResource != nullptr; }
+			bool HasValidData() const 	{ return BulkData.IsValid(); }
+			bool IsValid() const 		{ return BulkData.IsValid(); }
 
 			FHairMeshesBulkData BulkData;
-			FHairMeshesRestResource* RestResource = nullptr;
-			#if RHI_RAYTRACING
-			FHairStrandsRaytracingResource* RaytracingResource = nullptr;
-			#endif
 			bool bIsCookedOut = false;
 		};
 		TArray<FLOD> LODs;
@@ -366,7 +390,6 @@ struct HAIRSTRANDSCORE_API FHairGroupPlatformData
 	struct FDebug
 	{
 		FHairStrandsDebugDatas Data;
-		FHairStrandsDebugDatas::FResources* Resource = nullptr;
 	} Debug;
 };
 
@@ -509,6 +532,10 @@ public:
 	UE_DEPRECATED(5.3, "Please do not access this member directly; use UGroomAsset accessor.")
 	TArray<FHairGroupPlatformData> HairGroupsPlatformData;
 
+private:
+	/** Store strands/cards/meshes resources */
+	TArray<FHairGroupResources> HairGroupsResources;
+
 public:
 	static FName GetHairGroupsRenderingMemberName();
 	UFUNCTION(BlueprintGetter) TArray<FHairGroupsRendering>& GetHairGroupsRendering();
@@ -589,6 +616,9 @@ public:
 	const TArray<FHairGroupInfoWithVisibility>& GetHairGroupsInfo() const;
 	void SetHairGroupsInfo(const TArray<FHairGroupInfoWithVisibility>& In);
 
+	const TArray<FHairGroupResources>& GetHairGroupsResources() const;
+	TArray<FHairGroupResources>& GetHairGroupsResources();
+
 public:
 
 	//~ Begin UObject Interface.
@@ -597,12 +627,6 @@ public:
 	virtual void BeginDestroy() override;
 	virtual void Serialize(FArchive& Ar) override;
 
-	PRAGMA_DISABLE_DEPRECATION_WARNINGS // Suppress compiler warning on override of deprecated function
-	UE_DEPRECATED(5.0, "Use version that takes FObjectPreSaveContext instead.")
-	virtual void PreSave(const class ITargetPlatform* TargetPlatform) override;
-	PRAGMA_ENABLE_DEPRECATION_WARNINGS
-	virtual void PreSave(FObjectPreSaveContext ObjectSaveContext) override;
-
 #if WITH_EDITOR
 	FOnGroomAssetChanged& GetOnGroomAssetChanged() { return OnGroomAssetChanged;  }
 	FOnGroomAssetResourcesChanged& GetOnGroomAssetResourcesChanged() { return OnGroomAssetResourcesChanged; }
@@ -743,11 +767,14 @@ private:
 #if WITH_EDITORONLY_DATA
 	bool HasImportedStrandsData() const;
 
-	bool BuildCardsData(uint32 GroupIndex);
-	bool BuildMeshesData(uint32 GroupIndex);
+	bool BuildHairGroup_Cards(uint32 GroupIndex);
+	bool BuildHairGroup_Meshes(uint32 GroupIndex);
+
+	bool HasChanged_Cards(uint32 GroupIndex, TArray<bool>& OutIsValid) const;
+	bool HasChanged_Meshes(uint32 GroupIndex, TArray<bool>& OutIsValid) const;
 
-	bool HasValidCardsData(uint32 GroupIndex) const;
-	bool HasValidMeshesData(uint32 GroupIndex) const;
+	bool HasValidData_Cards(uint32 GroupIndex) const;
+	bool HasValidData_Meshes(uint32 GroupIndex) const;
 public:
 	enum EHairDescriptionType
 	{
@@ -843,7 +870,7 @@ struct FGroomAssetMemoryStats
 	FStrandsDetails Memory;
 	FStrandsDetails Curves;
 
-	static FGroomAssetMemoryStats Get(const FHairGroupPlatformData& In);
+	static FGroomAssetMemoryStats Get(const FHairGroupPlatformData& InData, const FHairGroupResources& In);
 	void Accumulate(const FGroomAssetMemoryStats& In);
 	uint32 GetTotalCPUSize() const;
 	uint32 GetTotalGPUSize() const;
diff --git a/Engine/Plugins/Runtime/HairStrands/Source/HairStrandsCore/Public/GroomBindingAsset.h b/Engine/Plugins/Runtime/HairStrands/Source/HairStrandsCore/Public/GroomBindingAsset.h
--- a/Engine/Plugins/Runtime/HairStrands/Source/HairStrandsCore/Public/GroomBindingAsset.h
+++ b/Engine/Plugins/Runtime/HairStrands/Source/HairStrandsCore/Public/GroomBindingAsset.h
@@ -19,7 +19,6 @@ class UGeometryCache;
 class UMaterialInterface;
 class UNiagaraSystem;
 class UGroomAsset;
-struct FHairGroupPlatformData;
 
 USTRUCT(BlueprintType)
 struct HAIRSTRANDSCORE_API FGoomBindingGroupInfo
diff --git a/Engine/Plugins/Runtime/HairStrands/Source/HairStrandsCore/Public/GroomInstance.h b/Engine/Plugins/Runtime/HairStrands/Source/HairStrandsCore/Public/GroomInstance.h
--- a/Engine/Plugins/Runtime/HairStrands/Source/HairStrandsCore/Public/GroomInstance.h
+++ b/Engine/Plugins/Runtime/HairStrands/Source/HairStrandsCore/Public/GroomInstance.h
@@ -79,15 +79,13 @@ struct HAIRSTRANDSCORE_API FHairGroupInstance : public FHairStrandsInstance
 	virtual ~FHairGroupInstance();
 
 	//////////////////////////////////////////////////////////////////////////////////////////
-	// Helper struct which aggregate strands based data/resources
-
-	struct FStrandsBase
+	// Simulation
+	struct FGuides
 	{
-		bool HasValidData() const { return Data != nullptr && Data->GetNumPoints() > 0; }
-		bool IsValid() const { return RestResource != nullptr; }
+		const FHairStrandsBulkData& GetData() const { return RestResource->BulkData; }
 
-		// Data - Render & sim (rest) data
-		FHairStrandsBulkData* Data = nullptr;
+		bool IsValid() const { return RestResource != nullptr; }
+		bool HasValidRootData() const { return RestRootResource != nullptr && DeformedRootResource != nullptr; }
 
 		// Resources - Strands rest position data for sim & render strands
 		// Resources - Strands deformed position data for sim & render strands
@@ -99,25 +97,6 @@ struct HAIRSTRANDSCORE_API FHairGroupInstance : public FHairStrandsInstance
 		FHairStrandsRestRootResource* RestRootResource = nullptr;
 		FHairStrandsDeformedRootResource* DeformedRootResource = nullptr;
 
-		bool HasValidRootData() const { return RestRootResource != nullptr && DeformedRootResource != nullptr; }
-	};
-
-	struct FStrandsBaseWithInterpolation : FStrandsBase
-	{
-		// Data - Interpolation data (weights/Id/...) for transfering sim strands (i.e. guide) motion to render strands
-		// Resources - Strands deformed position data for sim & render strands
-		FHairStrandsInterpolationResource* InterpolationResource = nullptr;
-
-		EHairInterpolationType HairInterpolationType = EHairInterpolationType::NoneSkinning;
-
-		// Indicates if culling is enabled for this hair strands data.
-		bool bCullingEnable = false;
-	};
-
-	//////////////////////////////////////////////////////////////////////////////////////////
-	// Simulation
-	struct FGuides : FStrandsBase
-	{
 		bool bIsSimulationEnable = false;
 		bool bIsDeformationEnable = false;
 		bool bHasGlobalInterpolation = false;
@@ -126,12 +105,31 @@ struct HAIRSTRANDSCORE_API FHairGroupInstance : public FHairStrandsInstance
 
 	//////////////////////////////////////////////////////////////////////////////////////////
 	// Strands
-	struct FStrands : FStrandsBaseWithInterpolation
+	struct FStrands
 	{
+		const FHairStrandsBulkData& GetData() const { return RestResource->BulkData; }
+
+		bool IsValid() const { return RestResource != nullptr; }
+		bool HasValidRootData() const { return RestRootResource != nullptr && DeformedRootResource != nullptr; }
+
+		// Resources - Strands rest position data for sim & render strands
+		// Resources - Strands deformed position data for sim & render strands
+		FHairStrandsRestResource* RestResource = nullptr;
+		FHairStrandsDeformedResource* DeformedResource = nullptr;
+
+		// Data - Interpolation data (weights/Id/...) for transfering sim strands (i.e. guide) motion to render strands
+		// Resources - Strands deformed position data for sim & render strands
+		FHairStrandsInterpolationResource* InterpolationResource = nullptr;
+
 		// Resources - Strands cluster data for culling/voxelization purpose
 		FHairStrandsClusterResource* ClusterResource = nullptr;
 		FHairStrandsCullingResource* CullingResource = nullptr;
 
+		// Resources - Rest root data, for deforming strands attached to a skinned mesh surface
+		// Resources - Deformed root data, for deforming strands attached to a skinned mesh surface
+		FHairStrandsRestRootResource* RestRootResource = nullptr;
+		FHairStrandsDeformedRootResource* DeformedRootResource = nullptr;
+
 		// Resources - Raytracing data when enabling (expensive) raytracing method
 		#if RHI_RAYTRACING
 		FHairStrandsRaytracingResource* RenRaytracingResource = nullptr;
@@ -148,6 +146,9 @@ struct HAIRSTRANDSCORE_API FHairGroupInstance : public FHairStrandsInstance
 
 		FHairStrandsUniformBuffer UniformBuffer;
 		FHairStrandsVertexFactory* VertexFactory = nullptr;
+
+		EHairInterpolationType HairInterpolationType = EHairInterpolationType::NoneSkinning;		
+		bool bCullingEnable = false; // Indicates if culling is enabled for this hair strands data.
 	} Strands;
 
 	//////////////////////////////////////////////////////////////////////////////////////////
@@ -159,22 +160,38 @@ struct HAIRSTRANDSCORE_API FHairGroupInstance : public FHairStrandsInstance
 		struct FLOD
 		{
 			bool IsValid() const { return RestResource != nullptr; }
-
-			// Data
-			FHairCardsBulkData* Data = nullptr;
+			const FHairCardsBulkData& GetData() const { return RestResource->BulkData; }
 
 			// Resources
 			FHairCardsRestResource* RestResource = nullptr;
 			FHairCardsDeformedResource* DeformedResource = nullptr;
+			FHairCardsInterpolationResource* InterpolationResource = nullptr;
 			#if RHI_RAYTRACING
 			FHairStrandsRaytracingResource* RaytracingResource = nullptr;
 			bool RaytracingResourceOwned = false;
 			#endif
 
-			// Interpolation data/resources
-			FHairCardsInterpolationResource* InterpolationResource = nullptr;
+			struct FGuides
+			{
+				bool IsValid() const { return RestResource != nullptr; }
+				bool HasValidRootData() const { return RestRootResource != nullptr && DeformedRootResource != nullptr; }
+				const FHairStrandsBulkData& GetData() const { return RestResource->BulkData; }
+
+				// Resources - Strands rest position data for sim & render strands
+				// Resources - Strands deformed position data for sim & render strands
+				FHairStrandsRestResource* RestResource = nullptr;
+				FHairStrandsDeformedResource* DeformedResource = nullptr;
+	
+				// Resources - Rest root data, for deforming strands attached to a skinned mesh surface
+				// Resources - Deformed root data, for deforming strands attached to a skinned mesh surface
+				FHairStrandsRestRootResource* RestRootResource = nullptr;
+				FHairStrandsDeformedRootResource* DeformedRootResource = nullptr;
+	
+				// Resources - Strands deformed position data for sim & render strands
+				FHairStrandsInterpolationResource* InterpolationResource = nullptr;
 
-			FStrandsBaseWithInterpolation Guides;
+				EHairInterpolationType HairInterpolationType = EHairInterpolationType::NoneSkinning;
+			} Guides;
 
 			FHairCardsUniformBuffer UniformBuffer;
 			FHairCardsVertexFactory* VertexFactory = nullptr;
@@ -197,9 +214,7 @@ struct HAIRSTRANDSCORE_API FHairGroupInstance : public FHairStrandsInstance
 		struct FLOD
 		{
 			bool IsValid() const { return RestResource != nullptr; }
-
-			// Data
-			FHairMeshesBulkData* Data = nullptr;
+			const FHairMeshesBulkData& GetData() const { return RestResource->BulkData; }
 
 			// Resources
 			FHairMeshesRestResource* RestResource = nullptr;
@@ -251,7 +266,7 @@ struct HAIRSTRANDSCORE_API FHairGroupInstance : public FHairStrandsInstance
 		TSharedPtr<class IGroomCacheBuffers, ESPMode::ThreadSafe> GroomCacheBuffers;
 
 		// Resources
-		FHairStrandsDebugDatas::FResources* HairDebugResource = nullptr;
+		FHairStrandsDebugResources* HairDebugResource = nullptr;
 	} Debug;
 
 	FTransform				LocalToWorld = FTransform::Identity;
@@ -265,10 +280,10 @@ struct HAIRSTRANDSCORE_API FHairGroupInstance : public FHairStrandsInstance
 	bool					bHoldout = false;
 	
 	// Deformed component to extract the bone buffer 
-	UMeshComponent*	 DeformedComponent = nullptr;
+	UMeshComponent* DeformedComponent = nullptr;
 	
 	// Section of the deformed component to be used 
-	int32	 DeformedSection = INDEX_NONE;
+	int32 DeformedSection = INDEX_NONE;
 	
 	bool IsValid() const 
 	{
diff --git a/Engine/Plugins/Runtime/HairStrands/Source/HairStrandsCore/Public/GroomResources.h b/Engine/Plugins/Runtime/HairStrands/Source/HairStrandsCore/Public/GroomResources.h
--- a/Engine/Plugins/Runtime/HairStrands/Source/HairStrandsCore/Public/GroomResources.h
+++ b/Engine/Plugins/Runtime/HairStrands/Source/HairStrandsCore/Public/GroomResources.h
@@ -759,6 +759,7 @@ struct FHairMeshesRestResource : public FHairCommonResource
 	/* Strand hair rest position buffer */
 	FHairCardsVertexBuffer RestPositionBuffer;
 	FHairCardIndexBuffer IndexBuffer;	
+	bool bInvertUV = false;
 
 	FHairCardsVertexBuffer NormalsBuffer;
 	FHairCardsVertexBuffer UVsBuffer;
@@ -821,4 +822,4 @@ struct FHairMeshesDeformedResource : public FHairCommonResource
 ////////////////////////////////////////////////////////////////////////////////////////////////////////////////////////
 // Debug data (used for debug visalization but also for texture generation)
 void CreateHairStrandsDebugDatas(const FHairStrandsDatas& InData, FHairStrandsDebugDatas& Out);
-void CreateHairStrandsDebugResources(class FRDGBuilder& GraphBuilder, const FHairStrandsDebugDatas* In, FHairStrandsDebugDatas::FResources* Out);
+void CreateHairStrandsDebugResources(class FRDGBuilder& GraphBuilder, const FHairStrandsDebugDatas* In, FHairStrandsDebugResources* Out);
diff --git a/Engine/Plugins/Runtime/HairStrands/Source/HairStrandsCore/Public/HairStrandsDatas.h b/Engine/Plugins/Runtime/HairStrands/Source/HairStrandsCore/Public/HairStrandsDatas.h
--- a/Engine/Plugins/Runtime/HairStrands/Source/HairStrandsCore/Public/HairStrandsDatas.h
+++ b/Engine/Plugins/Runtime/HairStrands/Source/HairStrandsCore/Public/HairStrandsDatas.h
@@ -931,12 +931,11 @@ struct HAIRSTRANDSCORE_API FHairStrandsDebugDatas
 	FDesc VoxelDescription;
 	TArray<FOffsetAndCount> VoxelOffsetAndCount;
 	TArray<FVoxel> VoxelData;
+};
 
-	struct FResources
-	{
-		FDesc VoxelDescription;
-
-		TRefCountPtr<FRDGPooledBuffer> VoxelOffsetAndCount;
-		TRefCountPtr<FRDGPooledBuffer> VoxelData;
-	};
+struct HAIRSTRANDSCORE_API FHairStrandsDebugResources
+{
+	FHairStrandsDebugDatas::FDesc VoxelDescription;
+	TRefCountPtr<FRDGPooledBuffer> VoxelOffsetAndCount;
+	TRefCountPtr<FRDGPooledBuffer> VoxelData;
 };
