Fix hair instance streaming sharing same asset.

diff --git a/Engine/Plugins/Runtime/HairStrands/Source/HairStrandsCore/Private/GroomAsset.cpp b/Engine/Plugins/Runtime/HairStrands/Source/HairStrandsCore/Private/GroomAsset.cpp
--- a/Engine/Plugins/Runtime/HairStrands/Source/HairStrandsCore/Private/GroomAsset.cpp
+++ b/Engine/Plugins/Runtime/HairStrands/Source/HairStrandsCore/Private/GroomAsset.cpp
@@ -90,6 +90,11 @@ bool IsHairStrandsDDCLogEnable()
 	return GHairStrandsDDCLogEnable > 0;
 }
 
+uint32 GetAssetNameHash(const FString& In)
+{
+	return CityHash32((const char*)*In, In.Len() * sizeof(FString::ElementType));
+}
+
 static int32 GHairMaxSimulatedLOD = -1;
 static FAutoConsoleVariableRef CVarHairMaxSimulatedLOD(TEXT("r.HairStrands.MaxSimulatedLOD"), GHairMaxSimulatedLOD, TEXT("Maximum hair LOD to be simulated"));
 static bool IsHairLODSimulationEnabled(const int32 LODIndex) { return (LODIndex >= 0 && (GHairMaxSimulatedLOD < 0 || (GHairMaxSimulatedLOD >= 0 && LODIndex <= GHairMaxSimulatedLOD))); }
@@ -1209,6 +1214,9 @@ void UGroomAsset::PostLoad()
 
 	Super::PostLoad();
 
+	// Compute a hash of the Groom asset fullname for finding unique groom during LOD selection/streaming
+	AssetNameHash = GetAssetNameHash(GetFullName());
+
 	// Legacy asset are missing rendering or interpolation settings
 #if WITH_EDITORONLY_DATA
 	const bool bIsLegacyAsset = GetHairGroupsInterpolation().Num() == 0;
diff --git a/Engine/Plugins/Runtime/HairStrands/Source/HairStrandsCore/Private/GroomBindingAsset.cpp b/Engine/Plugins/Runtime/HairStrands/Source/HairStrandsCore/Private/GroomBindingAsset.cpp
--- a/Engine/Plugins/Runtime/HairStrands/Source/HairStrandsCore/Private/GroomBindingAsset.cpp
+++ b/Engine/Plugins/Runtime/HairStrands/Source/HairStrandsCore/Private/GroomBindingAsset.cpp
@@ -31,6 +31,7 @@ LLM_DECLARE_TAG(Groom);
 /////////////////////////////////////////////////////////////////////////////////////////
 
 bool IsHairStrandsDDCLogEnable();
+uint32 GetAssetNameHash(const FString& In);
 
 ////////////////////////////////////////////////////////////////////////////////////////////////////////////
 
@@ -324,6 +325,10 @@ void UGroomBindingAsset::PostLoad()
 	LLM_SCOPE_BYTAG(Groom);
 
 	Super::PostLoad();
+
+	// Compute a hash of the Groom asset fullname for finding unique groom during LOD selection/streaming
+	AssetNameHash = GetAssetNameHash(GetFullName());
+
 	if (UGroomAsset* LocalGroom = GetGroom())
 	{
 		// Make sure that the asset initialized its resources first since the component needs them to initialize its own resources
diff --git a/Engine/Plugins/Runtime/HairStrands/Source/HairStrandsCore/Private/GroomComponent.cpp b/Engine/Plugins/Runtime/HairStrands/Source/HairStrandsCore/Private/GroomComponent.cpp
--- a/Engine/Plugins/Runtime/HairStrands/Source/HairStrandsCore/Private/GroomComponent.cpp
+++ b/Engine/Plugins/Runtime/HairStrands/Source/HairStrandsCore/Private/GroomComponent.cpp
@@ -2639,11 +2639,13 @@ void UGroomComponent::InitResources(bool bIsBindingReloading)
 		HairGroupInstances.Add(HairGroupInstance);
 		HairGroupInstance->Debug.GroupIndex = GroupIt;
 		HairGroupInstance->Debug.GroupCount = GroupCount;
+		HairGroupInstance->Debug.GroomAssetHash = GroomAsset->GetAssetHash();
 		HairGroupInstance->Debug.GroomAssetName = GroomAsset->GetName();
 		HairGroupInstance->Debug.MeshComponentForDebug = RegisteredMeshComponent;
 		HairGroupInstance->Debug.GroomComponentForDebug = this;
 		HairGroupInstance->Debug.GroomBindingType = BindingAsset ? BindingAsset->GetGroomBindingType() : EGroomBindingMeshType::SkeletalMesh;
 		HairGroupInstance->Debug.GroomCacheType = GetEffectiveGroomCacheType(GroomCache, GroomAsset);
+		HairGroupInstance->Debug.GroomBindingAssetHash = BindingAsset ? BindingAsset->GetAssetHash() : 0u;
 		HairGroupInstance->Debug.GroomCacheBuffers = GroomCacheBuffers;
 		HairGroupInstance->Debug.LODForcedIndex = LODForcedIndex;
 		HairGroupInstance->Debug.LODPredictedIndex = LODPredictedIndex;
diff --git a/Engine/Plugins/Runtime/HairStrands/Source/HairStrandsCore/Private/GroomManager.cpp b/Engine/Plugins/Runtime/HairStrands/Source/HairStrandsCore/Private/GroomManager.cpp
--- a/Engine/Plugins/Runtime/HairStrands/Source/HairStrandsCore/Private/GroomManager.cpp
+++ b/Engine/Plugins/Runtime/HairStrands/Source/HairStrandsCore/Private/GroomManager.cpp
@@ -2131,6 +2131,14 @@ static EHairGeometryType ConvertLODGeometryType(EHairGeometryType Type, bool Inb
 	return EHairGeometryType::NoneGeometry;
 }
 
+static EHairGeometryType GetLODGeometryType(FHairGroupInstance* InInstance, int32 InLODIndex, EShaderPlatform InShaderPlatform=GMaxRHIShaderPlatform)
+{
+	const TArray<EHairGeometryType>& LODGeometryTypes = InInstance->HairGroupPublicData->GetLODGeometryTypes();
+	check(LODGeometryTypes.IsValidIndex(InLODIndex));
+	const bool bForceCards = GHairStrands_UseCards > 0 || InInstance->bForceCards;
+	return ConvertLODGeometryType(LODGeometryTypes[InLODIndex], bForceCards, InShaderPlatform);
+}
+
 static void RunHairBufferSwap(const FHairStrandsInstances& Instances, const TArray<const FSceneView*> Views)
 {
 	EShaderPlatform ShaderPlatform = EShaderPlatform::SP_NumPlatforms;
@@ -2217,7 +2225,6 @@ void AddHairStreamingRequest(FHairGroupInstance* Instance, int32 InLODIndex)
 		const float MinLOD = FMath::Max(0, GHairStrandsMinLOD);
 		float LODIndex = FMath::Max(InLODIndex, MinLOD);
 
-		const TArray<EHairGeometryType>& LODGeometryTypes = Instance->HairGroupPublicData->GetLODGeometryTypes();
 		const TArray<bool>& LODVisibilities = Instance->HairGroupPublicData->GetLODVisibilities();
 		const int32 LODCount = LODVisibilities.Num();
 
@@ -2225,9 +2232,8 @@ void AddHairStreamingRequest(FHairGroupInstance* Instance, int32 InLODIndex)
 
 		const int32 IntLODIndex = FMath::Clamp(FMath::FloorToInt(LODIndex), 0, LODCount - 1);
 		const bool bIsVisible = LODVisibilities[IntLODIndex];
-		const bool bForceCards = GHairStrands_UseCards > 0 || Instance->bForceCards; // todo
-		EHairGeometryType GeometryType = ConvertLODGeometryType(LODGeometryTypes[IntLODIndex], bForceCards, GMaxRHIShaderPlatform);
 
+		EHairGeometryType GeometryType = GetLODGeometryType(Instance, IntLODIndex);
 		if (GeometryType == EHairGeometryType::Meshes)
 		{
 			if (!Instance->Meshes.IsValid(IntLODIndex))
@@ -2488,11 +2494,10 @@ static bool SelectValidLOD(
 	bool bHasPathTracingView)
 {
 	const int32 IntHairLODIndex = FMath::Clamp(FMath::FloorToInt(HairLODIndex), 0, HairLODCount - 1);
-	const TArray<EHairGeometryType>& LODGeometryTypes = Instance->HairGroupPublicData->GetLODGeometryTypes();
 	const TArray<bool>& LODVisibilities = Instance->HairGroupPublicData->GetLODVisibilities();
 	const bool bIsVisible = LODVisibilities[IntHairLODIndex];
-	const bool bForceCards = GHairStrands_UseCards > 0 || Instance->bForceCards; // todo
-	EHairGeometryType GeometryType = ConvertLODGeometryType(LODGeometryTypes[IntHairLODIndex], bForceCards, ShaderPlatform);
+
+	EHairGeometryType GeometryType = GetLODGeometryType(Instance, IntHairLODIndex, ShaderPlatform);
 
 	// Skip cluster culling if strands have a single LOD (i.e.: LOD0), and the instance is static (i.e., no skinning, no simulation, ...)
 	bool bCullingEnable = false;
@@ -2604,17 +2609,17 @@ static bool SelectValidLOD(
 	{
 		check(Instance->HairGroupPublicData);
 
-		if (Instance->Strands.RestRootResource)			{ Instance->Strands.RestRootResource->Allocate(GraphBuilder, EHairResourceLoadingType::Async, ResourceStatus, RequestedCurveCount, RequestedPointCount, MeshLODIndex); }
-		if (Instance->Strands.RestResource)				{ Instance->Strands.RestResource->Allocate(GraphBuilder, EHairResourceLoadingType::Async, ResourceStatus, RequestedCurveCount, RequestedPointCount); }
-		if (Instance->Strands.ClusterResource)			{ Instance->Strands.ClusterResource->Allocate(GraphBuilder, EHairResourceLoadingType::Async, ResourceStatus, RequestedCurveCount, RequestedPointCount); }
-		if (Instance->Strands.InterpolationResource)	{ Instance->Strands.InterpolationResource->Allocate(GraphBuilder, EHairResourceLoadingType::Async, ResourceStatus, RequestedCurveCount, RequestedPointCount); }
-		if (Instance->Strands.CullingResource)			{ Instance->Strands.CullingResource->Allocate(GraphBuilder, EHairResourceLoadingType::Async, ResourceStatus, RequestedCurveCount, RequestedPointCount); }
+		if (Instance->Strands.RestRootResource)			{ Instance->Strands.RestRootResource->Allocate		(GraphBuilder, EHairResourceLoadingType::Async, ResourceStatus, RequestedCurveCount, RequestedPointCount, MeshLODIndex,       false/*bAllowDeallocation*/); }
+		if (Instance->Strands.RestResource)				{ Instance->Strands.RestResource->Allocate			(GraphBuilder, EHairResourceLoadingType::Async, ResourceStatus, RequestedCurveCount, RequestedPointCount, -1/*MeshLODIndex*/, false/*bAllowDeallocation*/); }
+		if (Instance->Strands.ClusterResource)			{ Instance->Strands.ClusterResource->Allocate		(GraphBuilder, EHairResourceLoadingType::Async, ResourceStatus, RequestedCurveCount, RequestedPointCount, -1/*MeshLODIndex*/, false/*bAllowDeallocation*/); }
+		if (Instance->Strands.InterpolationResource)	{ Instance->Strands.InterpolationResource->Allocate	(GraphBuilder, EHairResourceLoadingType::Async, ResourceStatus, RequestedCurveCount, RequestedPointCount, -1/*MeshLODIndex*/, false/*bAllowDeallocation*/); }
+		if (Instance->Strands.CullingResource)			{ Instance->Strands.CullingResource->Allocate		(GraphBuilder, EHairResourceLoadingType::Async, ResourceStatus, RequestedCurveCount, RequestedPointCount); }
 
-		if (Instance->Strands.DeformedRootResource)		{ Instance->Strands.DeformedRootResource->Allocate(GraphBuilder, EHairResourceLoadingType::Async, ResourceStatus, RequestedCurveCount, RequestedPointCount, MeshLODIndex); }
-		if (Instance->Strands.DeformedResource)			{ Instance->Strands.DeformedResource->Allocate(GraphBuilder, EHairResourceLoadingType::Async, ResourceStatus, RequestedCurveCount, RequestedPointCount); }
+		if (Instance->Strands.DeformedRootResource)		{ Instance->Strands.DeformedRootResource->Allocate	(GraphBuilder, EHairResourceLoadingType::Async, ResourceStatus, RequestedCurveCount, RequestedPointCount, MeshLODIndex); }
+		if (Instance->Strands.DeformedResource)			{ Instance->Strands.DeformedResource->Allocate		(GraphBuilder, EHairResourceLoadingType::Async, ResourceStatus, RequestedCurveCount, RequestedPointCount); }
 #if RHI_RAYTRACING
 		if (bHasPathTracingView)						{ AllocateRaytracingResources(Instance); }
-		if (Instance->Strands.RenRaytracingResource)	{ Instance->Strands.RenRaytracingResource->Allocate(GraphBuilder, EHairResourceLoadingType::Async, ResourceStatus, RequestedCurveCount, RequestedPointCount); }
+		if (Instance->Strands.RenRaytracingResource)	{ Instance->Strands.RenRaytracingResource->Allocate	(GraphBuilder, EHairResourceLoadingType::Async, ResourceStatus, RequestedCurveCount, RequestedPointCount, -1/*MeshLODIndex*/, Instance->Strands.RenRaytracingResourceOwned/*bAllowDeallocation*/); }
 #endif
 		Instance->Strands.VertexFactory->InitResources(GraphBuilder.RHICmdList);
 
@@ -2732,6 +2737,41 @@ static void RunHairLODSelection(
 
 	};
 
+	struct FUniqueGroom
+	{
+		uint32 InstanceCount = 0;
+		FHairLOD HairLOD;
+		uint32 AssetHash = 0;
+		uint32 GroupIndex = 0;
+		int32 MeshLODIndex = -1;
+		uint32 RestCurveCount = 0;
+
+		// Resources
+		FHairStrandsRestResource*			RestResource = nullptr;
+		FHairStrandsInterpolationResource*	InterpolationResource = nullptr;
+		FHairStrandsClusterResource* 		ClusterResource = nullptr;
+	#if RHI_RAYTRACING
+		FHairStrandsRaytracingResource*		RaytracingResource = nullptr;
+	#endif
+	};
+	TArray<FUniqueGroom> UniqueGroomAssets;
+	UniqueGroomAssets.Reserve(Instances.Num());
+
+	struct FUniqueGroomBinding
+	{
+		uint32 InstanceCount = 0;
+		FHairLOD HairLOD;
+		uint32 AssetHash = 0;
+		uint32 GroupIndex = 0;
+		int32 MeshLODIndex = -1;
+		uint32 RestCurveCount = 0;
+
+		// Resources
+		FHairStrandsRestRootResource* RestRootResource = nullptr;
+	};
+	TArray<FUniqueGroomBinding> UniqueGroomBindingAssets;
+	UniqueGroomAssets.Reserve(Instances.Num());
+
 	// Find unique groom assets
 	TArray<FInstanceData> InstanceDatas;
 	InstanceDatas.Reserve(Instances.Num());
@@ -2748,6 +2788,100 @@ static void RunHairLODSelection(
 		InstanceData.PrevMeshLODIndex 		= Instance->HairGroupPublicData->GetMeshLODIndex();
 		InstanceData.HairLOD 				= ComputeHairLODIndex(Instance, Views);
 		InstanceData.MeshLODIndex 			= GetMeshLODIndex(GraphBuilder, ShaderMap, Scene, Instance, InstanceData.MeshLODLocalToWorld);
+
+		const EHairGeometryType GeometryType = GetLODGeometryType(Instance, InstanceData.HairLOD.HairLODIndex, ShaderPlatform);
+		const EHairBindingType BindingType = Instance->HairGroupPublicData->GetBindingType(InstanceData.HairLOD.HairLODIndex);
+
+		// Extract unique Groom/GroomBinding assets
+		// Only for strands geometry as, we only support fine-grain stream on strands resources
+		if (GeometryType == EHairGeometryType::Strands)
+		{
+			// Groom
+			{
+				FUniqueGroom* Unique = UniqueGroomAssets.FindByPredicate([Instance](const FUniqueGroom& A) { return A.AssetHash == Instance->Debug.GroomAssetHash && A.GroupIndex == Instance->Debug.GroupIndex; });
+				if (Unique)
+				{
+					if (InstanceData.HairLOD.ContinuousLODCurveCount > Unique->HairLOD.ContinuousLODCurveCount)
+					{
+						Unique->HairLOD = InstanceData.HairLOD;
+					}
+				}
+				else
+				{
+					Unique 					= &UniqueGroomAssets.AddDefaulted_GetRef();
+					Unique->AssetHash 		= Instance->Debug.GroomAssetHash;
+					Unique->GroupIndex 		= Instance->Debug.GroupIndex;
+					Unique->HairLOD 		= InstanceData.HairLOD;
+					Unique->RestCurveCount 	= Instance->HairGroupPublicData->RestCurveCount;
+				}
+				++Unique->InstanceCount;
+
+				if (!Unique->RestResource) 			{ Unique->RestResource 			= Instance->Strands.RestResource; }
+				if (!Unique->InterpolationResource) { Unique->InterpolationResource = Instance->Strands.InterpolationResource; }
+				if (!Unique->ClusterResource) 		{ Unique->ClusterResource 		= Instance->Strands.ClusterResource; }
+			#if RHI_RAYTRACING
+				if (!Unique->RaytracingResource && !Instance->Strands.RenRaytracingResourceOwned)
+				{
+					Unique->RaytracingResource 	= Instance->Strands.RenRaytracingResource;
+				}
+			#endif
+			}
+	
+			// Groom binding
+			if (BindingType == EHairBindingType::Skinning)
+			{
+				FUniqueGroomBinding* Unique = UniqueGroomBindingAssets.FindByPredicate([Instance](const FUniqueGroomBinding& A) { return A.AssetHash == Instance->Debug.GroomBindingAssetHash && A.GroupIndex == Instance->Debug.GroupIndex; });
+				if (Unique)
+				{
+					if (InstanceData.HairLOD.ContinuousLODCurveCount > Unique->HairLOD.ContinuousLODCurveCount)
+					{
+						Unique->HairLOD 	= InstanceData.HairLOD;
+						Unique->MeshLODIndex= InstanceData.MeshLODIndex;
+					}
+				}
+				else
+				{
+					Unique 					= &UniqueGroomBindingAssets.AddDefaulted_GetRef();
+					Unique->AssetHash 		= Instance->Debug.GroomBindingAssetHash;
+					Unique->GroupIndex 		= Instance->Debug.GroupIndex;
+					Unique->HairLOD 		= InstanceData.HairLOD;
+					Unique->RestCurveCount 	= Instance->HairGroupPublicData->RestCurveCount;
+					Unique->MeshLODIndex 	= InstanceData.MeshLODIndex;
+				}
+				++Unique->InstanceCount;
+
+				if (!Unique->RestRootResource) { Unique->RestRootResource = Instance->Strands.RestRootResource; }
+			}
+		}
+	}
+
+	// Emit streaming request for each unique Groom/GroomBinding group
+	for (FUniqueGroom& Unique : UniqueGroomAssets)
+	{
+		const uint32 RequestedCurveCount = Unique.HairLOD.ContinuousLODCurveCount;
+		const uint32 RequestedPointCount = Unique.HairLOD.ContinuousLODPointCount;
+
+		EHairResourceStatus ResourceStatus;
+		ResourceStatus.Status 				= EHairResourceStatus::EStatus::None;
+		ResourceStatus.AvailableCurveCount 	= Unique.RestCurveCount;
+
+		if (Unique.RestResource)			{ Unique.RestResource->Allocate(GraphBuilder, EHairResourceLoadingType::Async, ResourceStatus, RequestedCurveCount, RequestedPointCount); }
+		if (Unique.ClusterResource)			{ Unique.ClusterResource->Allocate(GraphBuilder, EHairResourceLoadingType::Async, ResourceStatus, RequestedCurveCount, RequestedPointCount); }
+		if (Unique.InterpolationResource)	{ Unique.InterpolationResource->Allocate(GraphBuilder, EHairResourceLoadingType::Async, ResourceStatus, RequestedCurveCount, RequestedPointCount); }
+		#if RHI_RAYTRACING
+		if (Unique.RaytracingResource)		{ Unique.RaytracingResource->Allocate(GraphBuilder, EHairResourceLoadingType::Async, ResourceStatus, RequestedCurveCount, RequestedPointCount); }
+		#endif
+	}
+	for (FUniqueGroomBinding& Unique : UniqueGroomBindingAssets)
+	{
+		const uint32 RequestedCurveCount = Unique.HairLOD.ContinuousLODCurveCount;
+		const uint32 RequestedPointCount = Unique.HairLOD.ContinuousLODPointCount;
+
+		EHairResourceStatus ResourceStatus;
+		ResourceStatus.Status 				= EHairResourceStatus::EStatus::None;
+		ResourceStatus.AvailableCurveCount 	= Unique.RestCurveCount;
+
+		if (Unique.RestRootResource)		{ Unique.RestRootResource->Allocate(GraphBuilder, EHairResourceLoadingType::Async, ResourceStatus, RequestedCurveCount, RequestedPointCount, Unique.MeshLODIndex); }
 	}
 
 	// LOD selection & load resources
diff --git a/Engine/Plugins/Runtime/HairStrands/Source/HairStrandsCore/Private/GroomResources.cpp b/Engine/Plugins/Runtime/HairStrands/Source/HairStrandsCore/Private/GroomResources.cpp
--- a/Engine/Plugins/Runtime/HairStrands/Source/HairStrandsCore/Private/GroomResources.cpp
+++ b/Engine/Plugins/Runtime/HairStrands/Source/HairStrandsCore/Private/GroomResources.cpp
@@ -723,7 +723,7 @@ void FHairCommonResource::Allocate(FRDGBuilder& GraphBuilder, EHairResourceLoadi
 }
 
 bool NeedDeallocation(uint32 InRequest, uint32 InAvailable);
-void FHairCommonResource::Allocate(FRDGBuilder& GraphBuilder, EHairResourceLoadingType LoadingType, EHairResourceStatus& Status, uint32 InRequestedCurveCount, uint32 InRequestedPointCount, int32 InLODIndex)
+void FHairCommonResource::Allocate(FRDGBuilder& GraphBuilder, EHairResourceLoadingType LoadingType, EHairResourceStatus& Status, uint32 InRequestedCurveCount, uint32 InRequestedPointCount, int32 InLODIndex, bool bAllowDeallocation)
 {
 	check(AllocationType == EHairStrandsAllocationType::Deferred);
 
@@ -752,7 +752,7 @@ void FHairCommonResource::Allocate(FRDGBuilder& GraphBuilder, EHairResourceLoadi
 		if (bIsInitialized && MaxAvailableCurveCount >= InRequestedCurveCount && InternalIsLODDataLoaded(InRequestedCurveCount, InRequestedPointCount, InLODIndex)) 
 		{ 
 			// Trim/Un-stream data if needed (ensure no streaming request is in-flight)
-			if (NeedDeallocation(InRequestedCurveCount, MaxAvailableCurveCount) && StreamingRequest.IsNone())
+			if (bAllowDeallocation && NeedDeallocation(InRequestedCurveCount, MaxAvailableCurveCount) && StreamingRequest.IsNone())
 			{
 				StreamingRequest.CurveCount = InRequestedCurveCount;
 				StreamingRequest.PointCount = InRequestedPointCount;
diff --git a/Engine/Plugins/Runtime/HairStrands/Source/HairStrandsCore/Public/GroomAsset.h b/Engine/Plugins/Runtime/HairStrands/Source/HairStrandsCore/Public/GroomAsset.h
--- a/Engine/Plugins/Runtime/HairStrands/Source/HairStrandsCore/Public/GroomAsset.h
+++ b/Engine/Plugins/Runtime/HairStrands/Source/HairStrandsCore/Public/GroomAsset.h
@@ -717,6 +717,7 @@ public:
 
 	/** Helper function to return the asset path name, optionally joined with the LOD index if LODIndex > -1. */
 	FName GetAssetPathName(int32 LODIndex = -1);
+	uint32 GetAssetHash() const { return AssetNameHash; }
 
 //private :
 #if WITH_EDITOR
@@ -826,6 +827,7 @@ private:
 	bool bRetryLoadFromGameThread = false;
 #endif // WITH_EDITORONLY_DATA
 	bool bIsInitialized = false;
+	uint32 AssetNameHash = 0;
 
 #if WITH_EDITOR
 public:
diff --git a/Engine/Plugins/Runtime/HairStrands/Source/HairStrandsCore/Public/GroomBindingAsset.h b/Engine/Plugins/Runtime/HairStrands/Source/HairStrandsCore/Public/GroomBindingAsset.h
--- a/Engine/Plugins/Runtime/HairStrands/Source/HairStrandsCore/Public/GroomBindingAsset.h
+++ b/Engine/Plugins/Runtime/HairStrands/Source/HairStrandsCore/Public/GroomBindingAsset.h
@@ -227,6 +227,7 @@ public:
 
 	/** Helper function to return the asset path name, optionally joined with the LOD index if LODIndex > -1. */
 	FName GetAssetPathName(int32 LODIndex = -1);
+	uint32 GetAssetHash() const { return AssetNameHash; }
 
 #if WITH_EDITOR
 	FOnGroomBindingAssetChanged& GetOnGroomBindingAssetChanged() { return OnGroomBindingAssetChanged; }
@@ -282,6 +283,7 @@ private:
 	TArray<FString> CachedDerivedDataKey;
 #endif
 	bool bIsValid = false;
+	uint32 AssetNameHash = 0;
 };
 
 UCLASS(BlueprintType, hidecategories = (Object))
diff --git a/Engine/Plugins/Runtime/HairStrands/Source/HairStrandsCore/Public/GroomInstance.h b/Engine/Plugins/Runtime/HairStrands/Source/HairStrandsCore/Public/GroomInstance.h
--- a/Engine/Plugins/Runtime/HairStrands/Source/HairStrandsCore/Public/GroomInstance.h
+++ b/Engine/Plugins/Runtime/HairStrands/Source/HairStrandsCore/Public/GroomInstance.h
@@ -244,6 +244,8 @@ struct HAIRSTRANDSCORE_API FHairGroupInstance : public FHairStrandsInstance
 		uint32					GroupCount = 0;
 		FString					GroomAssetName;
 		uint32					LastFrameIndex = ~0;
+		uint32					GroomAssetHash = 0;
+		uint32					GroomBindingAssetHash = 0;
 
 		float					LODPredictedIndex = -1.f;	// Computed on the rendering-thread, readback on the game-thread during tick()
 		float					LODForcedIndex = -1.f;		// Set by the game-thread, read on the rendering-thread to force a particular LOD
diff --git a/Engine/Plugins/Runtime/HairStrands/Source/HairStrandsCore/Public/GroomResources.h b/Engine/Plugins/Runtime/HairStrands/Source/HairStrandsCore/Public/GroomResources.h
--- a/Engine/Plugins/Runtime/HairStrands/Source/HairStrandsCore/Public/GroomResources.h
+++ b/Engine/Plugins/Runtime/HairStrands/Source/HairStrandsCore/Public/GroomResources.h
@@ -103,7 +103,7 @@ struct FHairCommonResource : public FRenderResource
 	/* Init/Release buffers (FHairCommonResource) */
 	void Allocate(FRDGBuilder& GraphBuilder, EHairResourceLoadingType LoadingType);
 	void Allocate(FRDGBuilder& GraphBuilder, EHairResourceLoadingType LoadingType, EHairResourceStatus& Status, int32 InLODIndex=-1);
-	void Allocate(FRDGBuilder& GraphBuilder, EHairResourceLoadingType LoadingType, EHairResourceStatus& Status, uint32 InRequestedCurveCount, uint32 InRequestedPointCount, int32 InLODIndex=-1);
+	void Allocate(FRDGBuilder& GraphBuilder, EHairResourceLoadingType LoadingType, EHairResourceStatus& Status, uint32 InRequestedCurveCount, uint32 InRequestedPointCount, int32 InLODIndex=-1, bool bAllowDeallocation=true);
 
 	void StreamInData(int32 InLODIndex=-1);
 
