Fix CitySample crash due to hair resources.

diff --git a/Engine/Plugins/Runtime/HairStrands/Source/HairStrandsCore/Private/GroomAsset.cpp b/Engine/Plugins/Runtime/HairStrands/Source/HairStrandsCore/Private/GroomAsset.cpp
--- a/Engine/Plugins/Runtime/HairStrands/Source/HairStrandsCore/Private/GroomAsset.cpp
+++ b/Engine/Plugins/Runtime/HairStrands/Source/HairStrandsCore/Private/GroomAsset.cpp
@@ -1377,8 +1377,6 @@ void UGroomAsset::PostLoad()
 			bSucceed = UGroomAsset::CacheDerivedDatas();
 		}
 	}
-#else
-	bool bSucceed = false;
 #endif
 
 	check(GetNumHairGroups() > 0);
@@ -1411,6 +1409,10 @@ void UGroomAsset::PostLoad()
 #if !WITH_EDITOR
 	if (!IsTemplate() && IsHairStrandsAssetLoadingEnable())
 	{
+		// Resize resources based on loaded data
+		const uint32 GroupCount = GetHairGroupsPlatformData().Num();
+		GetHairGroupsResources().Init(FHairGroupResources(), GroupCount);
+
 		InitResources();
 	}
 #endif
