Fix incorrect groom holdout initialization, causing flickering when toggling holdout inbetween frame.

diff --git a/Engine/Plugins/Runtime/HairStrands/Source/HairStrandsCore/Private/GroomComponent.cpp b/Engine/Plugins/Runtime/HairStrands/Source/HairStrandsCore/Private/GroomComponent.cpp
--- a/Engine/Plugins/Runtime/HairStrands/Source/HairStrandsCore/Private/GroomComponent.cpp
+++ b/Engine/Plugins/Runtime/HairStrands/Source/HairStrandsCore/Private/GroomComponent.cpp
@@ -456,7 +456,6 @@ public:
 			HairInstance->bForceCards = Component->bUseCards;
 			HairInstance->bUpdatePositionOffset = Component->RegisteredMeshComponent != nullptr;
 			HairInstance->bCastShadow = Component->CastShadow;
-			HairInstance->bHoldout = bHoldout;
 
 			if (HairInstance->Strands.IsValid() && HairInstance->Strands.VertexFactory == nullptr)
 			{
diff --git a/Engine/Plugins/Runtime/HairStrands/Source/HairStrandsCore/Private/HairStrandsInterpolation.cpp b/Engine/Plugins/Runtime/HairStrands/Source/HairStrandsCore/Private/HairStrandsInterpolation.cpp
--- a/Engine/Plugins/Runtime/HairStrands/Source/HairStrandsCore/Private/HairStrandsInterpolation.cpp
+++ b/Engine/Plugins/Runtime/HairStrands/Source/HairStrandsCore/Private/HairStrandsInterpolation.cpp
@@ -1712,7 +1712,7 @@ FHairGroupPublicData::FVertexFactoryInput InternalComputeHairStrandsVertexInputD
 	OutVFInput.Strands.Common.Flags |= Instance->Strands.Modifier.bScatterSceneLighting ? HAIR_FLAGS_SCATTER_SCENE_LIGHT : 0u;
 	OutVFInput.Strands.Common.Flags |= Instance->Strands.Modifier.bUseStableRasterization ? HAIR_FLAGS_STABLE_RASTER : 0u;
 	OutVFInput.Strands.Common.Flags |= bRaytracingGeometry ? HAIR_FLAGS_RAYTRACING_GEOMETRY : 0u;
-	OutVFInput.Strands.Common.Flags |= Instance->bHoldout ? HAIR_FLAGS_HOLDOUT : 0u;
+	OutVFInput.Strands.Common.Flags |= Instance->Debug.Proxy->Holdout() ? HAIR_FLAGS_HOLDOUT : 0u;
 	OutVFInput.Strands.Common.RaytracingProceduralSplits = GetHairRaytracingProceduralSplits();
 	OutVFInput.Strands.Common.GroupIndex = Instance->Debug.GroupIndex;
 	OutVFInput.Strands.Common.GroupCount = Instance->Debug.GroupCount;
diff --git a/Engine/Plugins/Runtime/HairStrands/Source/HairStrandsCore/Public/GroomInstance.h b/Engine/Plugins/Runtime/HairStrands/Source/HairStrandsCore/Public/GroomInstance.h
--- a/Engine/Plugins/Runtime/HairStrands/Source/HairStrandsCore/Public/GroomInstance.h
+++ b/Engine/Plugins/Runtime/HairStrands/Source/HairStrandsCore/Public/GroomInstance.h
@@ -279,7 +279,6 @@ struct HAIRSTRANDSCORE_API FHairGroupInstance : public FHairStrandsInstance
 	bool					bUpdatePositionOffset = false;
 	bool					bCastShadow = true;
 	bool 					bSupportStreaming = true;
-	bool					bHoldout = false;
 	
 	// Deformed component to extract the bone buffer 
 	UMeshComponent* DeformedComponent = nullptr;
