Motion vector fix for MRQ graph.

diff --git a/Engine/Plugins/Runtime/HairStrands/Source/HairStrandsCore/Private/GroomComponent.cpp b/Engine/Plugins/Runtime/HairStrands/Source/HairStrandsCore/Private/GroomComponent.cpp
--- a/Engine/Plugins/Runtime/HairStrands/Source/HairStrandsCore/Private/GroomComponent.cpp
+++ b/Engine/Plugins/Runtime/HairStrands/Source/HairStrandsCore/Private/GroomComponent.cpp
@@ -3576,12 +3576,21 @@ void UGroomComponent::SendRenderDynamicData_Concurrent()
 
 	if(GetHairSwapBufferType() == EHairBufferSwapType::RenderFrame)
 	{
+		FTransform SkelLocalToTransform = FTransform::Identity;
+		if (RegisteredMeshComponent)
+		{
+			SkelLocalToTransform = RegisteredMeshComponent->GetComponentTransform();
+		}
+
 		TArray<TRefCountPtr<FHairGroupInstance>> LocalHairGroupInstances = HairGroupInstances;
 		ENQUEUE_RENDER_COMMAND(FHairStrandsTick_TransformUpdate)(
-			[LocalHairGroupInstances](FRHICommandListImmediate& RHICmdList)
+			[LocalHairGroupInstances, SkelLocalToTransform](FRHICommandListImmediate& RHICmdList)
 		{
 			for (const TRefCountPtr<FHairGroupInstance>& Instance : LocalHairGroupInstances)
 			{
+				Instance->Debug.SkinningPreviousLocalToWorld = Instance->Debug.SkinningCurrentLocalToWorld;
+				Instance->Debug.SkinningCurrentLocalToWorld  = SkelLocalToTransform;
+
 				if (Instance->Guides.DeformedResource)  { Instance->Guides.DeformedResource->SwapBuffer(); }
 				if (Instance->Strands.DeformedResource) { Instance->Strands.DeformedResource->SwapBuffer(); }
 			}
diff --git a/Engine/Plugins/Runtime/HairStrands/Source/HairStrandsCore/Private/GroomManager.cpp b/Engine/Plugins/Runtime/HairStrands/Source/HairStrandsCore/Private/GroomManager.cpp
--- a/Engine/Plugins/Runtime/HairStrands/Source/HairStrandsCore/Private/GroomManager.cpp
+++ b/Engine/Plugins/Runtime/HairStrands/Source/HairStrandsCore/Private/GroomManager.cpp
@@ -2942,7 +2942,7 @@ static void RunHairLODSelection(
 		}
 
 		// Update the local-to-world transform based on the binding type 
-		if (GetHairSwapBufferType() != EHairBufferSwapType::Tick)
+		if (GetHairSwapBufferType() == EHairBufferSwapType::BeginOfFrame || GetHairSwapBufferType() == EHairBufferSwapType::EndOfFrame)
 		{
 			InstanceData.Instance->Debug.SkinningPreviousLocalToWorld = InstanceData.Instance->Debug.SkinningCurrentLocalToWorld;
 			InstanceData.Instance->Debug.SkinningCurrentLocalToWorld  = InstanceData.MeshLODLocalToWorld;
diff --git a/Engine/Plugins/Runtime/HairStrands/Source/HairStrandsCore/Private/HairStrandsInterpolation.cpp b/Engine/Plugins/Runtime/HairStrands/Source/HairStrandsCore/Private/HairStrandsInterpolation.cpp
--- a/Engine/Plugins/Runtime/HairStrands/Source/HairStrandsCore/Private/HairStrandsInterpolation.cpp
+++ b/Engine/Plugins/Runtime/HairStrands/Source/HairStrandsCore/Private/HairStrandsInterpolation.cpp
@@ -51,6 +51,9 @@ static FAutoConsoleVariableRef CVarHairStrandsTransferPositionOnLODChange(TEXT("
 static int32 GHairStrands_Raytracing_ForceRebuildBVH = 1;
 static FAutoConsoleVariableRef CVarHairRTGeomForceRebuild(TEXT("r.HairStrands.Strands.Raytracing.ForceRebuildBVH"), GHairStrands_Raytracing_ForceRebuildBVH, TEXT("Force BVH rebuild instead of doing a BVH refit when hair positions changed"));
 
+static int32 GHairStrands_MotionVector_CheckViewID = 1;
+static FAutoConsoleVariableRef CVarHairStrands_MotionVector_CheckViewID(TEXT("r.HairStrands.Strands.MotionVectorCheckViewID"), GHairStrands_MotionVector_CheckViewID, TEXT("Issue motion vector on hair strands only if updates happens on the same view"));
+
 ////////////////////////////////////////////////////////////////////////////////////////////////////////////////////////
 
 uint32 GetHairCardsInterpolationType()
@@ -1633,9 +1636,10 @@ FHairGroupPublicData::FVertexFactoryInput InternalComputeHairStrandsVertexInputD
 	// 2. Render Strands with deformation
 	else if (bSupportDeformation)
 	{
-		const bool bHasValidMotionVector =
+		const bool bSameUniqueView = 
 			Instance->Strands.DeformedResource->GetUniqueViewID(FHairStrandsDeformedResource::Current) ==
 			Instance->Strands.DeformedResource->GetUniqueViewID(FHairStrandsDeformedResource::Previous);
+		const bool bHasValidMotionVector = GHairStrands_MotionVector_CheckViewID ? bSameUniqueView : true;
 		if (bHasValidMotionVector)
 		{
 			CONVERT_HAIRSSTRANDS_VF_PARAMETERS(OutVFInput.Strands.PositionBuffer,			Instance->Strands.DeformedResource->GetBuffer(FHairStrandsDeformedResource::EFrameType::Current));
