Fix incorrect groom shadow casting not respecting hidden in game / hidden shadow casting flags on instance.

diff --git a/Engine/Plugins/Runtime/HairStrands/Source/HairStrandsCore/Private/GroomComponent.cpp b/Engine/Plugins/Runtime/HairStrands/Source/HairStrandsCore/Private/GroomComponent.cpp
--- a/Engine/Plugins/Runtime/HairStrands/Source/HairStrandsCore/Private/GroomComponent.cpp
+++ b/Engine/Plugins/Runtime/HairStrands/Source/HairStrandsCore/Private/GroomComponent.cpp
@@ -455,7 +455,6 @@ public:
 			check(HairInstance->HairGroupPublicData);
 			HairInstance->bForceCards = Component->bUseCards;
 			HairInstance->bUpdatePositionOffset = Component->RegisteredMeshComponent != nullptr;
-			HairInstance->bCastShadow = Component->CastShadow;
 
 			if (HairInstance->Strands.IsValid() && HairInstance->Strands.VertexFactory == nullptr)
 			{
diff --git a/Engine/Plugins/Runtime/HairStrands/Source/HairStrandsCore/Private/GroomManager.cpp b/Engine/Plugins/Runtime/HairStrands/Source/HairStrandsCore/Private/GroomManager.cpp
--- a/Engine/Plugins/Runtime/HairStrands/Source/HairStrandsCore/Private/GroomManager.cpp
+++ b/Engine/Plugins/Runtime/HairStrands/Source/HairStrandsCore/Private/GroomManager.cpp
@@ -1080,7 +1080,7 @@ static void RunHairStrandsInterpolation_Strands(
 		bool bNeedRaytracing = false;
 		bool bNeedRaytracingUpdate = false;
 		bool bNeedRaytracingBuild = false;
-
+		bool bSupportVoxelization = false;
 
 		FInstanceRDGResources RDGResources;
 
@@ -1113,6 +1113,8 @@ static void RunHairStrandsInterpolation_Strands(
 		}	
 		check(Instance->HairGroupPublicData);
 
+		const bool bCastShadow = (Instance->Debug.Proxy->IsDrawnInGame() && Instance->Debug.Proxy->CastsDynamicShadow()) || Instance->Debug.Proxy->CastsHiddenShadow();
+
 		FInstanceData& InstanceData = InstanceDatas.AddDefaulted_GetRef();
 		InstanceData.RegisteredIndex			= Instance->RegisteredIndex;
 		InstanceData.Instance 					= Instance;
@@ -1129,6 +1131,7 @@ static void RunHairStrandsInterpolation_Strands(
 		InstanceData.bNeedRaytracing			= false;
 		InstanceData.bNeedRaytracingUpdate		= false;
 		InstanceData.bNeedRaytracingBuild		= false;
+		InstanceData.bSupportVoxelization		= Instance->Strands.Modifier.bSupportVoxelization && bCastShadow;
 
 		// Register resources
 		FInstanceRDGResources RDGResources;
@@ -1548,7 +1551,7 @@ static void RunHairStrandsInterpolation_Strands(
 		for (FInstanceData& InstanceData : InstanceDatas)
 		{
 			// Optim: If an instance does not voxelize it's data, then there is no need for having valid AABB
-			bool bNeedGPUAABB = InstanceData.Instance->Strands.Modifier.bSupportVoxelization && InstanceData.Instance->bCastShadow;
+			bool bNeedGPUAABB = InstanceData.bSupportVoxelization;
 			const EHairAABBUpdateType UpdateType = InstanceData.bNeedDeformation ? EHairAABBUpdateType::UpdateClusterAABB : EHairAABBUpdateType::UpdateGroupAABB;
 				
 			FHairStrandClusterData::FHairGroup* HairGroupCluster = nullptr;
@@ -1688,7 +1691,7 @@ static void RunHairStrandsInterpolation_Strands(
 		InstanceData.Instance->Strands.UniformBuffer.UpdateUniformBufferImmediate(GraphBuilder.RHICmdList, InstanceData.Instance->GetHairStandsUniformShaderParameters(ViewMode));
 		InstanceData.Instance->HairGroupPublicData->VFInput.GeometryType = EHairGeometryType::Strands;
 		InstanceData.Instance->HairGroupPublicData->VFInput.LocalToWorldTransform = InstanceData.Instance->GetCurrentLocalToWorld();
-		InstanceData.Instance->HairGroupPublicData->bSupportVoxelization = InstanceData.Instance->Strands.Modifier.bSupportVoxelization && InstanceData.Instance->bCastShadow;
+		InstanceData.Instance->HairGroupPublicData->bSupportVoxelization = InstanceData.bSupportVoxelization;
 	}
 
 	ExternalAccessQueue.Submit(GraphBuilder);
diff --git a/Engine/Plugins/Runtime/HairStrands/Source/HairStrandsCore/Public/GroomInstance.h b/Engine/Plugins/Runtime/HairStrands/Source/HairStrandsCore/Public/GroomInstance.h
--- a/Engine/Plugins/Runtime/HairStrands/Source/HairStrandsCore/Public/GroomInstance.h
+++ b/Engine/Plugins/Runtime/HairStrands/Source/HairStrandsCore/Public/GroomInstance.h
@@ -277,7 +277,6 @@ struct HAIRSTRANDSCORE_API FHairGroupInstance : public FHairStrandsInstance
 	EHairBindingType		BindingType = EHairBindingType::NoneBinding;
 	bool					bForceCards = false;
 	bool					bUpdatePositionOffset = false;
-	bool					bCastShadow = true;
 	bool 					bSupportStreaming = true;
 	
 	// Deformed component to extract the bone buffer 
