Add render command pipe support for Groom.

diff --git a/Engine/Plugins/Runtime/HairStrands/Source/HairStrandsCore/Private/GroomAsset.cpp b/Engine/Plugins/Runtime/HairStrands/Source/HairStrandsCore/Private/GroomAsset.cpp
--- a/Engine/Plugins/Runtime/HairStrands/Source/HairStrandsCore/Private/GroomAsset.cpp
+++ b/Engine/Plugins/Runtime/HairStrands/Source/HairStrandsCore/Private/GroomAsset.cpp
@@ -440,7 +440,7 @@ static void InitAtlasTexture(ResourceType* InResource, const FHairGroupCardsText
 
 	if (bHasValidTextures)
 	{
-		ENQUEUE_RENDER_COMMAND(HairStrandsCardsTextureCommand)(
+		ENQUEUE_RENDER_COMMAND(HairStrandsCardsTextureCommand)(/*UE::RenderCommandPipe::Groom,*/
 			[InResource, InTextures, InLayoutIndex](FRHICommandListImmediate& RHICmdList)
 		{
 			FSamplerStateRHIRef DefaultSampler = TStaticSamplerState<SF_AnisotropicLinear, AM_Clamp, AM_Clamp, AM_Clamp, 0, 8/*MaxAnisotropy*/>::GetRHI();
@@ -977,7 +977,7 @@ inline void InternalReleaseResource(T*& Resource)
 	if (Resource)
 	{
 		T* InResource = Resource;
-		ENQUEUE_RENDER_COMMAND(ReleaseHairResourceCommand)(
+		ENQUEUE_RENDER_COMMAND(ReleaseHairResourceCommand)(UE::RenderCommandPipe::Groom,
 			[InResource](FRHICommandList& RHICmdList)
 			{
 				InResource->ReleaseResource();
@@ -3567,7 +3567,7 @@ void UGroomAsset::CreateDebugData()
 
 			FHairStrandsDebugDatas* InData = &GroupData.Debug.Data;
 			FHairStrandsDebugResources* InResource = GroupResource.Debug.Resource;
-			ENQUEUE_RENDER_COMMAND(HairStrandsDebugResourceCommand)(
+			ENQUEUE_RENDER_COMMAND(HairStrandsDebugResourceCommand)(/*UE::RenderCommandPipe::Groom,*/
 				[InData, InResource](FRHICommandListImmediate& RHICmdList)
 				{
 					FRDGBuilder GraphBuilder(RHICmdList);
diff --git a/Engine/Plugins/Runtime/HairStrands/Source/HairStrandsCore/Private/GroomBindingAsset.cpp b/Engine/Plugins/Runtime/HairStrands/Source/HairStrandsCore/Private/GroomBindingAsset.cpp
--- a/Engine/Plugins/Runtime/HairStrands/Source/HairStrandsCore/Private/GroomBindingAsset.cpp
+++ b/Engine/Plugins/Runtime/HairStrands/Source/HairStrandsCore/Private/GroomBindingAsset.cpp
@@ -232,7 +232,7 @@ void UGroomBindingAsset::ReleaseResource()
 		{
 			FHairStrandsRestRootResource* InSimRootResources = Resource.SimRootResources;
 			FHairStrandsRestRootResource* InRenRootResources = Resource.RenRootResources;
-			ENQUEUE_RENDER_COMMAND(ReleaseHairStrandsResourceCommand)(
+			ENQUEUE_RENDER_COMMAND(ReleaseHairStrandsResourceCommand)(UE::RenderCommandPipe::Groom,
 				[InSimRootResources, InRenRootResources](FRHICommandList& RHICmdList)
 			{
 				if (InSimRootResources)
@@ -251,7 +251,7 @@ void UGroomBindingAsset::ReleaseResource()
 
 			for (FHairStrandsRestRootResource*& InCardsRootResources : Resource.CardsRootResources)
 			{
-				ENQUEUE_RENDER_COMMAND(ReleaseHairStrandsResourceCommand)(
+				ENQUEUE_RENDER_COMMAND(ReleaseHairStrandsResourceCommand)(UE::RenderCommandPipe::Groom,
 					[InCardsRootResources](FRHICommandList& RHICmdList)
 					{
 						if (InCardsRootResources)
@@ -272,7 +272,7 @@ void UGroomBindingAsset::ReleaseResource()
 	{
 		FHairStrandsRestRootResource* InSimRootResources = ResourceToDelete.SimRootResources;
 		FHairStrandsRestRootResource* InRenRootResources = ResourceToDelete.RenRootResources;
-		ENQUEUE_RENDER_COMMAND(ReleaseHairStrandsResourceCommand)(
+		ENQUEUE_RENDER_COMMAND(ReleaseHairStrandsResourceCommand)(UE::RenderCommandPipe::Groom,
 			[InSimRootResources, InRenRootResources](FRHICommandList& RHICmdList)
 		{
 			if (InSimRootResources)
diff --git a/Engine/Plugins/Runtime/HairStrands/Source/HairStrandsCore/Private/GroomComponent.cpp b/Engine/Plugins/Runtime/HairStrands/Source/HairStrandsCore/Private/GroomComponent.cpp
--- a/Engine/Plugins/Runtime/HairStrands/Source/HairStrandsCore/Private/GroomComponent.cpp
+++ b/Engine/Plugins/Runtime/HairStrands/Source/HairStrandsCore/Private/GroomComponent.cpp
@@ -1718,7 +1718,7 @@ void UGroomComponent::SetForcedLOD(int32 CurrLODIndex)
 	if (FHairStrandsSceneProxy* GroomSceneProxy = (FHairStrandsSceneProxy*)SceneProxy)
 	{
 		const EHairLODSelectionType LocalLODSelectionType = LODSelectionType;
-		ENQUEUE_RENDER_COMMAND(FHairComponentSendLODIndex)(
+		ENQUEUE_RENDER_COMMAND(FHairComponentSendLODIndex)(/*UE::RenderCommandPipe::Groom,*/
 		[GroomSceneProxy, CurrLODIndex, LocalLODSelectionType, bHasLODSwitch, bPredictionLoad](FRHICommandListImmediate& RHICmdList)
 		{
 			for (const TRefCountPtr<FHairGroupInstance>& Instance : GroomSceneProxy->HairGroupInstances)
@@ -2233,7 +2233,7 @@ void UGroomComponent::UpdateSimulatedGroups()
 		TArray<TRefCountPtr<FHairGroupInstance>> LocalInstances = HairGroupInstances;
 		UGroomAsset* LocalGroomAsset = GroomAsset;
 		UGroomBindingAsset* LocalBindingAsset = BindingAsset;
-		ENQUEUE_RENDER_COMMAND(FHairStrandsTick_UEnableSimulatedGroups)(
+		ENQUEUE_RENDER_COMMAND(FHairStrandsTick_UEnableSimulatedGroups)(/*UE::RenderCommandPipe::Groom,*/
 			[LocalInstances, LocalGroomAsset, LocalBindingAsset, Id, bIsStrandsEnabled, LODIndex](FRHICommandListImmediate& RHICmdList)
 		{
 			int32 GroupIt = 0;
@@ -3356,7 +3356,7 @@ void UGroomComponent::TickAtThisTime(const float Time, bool bInIsRunning, bool b
 		}
 
 		// Queue the update for the render thread to sync it with GeometryCache rendering
-		ENQUEUE_RENDER_COMMAND(FGroomCacheUpdate)(
+		ENQUEUE_RENDER_COMMAND(FGroomCacheUpdate)(UE::RenderCommandPipe::Groom,
 		[this, Time](FRHICommandList& RHICmdList)
 		{
 			ElapsedTime = Time;
@@ -3526,7 +3526,7 @@ void UGroomComponent::TickComponent(float DeltaTime, enum ELevelTick TickType, F
 	{
 		TArray<TRefCountPtr<FHairGroupInstance>> LocalHairGroupInstances = HairGroupInstances;
 		const EHairLODSelectionType LocalLODSelectionType = LODSelectionType;
-		ENQUEUE_RENDER_COMMAND(FHairStrandsTick_TransformUpdate)(
+		ENQUEUE_RENDER_COMMAND(FHairStrandsTick_TransformUpdate)(/*UE::RenderCommandPipe::Groom,*/
 			[Id, FeatureLevel, LocalHairGroupInstances, EffectiveForceLOD, LocalLODSelectionType, bSwapBuffer, SkelLocalToTransform](FRHICommandListImmediate& RHICmdList)
 		{
 			for (const TRefCountPtr<FHairGroupInstance>& Instance : LocalHairGroupInstances)
@@ -3582,7 +3582,7 @@ void UGroomComponent::SendRenderDynamicData_Concurrent()
 		}
 
 		TArray<TRefCountPtr<FHairGroupInstance>> LocalHairGroupInstances = HairGroupInstances;
-		ENQUEUE_RENDER_COMMAND(FHairStrandsTick_TransformUpdate)(
+		ENQUEUE_RENDER_COMMAND(FHairStrandsTick_TransformUpdate)(/*UE::RenderCommandPipe::Groom,*/
 			[LocalHairGroupInstances, SkelLocalToTransform](FRHICommandListImmediate& RHICmdList)
 		{
 			for (const TRefCountPtr<FHairGroupInstance>& Instance : LocalHairGroupInstances)
diff --git a/Engine/Plugins/Runtime/HairStrands/Source/HairStrandsCore/Private/GroomResources.cpp b/Engine/Plugins/Runtime/HairStrands/Source/HairStrandsCore/Private/GroomResources.cpp
--- a/Engine/Plugins/Runtime/HairStrands/Source/HairStrandsCore/Private/GroomResources.cpp
+++ b/Engine/Plugins/Runtime/HairStrands/Source/HairStrandsCore/Private/GroomResources.cpp
@@ -36,6 +36,8 @@
 #include "DerivedDataRequestOwner.h"
 #endif
 
+DEFINE_RENDER_COMMAND_PIPE(Groom, ERenderCommandPipeFlags::None);
+
 static int32 GHairStrandsBulkData_ReleaseAfterUse = 0;
 static FAutoConsoleVariableRef CVarHairStrandsBulkData_ReleaseAfterUse(TEXT("r.HairStrands.Strands.BulkData.ReleaseAfterUse"), GHairStrandsBulkData_ReleaseAfterUse, TEXT("Release CPU bulk data once hair groom/groom binding asset GPU resources are created. This saves memory"));
 
diff --git a/Engine/Plugins/Runtime/HairStrands/Source/HairStrandsCore/Private/GroomTextureBuilder.cpp b/Engine/Plugins/Runtime/HairStrands/Source/HairStrandsCore/Private/GroomTextureBuilder.cpp
--- a/Engine/Plugins/Runtime/HairStrands/Source/HairStrandsCore/Private/GroomTextureBuilder.cpp
+++ b/Engine/Plugins/Runtime/HairStrands/Source/HairStrandsCore/Private/GroomTextureBuilder.cpp
@@ -438,7 +438,7 @@ void FGroomTextureBuilder::BuildFollicleTexture(const TArray<FFollicleInfo>& InC
 			}
 
 			const EPixelFormat Format = OutTexture->GetPixelFormat(0);
-			ENQUEUE_RENDER_COMMAND(FFollicleTextureQuery)(
+			ENQUEUE_RENDER_COMMAND(FFollicleTextureQuery)(/*UE::RenderCommandPipe::Groom,*/
 			[Resolution, MipCount, Format, GPUInfos, OutTexture](FRHICommandListImmediate& RHICmdList)
 			{
 				if (OutTexture->GetResource())
diff --git a/Engine/Plugins/Runtime/HairStrands/Source/HairStrandsCore/Private/HairCardsVertexFactory.cpp b/Engine/Plugins/Runtime/HairStrands/Source/HairStrandsCore/Private/HairCardsVertexFactory.cpp
--- a/Engine/Plugins/Runtime/HairStrands/Source/HairStrandsCore/Private/HairCardsVertexFactory.cpp
+++ b/Engine/Plugins/Runtime/HairStrands/Source/HairStrandsCore/Private/HairCardsVertexFactory.cpp
@@ -306,7 +306,7 @@ void FHairCardsVertexFactory::Copy(const FHairCardsVertexFactory& Other)
 {
 	FHairCardsVertexFactory* VertexFactory = this;
 	const FDataType* DataCopy = &Other.Data;
-	ENQUEUE_RENDER_COMMAND(FHairStrandsVertexFactoryCopyData)(
+	ENQUEUE_RENDER_COMMAND(FHairStrandsVertexFactoryCopyData)(/*UE::RenderCommandPipe::Groom,*/
 	[VertexFactory, DataCopy](FRHICommandListImmediate& RHICmdList)
 	{
 		VertexFactory->Data = *DataCopy;
diff --git a/Engine/Plugins/Runtime/HairStrands/Source/HairStrandsCore/Private/HairStrandsVertexFactory.cpp b/Engine/Plugins/Runtime/HairStrands/Source/HairStrandsCore/Private/HairStrandsVertexFactory.cpp
--- a/Engine/Plugins/Runtime/HairStrands/Source/HairStrandsCore/Private/HairStrandsVertexFactory.cpp
+++ b/Engine/Plugins/Runtime/HairStrands/Source/HairStrandsCore/Private/HairStrandsVertexFactory.cpp
@@ -214,7 +214,7 @@ void FHairStrandsVertexFactory::Copy(const FHairStrandsVertexFactory& Other)
 {
 	FHairStrandsVertexFactory* VertexFactory = this;
 	const FDataType* DataCopy = &Other.Data;
-	ENQUEUE_RENDER_COMMAND(FHairStrandsVertexFactoryCopyData)(
+	ENQUEUE_RENDER_COMMAND(FHairStrandsVertexFactoryCopyData)(/*UE::RenderCommandPipe::Groom,*/
 		[VertexFactory, DataCopy](FRHICommandListImmediate& RHICmdList)
 		{
 			VertexFactory->Data = *DataCopy;
diff --git a/Engine/Plugins/Runtime/HairStrands/Source/HairStrandsCore/Public/GroomResources.h b/Engine/Plugins/Runtime/HairStrands/Source/HairStrandsCore/Public/GroomResources.h
--- a/Engine/Plugins/Runtime/HairStrands/Source/HairStrandsCore/Public/GroomResources.h
+++ b/Engine/Plugins/Runtime/HairStrands/Source/HairStrandsCore/Public/GroomResources.h
@@ -14,6 +14,8 @@
 
 namespace UE::DerivedData { class FRequestOwner; }
 
+DECLARE_RENDER_COMMAND_PIPE(Groom, HAIRSTRANDSCORE_API);
+
 ////////////////////////////////////////////////////////////////////////////////////////////////////////////////////////
 // Helpers
 
