Fix Groom RBF resource not correctly be built/tracked.

diff --git a/Engine/Plugins/Runtime/HairStrands/Source/HairStrandsCore/Private/GroomAsset.cpp b/Engine/Plugins/Runtime/HairStrands/Source/HairStrandsCore/Private/GroomAsset.cpp
--- a/Engine/Plugins/Runtime/HairStrands/Source/HairStrandsCore/Private/GroomAsset.cpp
+++ b/Engine/Plugins/Runtime/HairStrands/Source/HairStrandsCore/Private/GroomAsset.cpp
@@ -3328,47 +3328,29 @@ void UGroomAsset::InitMeshesResources()
 
 EGroomGeometryType UGroomAsset::GetGeometryType(int32 GroupIndex, int32 LODIndex) const
 {
-	if (GroupIndex < 0 || GroupIndex >= GetHairGroupsLOD().Num())
+	if (GetHairGroupsLOD().IsValidIndex(GroupIndex) && GetHairGroupsLOD()[GroupIndex].LODs.IsValidIndex(LODIndex))
 	{
-		return EGroomGeometryType::Strands;
+		return GetHairGroupsLOD()[GroupIndex].LODs[LODIndex].GeometryType;
 	}
-
-	if (LODIndex < 0 || LODIndex >= GetHairGroupsLOD()[GroupIndex].LODs.Num())
-	{
-		return EGroomGeometryType::Strands;
-	}
-
-	return GetHairGroupsLOD()[GroupIndex].LODs[LODIndex].GeometryType;
+	return EGroomGeometryType::Strands;
 }
 
 EGroomBindingType UGroomAsset::GetBindingType(int32 GroupIndex, int32 LODIndex) const
 {
-	if (GroupIndex < 0 || GroupIndex >= GetHairGroupsLOD().Num() || !IsHairStrandsBindingEnable())
+	if (GetHairGroupsLOD().IsValidIndex(GroupIndex) && GetHairGroupsLOD()[GroupIndex].LODs.IsValidIndex(LODIndex) && IsHairStrandsBindingEnable())
 	{
-		return EGroomBindingType::Rigid; // Fallback to rigid by default
+		return GetHairGroupsLOD()[GroupIndex].LODs[LODIndex].BindingType;
 	}
-
-	if (LODIndex < 0 || LODIndex >= GetHairGroupsLOD()[GroupIndex].LODs.Num())
-	{
-		return EGroomBindingType::Rigid; // Fallback to rigid by default
-	}
-
-	return GetHairGroupsLOD()[GroupIndex].LODs[LODIndex].BindingType;
+	return EGroomBindingType::Rigid; // Fallback to rigid by default
 }
 
 bool UGroomAsset::IsVisible(int32 GroupIndex, int32 LODIndex) const
 {
-	if (GroupIndex < 0 || GroupIndex >= GetHairGroupsLOD().Num())
+	if (GetHairGroupsLOD().IsValidIndex(GroupIndex) && GetHairGroupsLOD()[GroupIndex].LODs.IsValidIndex(LODIndex))
 	{
-		return false;
+		return GetHairGroupsLOD()[GroupIndex].LODs[LODIndex].bVisible;
 	}
-
-	if (LODIndex < 0 || LODIndex >= GetHairGroupsLOD()[GroupIndex].LODs.Num())
-	{
-		return false;
-	}
-
-	return GetHairGroupsLOD()[GroupIndex].LODs[LODIndex].bVisible;
+	return false;
 }
 
 bool UGroomAsset::IsDeformationEnable(int32 GroupIndex) const
@@ -3380,12 +3362,7 @@ bool UGroomAsset::IsDeformationEnable(int32 GroupIndex) const
 
 bool UGroomAsset::IsSimulationEnable(int32 GroupIndex, int32 LODIndex) const 
 {
-	if (GroupIndex < 0 || GroupIndex >= GetHairGroupsLOD().Num() || !IsHairStrandsSimulationEnable())
-	{
-		return false;
-	}
-
-	if (LODIndex >= GetHairGroupsLOD()[GroupIndex].LODs.Num())
+	if (!GetHairGroupsLOD().IsValidIndex(GroupIndex) || LODIndex >= GetHairGroupsLOD()[GroupIndex].LODs.Num() || !IsHairStrandsSimulationEnable())
 	{
 		return false;
 	}
@@ -3410,32 +3387,23 @@ bool UGroomAsset::IsSimulationEnable(int32 GroupIndex, int32 LODIndex) const
 
 bool UGroomAsset::IsGlobalInterpolationEnable(int32 GroupIndex, int32 LODIndex) const
 {
-	if (GroupIndex < 0 || GroupIndex >= GetHairGroupsLOD().Num())
-	{
-		return false;
-	}
-
-	if (LODIndex < 0 || LODIndex >= GetHairGroupsLOD()[GroupIndex].LODs.Num())
-	{
-		return false;
-	}
-
-	// If the LOD index is not forced, then its value could be -1. In this case we return the 'global' asset value
-	if (LODIndex < 0)
+	if (GetHairGroupsLOD().IsValidIndex(GroupIndex))
 	{
-		return GetEnableGlobalInterpolation();
+		if (GetHairGroupsLOD()[GroupIndex].LODs.IsValidIndex(LODIndex))
+		{
+			return
+				GetHairGroupsLOD()[GroupIndex].LODs[LODIndex].GlobalInterpolation == EGroomOverrideType::Enable ||
+				(GetHairGroupsLOD()[GroupIndex].LODs[LODIndex].GlobalInterpolation == EGroomOverrideType::Auto && GetEnableGlobalInterpolation());
+		}
 	}
-
-	return
-		GetHairGroupsLOD()[GroupIndex].LODs[LODIndex].GlobalInterpolation == EGroomOverrideType::Enable ||
-		(GetHairGroupsLOD()[GroupIndex].LODs[LODIndex].GlobalInterpolation == EGroomOverrideType::Auto && GetEnableGlobalInterpolation());
+	return GetEnableGlobalInterpolation();
 }
 
 bool UGroomAsset::IsSimulationEnable() const
 {
 	for (int32 GroupIndex = 0; GroupIndex < GetHairGroupsPlatformData().Num(); ++GroupIndex)
 	{
-		for (int32 LODIt = 0; LODIt < GetHairGroupsLOD().Num(); ++LODIt)
+		for (int32 LODIt = 0; LODIt < GetHairGroupsLOD()[GroupIndex].LODs.Num(); ++LODIt)
 		{
 			if (IsSimulationEnable(GroupIndex, LODIt))
 			{
@@ -3449,7 +3417,7 @@ bool UGroomAsset::IsSimulationEnable() const
 
 bool UGroomAsset::NeedsInterpolationData(int32 GroupIndex) const
 {
-	for (int32 LODIt = 0; LODIt < GetHairGroupsLOD().Num(); ++LODIt)
+	for (int32 LODIt = 0; LODIt < GetHairGroupsLOD()[GroupIndex].LODs.Num(); ++LODIt)
 	{
 		if (IsSimulationEnable(GroupIndex, LODIt) || IsGlobalInterpolationEnable(GroupIndex, LODIt) || GetEnableSimulationCache())
 		{
diff --git a/Engine/Plugins/Runtime/HairStrands/Source/HairStrandsCore/Private/GroomBuilder.cpp b/Engine/Plugins/Runtime/HairStrands/Source/HairStrandsCore/Private/GroomBuilder.cpp
--- a/Engine/Plugins/Runtime/HairStrands/Source/HairStrandsCore/Private/GroomBuilder.cpp
+++ b/Engine/Plugins/Runtime/HairStrands/Source/HairStrandsCore/Private/GroomBuilder.cpp
@@ -40,7 +40,7 @@ static FAutoConsoleVariableRef CVarHairGroupIndexBuilder_MaxVoxelResolution(TEXT
 
 FString FGroomBuilder::GetVersion()
 {
-	return TEXT("v8r55");
+	return TEXT("v9a");
 }
 
 namespace FHairStrandsDecimation
diff --git a/Engine/Plugins/Runtime/HairStrands/Source/HairStrandsCore/Private/Niagara/NiagaraDataInterfaceHairStrands.cpp b/Engine/Plugins/Runtime/HairStrands/Source/HairStrandsCore/Private/Niagara/NiagaraDataInterfaceHairStrands.cpp
--- a/Engine/Plugins/Runtime/HairStrands/Source/HairStrandsCore/Private/Niagara/NiagaraDataInterfaceHairStrands.cpp
+++ b/Engine/Plugins/Runtime/HairStrands/Source/HairStrandsCore/Private/Niagara/NiagaraDataInterfaceHairStrands.cpp
@@ -381,7 +381,7 @@ void FNDIHairStrandsData::Update(UNiagaraDataInterfaceHairStrands* Interface, FN
 
 		const bool bHasValidBindingAsset = (Interface->IsComponentValid() && Interface->SourceComponent->BindingAsset && Interface->SourceComponent->GroomAsset);
 
-		GlobalInterpolation = bHasValidBindingAsset ? Interface->SourceComponent->GroomAsset->GetEnableGlobalInterpolation() : false;
+		GlobalInterpolation = bHasValidBindingAsset ? Interface->SourceComponent->GroomAsset->IsGlobalInterpolationEnable(GroupIndex, LODIndex) : false;
 		bSkinningTransfer = bHasValidBindingAsset ?
 			(Interface->SourceComponent->BindingAsset->GetSourceSkeletalMesh() && Interface->SourceComponent->BindingAsset->GetTargetSkeletalMesh() &&
 			 Interface->SourceComponent->BindingAsset->GetSourceSkeletalMesh() != Interface->SourceComponent->BindingAsset->GetTargetSkeletalMesh()) : false;
diff --git a/Engine/Source/Runtime/Renderer/Public/HairStrandsInterface.h b/Engine/Source/Runtime/Renderer/Public/HairStrandsInterface.h
--- a/Engine/Source/Runtime/Renderer/Public/HairStrandsInterface.h
+++ b/Engine/Source/Runtime/Renderer/Public/HairStrandsInterface.h
@@ -250,32 +250,27 @@ public:
 
 	bool IsVisible(int32 InLODIndex) const
 	{
-		if (InLODIndex < 0 && InLODIndex >= LODVisibilities.Num()) return false;
-		return LODVisibilities[InLODIndex];
+		return LODVisibilities.IsValidIndex(InLODIndex) ? LODVisibilities[InLODIndex] : false;
 	}
 
 	EHairGeometryType GetGeometryType(int32 InLODIndex) const
 	{
-		if (InLODIndex < 0 && InLODIndex >= LODGeometryTypes.Num()) return EHairGeometryType::NoneGeometry;
-		return LODGeometryTypes[InLODIndex];
+		return LODGeometryTypes.IsValidIndex(InLODIndex) ? LODGeometryTypes[InLODIndex] : EHairGeometryType::NoneGeometry;
 	}
 
 	EHairBindingType GetBindingType(int32 InLODIndex) const
 	{ 
-		if (InLODIndex < 0 || InLODIndex >= BindingTypes.Num()) return EHairBindingType::NoneBinding;
-		return BindingTypes[InLODIndex];
+		return BindingTypes.IsValidIndex(InLODIndex) ? BindingTypes[InLODIndex] :EHairBindingType::NoneBinding;
 	}	
 
 	bool IsSimulationEnable(int32 InLODIndex) const
 	{
-		if (InLODIndex < 0 || InLODIndex >= LODSimulations.Num()) return false;
-		return LODSimulations[InLODIndex];
+		return LODSimulations.IsValidIndex(InLODIndex) ? LODSimulations[InLODIndex] : false;
 	}
 
 	bool IsGlobalInterpolationEnable(int32 InLODIndex) const 
 	{
-		if (InLODIndex < 0 || InLODIndex >= LODGlobalInterpolations.Num()) return false;
-		return LODGlobalInterpolations[InLODIndex];
+		return LODGlobalInterpolations.IsValidIndex(InLODIndex) ? LODGlobalInterpolations[InLODIndex] : false;
 	}
 
 	void SetLODScreenSizes(const TArray<float>& ScreenSizes) { LODScreenSizes = ScreenSizes; }
