* Fix resource transition on mobile platforms.
* Fix crash when toggling platform preview with hair on mobile platforms within editor.

diff --git a/Engine/Plugins/Runtime/HairStrands/Source/HairStrandsCore/Private/GroomAsset.cpp b/Engine/Plugins/Runtime/HairStrands/Source/HairStrandsCore/Private/GroomAsset.cpp
--- a/Engine/Plugins/Runtime/HairStrands/Source/HairStrandsCore/Private/GroomAsset.cpp
+++ b/Engine/Plugins/Runtime/HairStrands/Source/HairStrandsCore/Private/GroomAsset.cpp
@@ -972,17 +972,22 @@ inline void InternalUpdateResource(FRenderResource* Resource)
 }
 
 template<typename T>
-inline void InternalReleaseResource(T*& Resource)
+inline void InternalReleaseResource(T*& Resource, bool bResetLoadedSize=false)
 {
 	if (Resource)
 	{
 		T* InResource = Resource;
+		if (bResetLoadedSize)
+		{
+			Resource->InternalResetLoadedSize();
+		}
+
 		ENQUEUE_RENDER_COMMAND(ReleaseHairResourceCommand)(UE::RenderCommandPipe::Groom,
-			[InResource](FRHICommandList& RHICmdList)
-			{
-				InResource->ReleaseResource();
-				delete InResource;
-			});
+		[InResource](FRHICommandList& RHICmdList)
+		{
+			InResource->ReleaseResource();
+			delete InResource;
+		});
 		Resource = nullptr;
 	}
 }
@@ -1078,9 +1083,9 @@ void UGroomAsset::ReleaseStrandsResource(uint32 GroupIndex)
 	if (GetHairGroupsResources().IsValidIndex(GroupIndex))
 	{
 		FHairGroupResources& GroupData = GetHairGroupsResources()[GroupIndex];
-		InternalReleaseResource(GroupData.Strands.RestResource);
-		InternalReleaseResource(GroupData.Strands.ClusterResource);
-		InternalReleaseResource(GroupData.Strands.InterpolationResource);
+		InternalReleaseResource(GroupData.Strands.RestResource, true);
+		InternalReleaseResource(GroupData.Strands.ClusterResource, true);
+		InternalReleaseResource(GroupData.Strands.InterpolationResource, true);
 		#if RHI_RAYTRACING
 		InternalReleaseResource(GroupData.Strands.RaytracingResource);
 		#endif
@@ -3786,4 +3791,33 @@ const TArray<FHairGroupResources>& UGroomAsset::GetHairGroupsResources() const
 TArray<FHairGroupResources>& UGroomAsset::GetHairGroupsResources()
 {
 	return HairGroupsResources;
-}
\ No newline at end of file
+}
+
+#if WITH_EDITOR
+void UGroomAsset::RecreateResources()
+{
+	FGroomComponentRecreateRenderStateContext RecreateContext(this);
+	ReleaseResource();
+	InitResources();
+}
+
+void UGroomAsset::ChangePlatformLevel(ERHIFeatureLevel::Type In)
+{
+	// When changing platform preview level, recreate resources to the correct platform settings (e.g., r.hairstrands.strands=0/1)
+	if (CachedResourcesPlatformLevel != In)
+	{
+		RecreateResources();
+		CachedResourcesPlatformLevel = In;
+	}
+}
+
+void UGroomAsset::ChangeFeatureLevel(ERHIFeatureLevel::Type In)
+{
+	// When changing feature level, recreate resources to the correct feature level
+	if (CachedResourcesFeatureLevel != In)
+	{
+		RecreateResources();
+		CachedResourcesFeatureLevel = In;
+	}
+}
+#endif
\ No newline at end of file
diff --git a/Engine/Plugins/Runtime/HairStrands/Source/HairStrandsCore/Private/GroomBindingAsset.cpp b/Engine/Plugins/Runtime/HairStrands/Source/HairStrandsCore/Private/GroomBindingAsset.cpp
--- a/Engine/Plugins/Runtime/HairStrands/Source/HairStrandsCore/Private/GroomBindingAsset.cpp
+++ b/Engine/Plugins/Runtime/HairStrands/Source/HairStrandsCore/Private/GroomBindingAsset.cpp
@@ -242,6 +242,7 @@ void UGroomBindingAsset::ReleaseResource()
 				}
 				if (InRenRootResources)
 				{
+					InRenRootResources->InternalResetLoadedSize();
 					InRenRootResources->ReleaseResource();
 					delete InRenRootResources;
 				}
@@ -1225,3 +1226,32 @@ TArray<UGroomBindingAsset::FHairGroupPlatformData>& UGroomBindingAsset::GetHairG
 	return HairGroupsPlatformData;
 	PRAGMA_ENABLE_DEPRECATION_WARNINGS
 }
+
+#if WITH_EDITOR
+void UGroomBindingAsset::RecreateResources()
+{
+	ReleaseResource();
+	InitResource();
+	OnGroomBindingAssetChanged.Broadcast();
+}
+
+void UGroomBindingAsset::ChangeFeatureLevel(ERHIFeatureLevel::Type In)
+{
+	// When changing feature level, recreate resources to the correct feature level
+	if (CachedResourcesFeatureLevel != In)
+	{
+		RecreateResources();
+		CachedResourcesFeatureLevel = In;
+	}
+}
+
+void UGroomBindingAsset::ChangePlatformLevel(ERHIFeatureLevel::Type In)
+{
+	// When changing platform preview level, recreate resources to the correct platform settings (e.g., r.hairstrands.strands=0/1)
+	if (CachedResourcesPlatformLevel != In)
+	{
+		RecreateResources();
+		CachedResourcesPlatformLevel = In;
+	}
+}
+#endif // WITH_EDITOR
\ No newline at end of file
diff --git a/Engine/Plugins/Runtime/HairStrands/Source/HairStrandsCore/Private/GroomComponent.cpp b/Engine/Plugins/Runtime/HairStrands/Source/HairStrandsCore/Private/GroomComponent.cpp
--- a/Engine/Plugins/Runtime/HairStrands/Source/HairStrandsCore/Private/GroomComponent.cpp
+++ b/Engine/Plugins/Runtime/HairStrands/Source/HairStrandsCore/Private/GroomComponent.cpp
@@ -2405,9 +2405,9 @@ static UMeshComponent* ValidateBindingAsset(
 	return ValidateBindingAsset(GroomAsset, BindingAsset, Cast<UGeometryCacheComponent>(MeshComponent), bIsBindingReloading, bValidationEnable, Component);
 }
 
-static EGroomGeometryType GetEffectiveGeometryType(EGroomGeometryType Type, bool bUseCards)
+static EGroomGeometryType GetEffectiveGeometryType(EGroomGeometryType Type, bool bUseCards, EShaderPlatform InPlatform)
 {
-	return Type == EGroomGeometryType::Strands && (!IsHairStrandsEnabled(EHairStrandsShaderType::Strands) || bUseCards) ? EGroomGeometryType::Cards : Type;
+	return Type == EGroomGeometryType::Strands && (!IsHairStrandsEnabled(EHairStrandsShaderType::Strands, InPlatform) || bUseCards) ? EGroomGeometryType::Cards : Type;
 }
 
 static EGroomCacheType GetEffectiveGroomCacheType(const UGroomCache* InCache, const UGroomAsset* InGroom)
@@ -2444,6 +2444,8 @@ void UGroomComponent::InitResources(bool bIsBindingReloading)
 
 	InitializedResources = GroomAsset;
 
+	const EShaderPlatform ShaderPlatform = GetWorld() && GetWorld()->Scene ? GetWorld()->Scene->GetShaderPlatform() : EShaderPlatform::SP_NumPlatforms;
+
 	// 1. Check if we need any kind of binding data, simulation data, or RBF data
 	//
 	//					  Requires    	| Requires     | Requires
@@ -2470,12 +2472,12 @@ void UGroomComponent::InitResources(bool bIsBindingReloading)
 	{
 		for (uint32 LODIt = 0, LODCount = GroomAsset->GetLODCount(); LODIt < LODCount; ++LODIt)
 		{
-			const EGroomGeometryType GeometryType = GetEffectiveGeometryType(GroomAsset->GetGeometryType(GroupIt, LODIt), bUseCards);
+			const EGroomGeometryType GeometryType = GetEffectiveGeometryType(GroomAsset->GetGeometryType(GroupIt, LODIt), bUseCards, ShaderPlatform);
 			const EGroomBindingType BindingType = GroomAsset->GetBindingType(GroupIt, LODIt);
 
 			// Note on Global Deformation:
 			// * Global deformation require to have skinning binding
-			// * Force global interpolation to be enable for meshes with skinning binding as we use RBF defomation for 'sticking' meshes onto skel. mesh surface			
+			// * Force global interpolation to be enable for meshes with skinning binding as we use RBF defomation for 'sticking' meshes onto skel. mesh surface
 			bHasNeedSkeletalMesh				= bHasNeedSkeletalMesh || BindingType == EGroomBindingType::Rigid;
 			bHasNeedSkinningBinding				= bHasNeedSkinningBinding || BindingType == EGroomBindingType::Skinning;
 			bHasNeedSimulation[GroupIt]			= bHasNeedSimulation[GroupIt] || IsSimulationEnable(GroupIt, LODIt);
@@ -2700,7 +2702,7 @@ void UGroomComponent::InitResources(bool bIsBindingReloading)
 
 				// * Force global interpolation to be enable for meshes with skinning binding as we use RBF defomation for 'sticking' meshes onto skel. mesh surface
 				// * Global deformation are allowed only with 'Skinning' binding type
-				const EHairGeometryType GeometryType = ToHairGeometryType(GetEffectiveGeometryType(GroomAsset->GetGeometryType(GroupIt, LODIt), bUseCards));
+				const EHairGeometryType GeometryType = ToHairGeometryType(GetEffectiveGeometryType(GroomAsset->GetGeometryType(GroupIt, LODIt), bUseCards, ShaderPlatform));
 				const bool LODSimulation = IsSimulationEnable(GroupIt, LODIt);
 				const bool LODGlobalInterpolation = LocalBindingAsset && BindingType == EHairBindingType::Skinning && GroomAsset->IsGlobalInterpolationEnable(GroupIt, LODIt);
 				bNeedStrandsData = bNeedStrandsData || GeometryType == EHairGeometryType::Strands;
@@ -3038,7 +3040,6 @@ void UGroomComponent::InitResources(bool bIsBindingReloading)
 
 void UGroomComponent::ReleaseResources()
 {
-	FHairStrandsSceneProxy* GroomSceneProxy = (FHairStrandsSceneProxy*)SceneProxy;
 	InitializedResources = nullptr;
 
 	HairGroupInstances.Empty();
@@ -3972,7 +3973,38 @@ void UGroomComponent::CheckForErrors()
 
 	ValidateMaterials(true);
 }
-#endif
+
+void UGroomComponent::PreFeatureLevelChange(ERHIFeatureLevel::Type PendingFeatureLevel)
+{
+
+}
+
+void UGroomComponent::HandlePlatformPreviewChanged(ERHIFeatureLevel::Type InFeatureLevel)
+{
+	if (BindingAsset)
+	{
+		BindingAsset->ChangePlatformLevel(InFeatureLevel);
+	}
+	if (GroomAsset)
+	{
+		GroomAsset->ChangePlatformLevel(InFeatureLevel);
+	}
+	InvalidateAndRecreate();
+}
+
+void UGroomComponent::HandleFeatureLevelChanged(ERHIFeatureLevel::Type InFeatureLevel)
+{
+	if (BindingAsset)
+	{
+		BindingAsset->ChangeFeatureLevel(InFeatureLevel);
+	}
+	if (GroomAsset)
+	{
+		GroomAsset->ChangeFeatureLevel(InFeatureLevel);
+	}
+	InvalidateAndRecreate();
+}
+#endif // WITH_EDITOR
 
 template<typename T>
 void InternalAddDedicatedVideoMemoryBytes(FResourceSizeEx& CumulativeResourceSize, T Resource)
diff --git a/Engine/Plugins/Runtime/HairStrands/Source/HairStrandsCore/Private/GroomManager.cpp b/Engine/Plugins/Runtime/HairStrands/Source/HairStrandsCore/Private/GroomManager.cpp
--- a/Engine/Plugins/Runtime/HairStrands/Source/HairStrandsCore/Private/GroomManager.cpp
+++ b/Engine/Plugins/Runtime/HairStrands/Source/HairStrandsCore/Private/GroomManager.cpp
@@ -1966,6 +1966,7 @@ static void RunHairStrandsInterpolation_Cards(
 				AddHairCardsDeformationPass(
 					GraphBuilder,
 					ShaderMap,
+					View->GetFeatureLevel(),
 					ShaderPrintData,
 					InstanceData.Instance,
 					InstanceData.MeshLODIndex);
diff --git a/Engine/Plugins/Runtime/HairStrands/Source/HairStrandsCore/Private/GroomResources.cpp b/Engine/Plugins/Runtime/HairStrands/Source/HairStrandsCore/Private/GroomResources.cpp
--- a/Engine/Plugins/Runtime/HairStrands/Source/HairStrandsCore/Private/GroomResources.cpp
+++ b/Engine/Plugins/Runtime/HairStrands/Source/HairStrandsCore/Private/GroomResources.cpp
@@ -858,11 +858,13 @@ void FHairCardsRestResource::InternalAllocate()
 {
 	FRHICommandListBase& RHICmdList = FRHICommandListImmediate::Get();
 
+	const ERHIAccess InitialAccess = ERHIAccess::SRVMask;
+
 	// These resources are kept as regular (i.e., non-RDG resources) as they need to be bound at the input assembly stage by the Vertex declaraction which requires FVertexBuffer type
-	CreateBuffer<FHairCardsPositionFormat>(RHICmdList, BulkData.Positions, RestPositionBuffer, ToHairResourceDebugName(TEXT("Hair.CardsRest_PositionBuffer"), ResourceName), OwnerName);
-	CreateBuffer<FHairCardsNormalFormat>(RHICmdList, BulkData.Normals, NormalsBuffer, ToHairResourceDebugName(TEXT("Hair.CardsRest_NormalBuffer"), ResourceName), OwnerName);
-	CreateBuffer<FHairCardsUVFormat>(RHICmdList, BulkData.UVs, UVsBuffer, ToHairResourceDebugName(TEXT("Hair.CardsRest_UVBuffer"), ResourceName), OwnerName);
-	CreateBuffer<FHairCardsMaterialFormat>(RHICmdList, BulkData.Materials, MaterialsBuffer, ToHairResourceDebugName(TEXT("Hair.CardsRest_MaterialBuffer"), ResourceName), OwnerName);
+	CreateBuffer<FHairCardsPositionFormat>(RHICmdList, BulkData.Positions, RestPositionBuffer, ToHairResourceDebugName(TEXT("Hair.CardsRest_PositionBuffer"), ResourceName), OwnerName, InitialAccess);
+	CreateBuffer<FHairCardsNormalFormat>(RHICmdList, BulkData.Normals, NormalsBuffer, ToHairResourceDebugName(TEXT("Hair.CardsRest_NormalBuffer"), ResourceName), OwnerName, InitialAccess);
+	CreateBuffer<FHairCardsUVFormat>(RHICmdList, BulkData.UVs, UVsBuffer, ToHairResourceDebugName(TEXT("Hair.CardsRest_UVBuffer"), ResourceName), OwnerName, InitialAccess);
+	CreateBuffer<FHairCardsMaterialFormat>(RHICmdList, BulkData.Materials, MaterialsBuffer, ToHairResourceDebugName(TEXT("Hair.CardsRest_MaterialBuffer"), ResourceName), OwnerName, InitialAccess);
 }
 
 void FHairCardsRestResource::InternalRelease()
@@ -947,10 +949,12 @@ void FHairMeshesRestResource::InternalAllocate()
 {
 	FRHICommandListBase& RHICmdList = FRHICommandListImmediate::Get();
 
+	const ERHIAccess InitialAccess = ERHIAccess::SRVMask;
+
 	// These resources are kept as regular (i.e., non-RDG resources) as they need to be bound at the input assembly stage by the Vertex declaraction which requires FVertexBuffer type
-	CreateBuffer<FHairCardsPositionFormat>(RHICmdList, BulkData.Positions, RestPositionBuffer, ToHairResourceDebugName(TEXT("Hair.MeshesRest_Positions"), ResourceName), OwnerName);
-	CreateBuffer<FHairCardsNormalFormat>(RHICmdList, BulkData.Normals, NormalsBuffer, ToHairResourceDebugName(TEXT("Hair.MeshesRest_Normals"), ResourceName), OwnerName);
-	CreateBuffer<FHairCardsUVFormat>(RHICmdList, BulkData.UVs, UVsBuffer, ToHairResourceDebugName(TEXT("Hair.MeshesRest_UVs"), ResourceName), OwnerName);
+	CreateBuffer<FHairCardsPositionFormat>(RHICmdList, BulkData.Positions, RestPositionBuffer, ToHairResourceDebugName(TEXT("Hair.MeshesRest_Positions"), ResourceName), OwnerName, InitialAccess);
+	CreateBuffer<FHairCardsNormalFormat>(RHICmdList, BulkData.Normals, NormalsBuffer, ToHairResourceDebugName(TEXT("Hair.MeshesRest_Normals"), ResourceName), OwnerName, InitialAccess);
+	CreateBuffer<FHairCardsUVFormat>(RHICmdList, BulkData.UVs, UVsBuffer, ToHairResourceDebugName(TEXT("Hair.MeshesRest_UVs"), ResourceName), OwnerName, InitialAccess);
 }
 
 void FHairMeshesRestResource::InternalRelease()
diff --git a/Engine/Plugins/Runtime/HairStrands/Source/HairStrandsCore/Private/HairCardsVertexFactory.cpp b/Engine/Plugins/Runtime/HairStrands/Source/HairStrandsCore/Private/HairCardsVertexFactory.cpp
--- a/Engine/Plugins/Runtime/HairStrands/Source/HairStrandsCore/Private/HairCardsVertexFactory.cpp
+++ b/Engine/Plugins/Runtime/HairStrands/Source/HairStrandsCore/Private/HairCardsVertexFactory.cpp
@@ -326,15 +326,17 @@ void FHairCardsVertexFactory::InitResources(FRHICommandListBase& RHICmdList)
 	// We create different streams based on feature level
 	check(HasValidFeatureLevel());
 
+	const ERHIFeatureLevel::Type CurrentFeatureLevel = GetFeatureLevel();
+
 	// If the platform does not support manual vertex fetching we assume it is a low end platform, and so we don't enable deformation.
 	// VertexFactory needs to be able to support max possible shader platform and feature level
 	// in case if we switch feature level at runtime.
 	FVertexDeclarationElementList Elements;
-	SetPrimitiveIdStreamIndex(GetFeatureLevel(), EVertexInputStreamType::Default, -1);
+	SetPrimitiveIdStreamIndex(CurrentFeatureLevel, EVertexInputStreamType::Default, -1);
 	AddPrimitiveIdStreamElement(EVertexInputStreamType::Default, Elements, 13, 0xff);
 
 	// Note this is a local version of the VF's bSupportsManualVertexFetch, which take into account the feature level
-	const bool bManualFetch = SupportsManualVertexFetch(GetFeatureLevel());
+	const bool bManualFetch = SupportsManualVertexFetch(CurrentFeatureLevel);
 	if (!bManualFetch)
 	{
 		if (Data.GeometryType == EHairGeometryType::Cards)
@@ -346,6 +348,17 @@ void FHairCardsVertexFactory::InitResources(FRHICommandListBase& RHICmdList)
 			Elements.Add(AccessStreamComponent(FVertexStreamComponent(&LOD.RestResource->NormalsBuffer,			0, FHairCardsNormalFormat::SizeInByte,	FHairCardsNormalFormat::SizeInByte * FHairCardsNormalFormat::ComponentCount,		FHairCardsNormalFormat::VertexElementType,		EVertexStreamUsage::Default), 2));
 			Elements.Add(AccessStreamComponent(FVertexStreamComponent(&LOD.RestResource->UVsBuffer,				0, 0,									FHairCardsUVFormat::SizeInByte,														FHairCardsUVFormat::VertexElementType,			EVertexStreamUsage::Default), 3));
 			Elements.Add(AccessStreamComponent(FVertexStreamComponent(&LOD.RestResource->MaterialsBuffer,		0, 0,									FHairCardsMaterialFormat::SizeInByte,												FHairCardsMaterialFormat::VertexElementType,	EVertexStreamUsage::Default), 4));
+
+			// Ensure the rest resources are in correct states for the VF
+			const FHairCardsRestResource* RestResource = LOD.RestResource;
+			ENQUEUE_RENDER_COMMAND(FHairCardsVertexFactoryTransition)(/*UE::RenderCommandPipe::Groom,*/
+			[RestResource](FRHICommandListImmediate& RHICmdImmediateList)
+			{
+				RHICmdImmediateList.Transition(FRHITransitionInfo(RestResource->RestPositionBuffer.VertexBufferRHI,	ERHIAccess::Unknown, ERHIAccess::VertexOrIndexBuffer));
+				RHICmdImmediateList.Transition(FRHITransitionInfo(RestResource->NormalsBuffer.VertexBufferRHI, 		ERHIAccess::Unknown, ERHIAccess::VertexOrIndexBuffer));
+				RHICmdImmediateList.Transition(FRHITransitionInfo(RestResource->UVsBuffer.VertexBufferRHI, 			ERHIAccess::Unknown, ERHIAccess::VertexOrIndexBuffer));
+				RHICmdImmediateList.Transition(FRHITransitionInfo(RestResource->MaterialsBuffer.VertexBufferRHI, 	ERHIAccess::Unknown, ERHIAccess::VertexOrIndexBuffer));
+			});
 		}
 		else if (Data.GeometryType == EHairGeometryType::Meshes)
 		{
@@ -357,6 +370,16 @@ void FHairCardsVertexFactory::InitResources(FRHICommandListBase& RHICmdList)
 			Elements.Add(AccessStreamComponent(FVertexStreamComponent(&LOD.RestResource->NormalsBuffer,			0, FHairCardsNormalFormat::SizeInByte,	FHairCardsNormalFormat::SizeInByte * FHairCardsNormalFormat::ComponentCount,		FHairCardsNormalFormat::VertexElementType,		EVertexStreamUsage::Default), 2));
 			Elements.Add(AccessStreamComponent(FVertexStreamComponent(&LOD.RestResource->UVsBuffer,				0, 0,									FHairCardsUVFormat::SizeInByte,														FHairCardsUVFormat::VertexElementType,			EVertexStreamUsage::Default), 3));
 			Elements.Add(AccessStreamComponent(FVertexStreamComponent(&LOD.RestResource->NormalsBuffer,			0, 0,									FHairCardsMaterialFormat::SizeInByte,												FHairCardsMaterialFormat::VertexElementType,	EVertexStreamUsage::Default), 4)); 
+
+			// Ensure the rest resources are in correct states for the VF
+			const FHairMeshesRestResource* RestResource = LOD.RestResource;
+			ENQUEUE_RENDER_COMMAND(FHairCardsVertexFactoryTransition)(/*UE::RenderCommandPipe::Groom,*/
+			[RestResource](FRHICommandListImmediate& RHICmdImmediateList)
+			{
+				RHICmdImmediateList.Transition(FRHITransitionInfo(RestResource->RestPositionBuffer.VertexBufferRHI,	ERHIAccess::Unknown, ERHIAccess::VertexOrIndexBuffer));
+				RHICmdImmediateList.Transition(FRHITransitionInfo(RestResource->NormalsBuffer.VertexBufferRHI, 		ERHIAccess::Unknown, ERHIAccess::VertexOrIndexBuffer));
+				RHICmdImmediateList.Transition(FRHITransitionInfo(RestResource->UVsBuffer.VertexBufferRHI, 			ERHIAccess::Unknown, ERHIAccess::VertexOrIndexBuffer));
+			});
 		}
 
 		bNeedsDeclaration = true;
diff --git a/Engine/Plugins/Runtime/HairStrands/Source/HairStrandsCore/Private/HairStrandsCore.cpp b/Engine/Plugins/Runtime/HairStrands/Source/HairStrandsCore/Private/HairStrandsCore.cpp
--- a/Engine/Plugins/Runtime/HairStrands/Source/HairStrandsCore/Private/HairStrandsCore.cpp
+++ b/Engine/Plugins/Runtime/HairStrands/Source/HairStrandsCore/Private/HairStrandsCore.cpp
@@ -174,7 +174,6 @@ UGroomBindingAsset* FHairStrandsCore::CreateGroomBindingAsset(const FString& InP
 
 UGroomBindingAsset* FHairStrandsCore::CreateGroomBindingAsset(EGroomBindingMeshType BindingType, const FString& InPackageName, UObject* InParent, UGroomAsset* GroomAsset, UObject* Source, UObject* Target, const int32 NumInterpolationPoints, const int32 MatchingSection)
 {
-#if WITH_EDITOR
 	if (!GroomAsset ||
 		!((BindingType == EGroomBindingMeshType::SkeletalMesh && Cast<USkeletalMesh>(Target)) || 
 		 (BindingType == EGroomBindingMeshType::GeometryCache && Cast<UGeometryCache>(Target)))
@@ -238,7 +237,6 @@ UGroomBindingAsset* FHairStrandsCore::CreateGroomBindingAsset(EGroomBindingMeshT
 		AssetHelper().RegisterAsset(Out);
 		return Out;
 	}
-#endif
 	return nullptr;
 }
 
@@ -261,4 +259,4 @@ void FHairStrandsCore::SaveAsset(UObject* Object)
 
 	AssetHelper().SaveAsset(Object);
 }
-#endif
+#endif // WITH_EDITOR
diff --git a/Engine/Plugins/Runtime/HairStrands/Source/HairStrandsCore/Private/HairStrandsInterpolation.cpp b/Engine/Plugins/Runtime/HairStrands/Source/HairStrandsCore/Private/HairStrandsInterpolation.cpp
--- a/Engine/Plugins/Runtime/HairStrands/Source/HairStrandsCore/Private/HairStrandsInterpolation.cpp
+++ b/Engine/Plugins/Runtime/HairStrands/Source/HairStrandsCore/Private/HairStrandsInterpolation.cpp
@@ -933,6 +933,7 @@ IMPLEMENT_GLOBAL_SHADER(FHairCardsDeformationCS, "/Engine/Private/HairStrands/Ha
 void AddHairCardsDeformationPass(
 	FRDGBuilder& GraphBuilder,
 	FGlobalShaderMap* ShaderMap,
+	const ERHIFeatureLevel::Type FeatureLevel,
 	const FShaderPrintData* ShaderPrintData,
 	FHairGroupInstance* Instance,
 	const int32 MeshLODIndex)
@@ -1012,12 +1013,31 @@ void AddHairCardsDeformationPass(
 	TShaderMapRef<FHairCardsDeformationCS> ComputeShader(ShaderMap, PermutationVector);
 
 	const FIntVector DispatchCount = FIntVector(FMath::DivideAndRoundUp(Parameters->CardsVertexCount, FHairCardsDeformationCS::GetGroupSize()), 1, 1);
-	FComputeShaderUtils::AddPass(
-		GraphBuilder,
+
+	const bool bManualFetch = FVertexFactoryType::CheckManualVertexFetchSupport(FeatureLevel);
+	FBufferRHIRef CardsRestPositionBuffer = LOD.RestResource->RestPositionBuffer.VertexBufferRHI;
+	FBufferRHIRef CardsRestTangentBuffer  = LOD.RestResource->NormalsBuffer.VertexBufferRHI;
+
+	GraphBuilder.AddPass(
 		RDG_EVENT_NAME("HairStrands::CardsDeformation(%s)", bSupportDynamicMesh ? TEXT("Dynamic") : TEXT("Static")),
-		ComputeShader,
 		Parameters,
-		DispatchCount);
+		ERDGPassFlags::Compute,
+		[Parameters, ComputeShader, DispatchCount, CardsRestPositionBuffer, CardsRestTangentBuffer, bManualFetch](FRHIComputeCommandList& RHICmdList)
+		{
+			// On platforms not supporting manual vertex fetching, ensure the resources are in 'VerteOrIndexBuffer' state after position/normals update
+			if (!bManualFetch)
+			{
+				RHICmdList.Transition(FRHITransitionInfo(CardsRestPositionBuffer, ERHIAccess::Unknown, ERHIAccess::SRVMask));
+				RHICmdList.Transition(FRHITransitionInfo(CardsRestTangentBuffer, ERHIAccess::Unknown, ERHIAccess::SRVMask));
+			}
+			FComputeShaderUtils::Dispatch(RHICmdList, ComputeShader, *Parameters, DispatchCount);
+			if (!bManualFetch)
+			{
+				RHICmdList.Transition(FRHITransitionInfo(CardsRestPositionBuffer, ERHIAccess::SRVMask, ERHIAccess::VertexOrIndexBuffer));
+				RHICmdList.Transition(FRHITransitionInfo(CardsRestTangentBuffer, ERHIAccess::SRVMask, ERHIAccess::VertexOrIndexBuffer));
+			}
+		});
+
 
 	// If LOD has switched, copy the current buffer, so that we don't get incorrect motion vector
 	if (bHasLODSwitch)
diff --git a/Engine/Plugins/Runtime/HairStrands/Source/HairStrandsCore/Public/GroomAsset.h b/Engine/Plugins/Runtime/HairStrands/Source/HairStrandsCore/Public/GroomAsset.h
--- a/Engine/Plugins/Runtime/HairStrands/Source/HairStrandsCore/Public/GroomAsset.h
+++ b/Engine/Plugins/Runtime/HairStrands/Source/HairStrandsCore/Public/GroomAsset.h
@@ -726,6 +726,10 @@ public:
 	FOnGroomAsyncLoadFinished OnGroomAsyncLoadFinished;
 
 	void MarkMaterialsHasChanged();
+
+	void RecreateResources();
+	void ChangeFeatureLevel(ERHIFeatureLevel::Type PendingFeatureLevel);
+	void ChangePlatformLevel(ERHIFeatureLevel::Type PendingFeatureLevel);
 #endif
 
 	/** Array of user data stored with the asset */
@@ -843,6 +847,8 @@ private:
 	TArray<FHairGroupsLOD>						CachedHairGroupsLOD;
 	TArray<FHairGroupsCardsSourceDescription>	CachedHairGroupsCards;
 	TArray<FHairGroupsMeshesSourceDescription>	CachedHairGroupsMeshes;
+	ERHIFeatureLevel::Type 						CachedResourcesPlatformLevel= ERHIFeatureLevel::Num;
+	ERHIFeatureLevel::Type 						CachedResourcesFeatureLevel = ERHIFeatureLevel::Num;
 
 	// Queue of procedural assets which needs to be saved
 	TQueue<UStaticMesh*> AssetToSave_Meshes;
diff --git a/Engine/Plugins/Runtime/HairStrands/Source/HairStrandsCore/Public/GroomBindingAsset.h b/Engine/Plugins/Runtime/HairStrands/Source/HairStrandsCore/Public/GroomBindingAsset.h
--- a/Engine/Plugins/Runtime/HairStrands/Source/HairStrandsCore/Public/GroomBindingAsset.h
+++ b/Engine/Plugins/Runtime/HairStrands/Source/HairStrandsCore/Public/GroomBindingAsset.h
@@ -254,6 +254,10 @@ public:
 	//private :
 #if WITH_EDITOR
 	FOnGroomBindingAssetChanged OnGroomBindingAssetChanged;
+
+	void RecreateResources();
+	void ChangeFeatureLevel(ERHIFeatureLevel::Type PendingFeatureLevel);
+	void ChangePlatformLevel(ERHIFeatureLevel::Type PendingFeatureLevel);
 #endif
 
 #if WITH_EDITORONLY_DATA
@@ -281,6 +285,10 @@ private:
 	bool bRegisterTargetMeshCallback = false;
 	bool bRegisterGroomAssetCallback = false;
 	TArray<FString> CachedDerivedDataKey;
+#endif
+#if WITH_EDITOR
+	ERHIFeatureLevel::Type CachedResourcesFeatureLevel = ERHIFeatureLevel::Num;
+	ERHIFeatureLevel::Type CachedResourcesPlatformLevel = ERHIFeatureLevel::Num;
 #endif
 	bool bIsValid = false;
 	uint32 AssetNameHash = 0;
diff --git a/Engine/Plugins/Runtime/HairStrands/Source/HairStrandsCore/Public/GroomComponent.h b/Engine/Plugins/Runtime/HairStrands/Source/HairStrandsCore/Public/GroomComponent.h
--- a/Engine/Plugins/Runtime/HairStrands/Source/HairStrandsCore/Public/GroomComponent.h
+++ b/Engine/Plugins/Runtime/HairStrands/Source/HairStrandsCore/Public/GroomComponent.h
@@ -177,6 +177,12 @@ public:
 	void InvalidateAndRecreate();
 #endif
 
+#if WITH_EDITOR
+	virtual void PreFeatureLevelChange(ERHIFeatureLevel::Type PendingFeatureLevel) override;
+	void HandlePlatformPreviewChanged(ERHIFeatureLevel::Type InFeatureLevel);
+	void HandleFeatureLevelChanged(ERHIFeatureLevel::Type InFeatureLevel);
+#endif
+
 	/* Accessor function for changing Groom asset from blueprint/sequencer */
 	UFUNCTION(BlueprintCallable, Category = "Groom")
 	void SetGroomAsset(UGroomAsset* Asset);
diff --git a/Engine/Plugins/Runtime/HairStrands/Source/HairStrandsCore/Public/HairStrandsInterpolation.h b/Engine/Plugins/Runtime/HairStrands/Source/HairStrandsCore/Public/HairStrandsInterpolation.h
--- a/Engine/Plugins/Runtime/HairStrands/Source/HairStrandsCore/Public/HairStrandsInterpolation.h
+++ b/Engine/Plugins/Runtime/HairStrands/Source/HairStrandsCore/Public/HairStrandsInterpolation.h
@@ -225,6 +225,7 @@ void AddBuildHairMeshAccelerationStructurePass(
 void AddHairCardsDeformationPass(
 	FRDGBuilder& GraphBuilder,
 	FGlobalShaderMap* ShaderMap,
+	const ERHIFeatureLevel::Type FeatureLevel,
 	const FShaderPrintData* ShaderPrintData,
 	FHairGroupInstance* Instance,
 	const int32 MeshLODIndex);
\ No newline at end of file
diff --git a/Engine/Plugins/Runtime/HairStrands/Source/HairStrandsEditor/Private/HairStrandsEditor.cpp b/Engine/Plugins/Runtime/HairStrands/Source/HairStrandsEditor/Private/HairStrandsEditor.cpp
--- a/Engine/Plugins/Runtime/HairStrands/Source/HairStrandsEditor/Private/HairStrandsEditor.cpp
+++ b/Engine/Plugins/Runtime/HairStrands/Source/HairStrandsEditor/Private/HairStrandsEditor.cpp
@@ -29,6 +29,8 @@
 #include "ISettingsModule.h"
 #include "ISettingsSection.h"
 #include "ISequencerModule.h"
+#include "Delegates/DelegateCombinations.h"
+#include "UObject/UObjectIterator.h"
 
 IMPLEMENT_MODULE(FGroomEditor, HairStrandsEditor);
 
@@ -74,6 +76,10 @@ void FGroomEditor::StartupModule()
 	LLM_SCOPE_BYTAG(GroomEditor)
 
 	UToolMenus::RegisterStartupCallback(FSimpleMulticastDelegate::FDelegate::CreateRaw(this, &FGroomEditor::RegisterMenus));
+	
+	// Any attempt to use GEditor right now will fail as it hasn't been initialized yet. Waiting for post engine init resolves that.
+	FCoreDelegates::OnPostEngineInit.AddRaw(this, &FGroomEditor::OnPostEngineInit);
+	FCoreDelegates::OnEnginePreExit.AddRaw(this, &FGroomEditor::OnPreExit);
 
 	IAssetTools& AssetTools = FModuleManager::LoadModuleChecked<FAssetToolsModule>("AssetTools").Get();
 
@@ -147,6 +153,9 @@ void FGroomEditor::ShutdownModule()
 	UToolMenus::UnRegisterStartupCallback(this);
 	UToolMenus::UnregisterOwner(UE_MODULE_NAME);
 
+	FCoreDelegates::OnPostEngineInit.RemoveAll(this);
+	FCoreDelegates::OnEnginePreExit.RemoveAll(this);
+
 	ISettingsModule* SettingsModule = FModuleManager::GetModulePtr<ISettingsModule>("Settings");
 	if (SettingsModule != nullptr)
 	{
@@ -195,4 +204,64 @@ TArray<TSharedPtr<IGroomTranslator>> FGroomEditor::GetHairTranslators()
 	return Translators;
 }
 
+void FGroomEditor::OnPostEngineInit()
+{
+	// The editor should be valid at this point.. log a warning if not!
+	if (GEditor)
+	{
+		UEditorEngine* EditorEngine = CastChecked<UEditorEngine>(GEngine);
+		PreviewPlatformChangedHandle = EditorEngine->OnPreviewPlatformChanged().AddRaw(this, &FGroomEditor::OnPreviewPlatformChanged);
+		PreviewFeatureLevelChangedHandle = EditorEngine->OnPreviewFeatureLevelChanged().AddRaw(this, &FGroomEditor::OnPreviewFeatureLevelChanged);
+	}
+}
+
+void FGroomEditor::OnPreExit()
+{
+	if (GEditor)
+	{
+		CastChecked<UEditorEngine>(GEngine)->OnPreviewPlatformChanged().Remove(PreviewPlatformChangedHandle);
+		CastChecked<UEditorEngine>(GEngine)->OnPreviewFeatureLevelChanged().Remove(PreviewFeatureLevelChangedHandle);
+	}
+}
+
+void FGroomEditor::OnPreviewPlatformChanged()
+{
+#if WITH_EDITOR
+	UEditorEngine* EditorEngine = CastChecked<UEditorEngine>(GEngine);
+	ERHIFeatureLevel::Type ActiveFeatureLevel = EditorEngine->GetDefaultWorldFeatureLevel();
+	if (EditorEngine->IsFeatureLevelPreviewActive())
+	{
+		ActiveFeatureLevel = EditorEngine->GetActiveFeatureLevelPreviewType();
+	}
+
+	for (TObjectIterator<UGroomComponent> It; It; ++It)
+	{
+		if (UGroomComponent* Component = *It)
+		{
+			Component->HandlePlatformPreviewChanged(ActiveFeatureLevel);
+		}
+	}
+#endif
+}
+
+void FGroomEditor::OnPreviewFeatureLevelChanged(ERHIFeatureLevel::Type InPreviewFeatureLevel)
+{
+#if WITH_EDITOR
+	UEditorEngine* EditorEngine = CastChecked<UEditorEngine>(GEngine);
+	ERHIFeatureLevel::Type ActiveFeatureLevel = EditorEngine->GetDefaultWorldFeatureLevel();
+	if (EditorEngine->IsFeatureLevelPreviewActive())
+	{
+		ActiveFeatureLevel = EditorEngine->GetActiveFeatureLevelPreviewType();
+	}
+
+	for (TObjectIterator<UGroomComponent> It; It; ++It)
+	{
+		if (UGroomComponent* Component = *It)
+		{
+			Component->HandleFeatureLevelChanged(InPreviewFeatureLevel);
+		}
+	}
+#endif
+}
+
 #undef LOCTEXT_NAMESPACE
\ No newline at end of file
diff --git a/Engine/Plugins/Runtime/HairStrands/Source/HairStrandsEditor/Public/HairStrandsEditor.h b/Engine/Plugins/Runtime/HairStrands/Source/HairStrandsEditor/Public/HairStrandsEditor.h
--- a/Engine/Plugins/Runtime/HairStrands/Source/HairStrandsEditor/Public/HairStrandsEditor.h
+++ b/Engine/Plugins/Runtime/HairStrands/Source/HairStrandsEditor/Public/HairStrandsEditor.h
@@ -7,6 +7,7 @@
 #include "Styling/ISlateStyle.h"
 #include "Styling/SlateStyle.h"
 #include "HAL/LowLevelMemTracker.h"
+#include "RHIFwd.h"
 
 #define HAIRSTRANDSEDITOR_MODULE_NAME TEXT("HairStrandsEditor")
 
@@ -23,11 +24,12 @@ public:
 
 	virtual void StartupModule() override;
 	virtual void ShutdownModule() override;
+	virtual bool SupportsDynamicReloading() override { return false; }
 
-	virtual bool SupportsDynamicReloading() override
-	{
-		return false;
-	}
+	void OnPostEngineInit();
+	void OnPreviewPlatformChanged();
+	void OnPreExit();
+	void OnPreviewFeatureLevelChanged(ERHIFeatureLevel::Type InPreviewFeatureLevel);
 
 	static inline FGroomEditor& Get()
 	{
@@ -56,6 +58,8 @@ private:
 	TSharedPtr<FSlateStyleSet> StyleSet;
 
 	FDelegateHandle TrackEditorBindingHandle;
+	FDelegateHandle PreviewPlatformChangedHandle;
+	FDelegateHandle PreviewFeatureLevelChangedHandle;
 
 public:
 
