Fix mismatching/incorrect primitive ID mode on mobile platforms for hair cards VF mesh batch.

#prefllight 659d57bd1673e5e979dea152

diff --git a/Engine/Plugins/Runtime/HairStrands/Source/HairStrandsCore/Private/GroomComponent.cpp b/Engine/Plugins/Runtime/HairStrands/Source/HairStrandsCore/Private/GroomComponent.cpp
--- a/Engine/Plugins/Runtime/HairStrands/Source/HairStrandsCore/Private/GroomComponent.cpp
+++ b/Engine/Plugins/Runtime/HairStrands/Source/HairStrandsCore/Private/GroomComponent.cpp
@@ -864,12 +864,14 @@ public:
 		const FVertexFactory* VertexFactory = nullptr;
 		FIndexBuffer* IndexBuffer = nullptr;
 		const FMaterialRenderProxy* MaterialRenderProxy = Debug_MaterialProxy;
+		const ERHIFeatureLevel::Type FeatureLevel = View->GetFeatureLevel();
 
 		uint32 NumPrimitive = 0;
 		uint32 HairVertexCount = 0;
 		uint32 MaxVertexIndex = 0;
 		bool bUseCulling = false;
 		bool bWireframe = false;
+		EPrimitiveIdMode PrimitiveIdMode = PrimID_Num;
 		if (GeometryType == EHairGeometryType::Meshes)
 		{
 			if (!Instance->Meshes.IsValid(IntLODIndex))
@@ -878,6 +880,7 @@ public:
 			}
 			VertexFactory = (FVertexFactory*)Instance->Meshes.LODs[IntLODIndex].GetVertexFactory();
 			check(VertexFactory);
+			PrimitiveIdMode = Instance->Meshes.LODs[IntLODIndex].GetVertexFactory()->GetPrimitiveIdMode(FeatureLevel);
 			HairVertexCount = Instance->Meshes.LODs[IntLODIndex].RestResource->GetPrimitiveCount() * 3;
 			MaxVertexIndex = HairVertexCount;
 			NumPrimitive = HairVertexCount / 3;
@@ -898,6 +901,7 @@ public:
 
 			VertexFactory = (FVertexFactory*)Instance->Cards.LODs[IntLODIndex].GetVertexFactory();
 			check(VertexFactory);
+			PrimitiveIdMode = Instance->Meshes.LODs[IntLODIndex].GetVertexFactory()->GetPrimitiveIdMode(FeatureLevel);
 			HairVertexCount = Instance->Cards.LODs[IntLODIndex].RestResource->GetPrimitiveCount() * 3;
 			MaxVertexIndex = HairVertexCount;
 			NumPrimitive = HairVertexCount / 3;
@@ -912,6 +916,7 @@ public:
 		else // if (GeometryType == EHairGeometryType::Strands)
 		{
 			VertexFactory = (FVertexFactory*)Instance->Strands.VertexFactory;
+			PrimitiveIdMode = Instance->Strands.VertexFactory->GetPrimitiveIdMode(FeatureLevel);
 			HairVertexCount = Instance->HairGroupPublicData->GetActiveStrandsPointCount();
 			MaxVertexIndex = HairVertexCount * HAIR_POINT_TO_VERTEX;
 			bUseCulling = Instance->Strands.bCullingEnable;
@@ -1007,7 +1012,7 @@ public:
 		//primtiveid is set to 0
 		BatchElement.FirstIndex = 0;
 		BatchElement.NumInstances = 1;
-		BatchElement.PrimitiveIdMode = GeometryType == EHairGeometryType::Strands ? PrimID_ForceZero : PrimID_DynamicPrimitiveShaderData;
+		BatchElement.PrimitiveIdMode = PrimitiveIdMode;
 		if (bUseCulling)
 		{
 			BatchElement.NumPrimitives = 0;
diff --git a/Engine/Plugins/Runtime/HairStrands/Source/HairStrandsCore/Private/HairCardsVertexFactory.cpp b/Engine/Plugins/Runtime/HairStrands/Source/HairStrandsCore/Private/HairCardsVertexFactory.cpp
--- a/Engine/Plugins/Runtime/HairStrands/Source/HairStrandsCore/Private/HairCardsVertexFactory.cpp
+++ b/Engine/Plugins/Runtime/HairStrands/Source/HairStrandsCore/Private/HairCardsVertexFactory.cpp
@@ -292,6 +292,11 @@ void FHairCardsVertexFactory::GetPSOPrecacheVertexFetchElements(EVertexInputStre
 	Elements.Add(FVertexElement(0, 0, VET_UInt, 13, 0, true));
 }
 
+EPrimitiveIdMode FHairCardsVertexFactory::GetPrimitiveIdMode(ERHIFeatureLevel::Type In) const
+{
+	return In > ERHIFeatureLevel::ES3_1 ? PrimID_DynamicPrimitiveShaderData : PrimID_ForceZero;
+}
+
 void FHairCardsVertexFactory::SetData(const FDataType& InData)
 {
 	Data = InData;
@@ -333,7 +338,13 @@ void FHairCardsVertexFactory::InitResources(FRHICommandListBase& RHICmdList)
 	// in case if we switch feature level at runtime.
 	FVertexDeclarationElementList Elements;
 	SetPrimitiveIdStreamIndex(CurrentFeatureLevel, EVertexInputStreamType::Default, -1);
-	AddPrimitiveIdStreamElement(EVertexInputStreamType::Default, Elements, 13, 0xff);
+
+	// Sanity check - When using mobile feature level, ensure we don't use primitive data as we set the PrimitiveIdStream to 0xFF for mobile
+	if (CurrentFeatureLevel <= ERHIFeatureLevel::ES3_1)
+	{
+		check(GetPrimitiveIdMode(CurrentFeatureLevel) == PrimID_ForceZero);
+	}
+	AddPrimitiveIdStreamElement(EVertexInputStreamType::Default, Elements, 13 /*AttributeIndex*/, 0xff /*AttributeIndex_Mobile*/);
 
 	// Note this is a local version of the VF's bSupportsManualVertexFetch, which take into account the feature level
 	const bool bManualFetch = SupportsManualVertexFetch(CurrentFeatureLevel);
diff --git a/Engine/Plugins/Runtime/HairStrands/Source/HairStrandsCore/Private/HairStrandsVertexFactory.cpp b/Engine/Plugins/Runtime/HairStrands/Source/HairStrandsCore/Private/HairStrandsVertexFactory.cpp
--- a/Engine/Plugins/Runtime/HairStrands/Source/HairStrandsCore/Private/HairStrandsVertexFactory.cpp
+++ b/Engine/Plugins/Runtime/HairStrands/Source/HairStrandsCore/Private/HairStrandsVertexFactory.cpp
@@ -262,6 +262,11 @@ void FHairStrandsVertexFactory::GetPSOPrecacheVertexFetchElements(EVertexInputSt
 #endif
 }
 
+EPrimitiveIdMode FHairStrandsVertexFactory::GetPrimitiveIdMode(ERHIFeatureLevel::Type In) const
+{
+	return PrimID_ForceZero;
+}
+
 IMPLEMENT_VERTEX_FACTORY_PARAMETER_TYPE(FHairStrandsVertexFactory, SF_Vertex,		FHairStrandsVertexFactoryShaderParameters);
 IMPLEMENT_VERTEX_FACTORY_PARAMETER_TYPE(FHairStrandsVertexFactory, SF_Pixel,		FHairStrandsVertexFactoryShaderParameters);
 IMPLEMENT_VERTEX_FACTORY_PARAMETER_TYPE(FHairStrandsVertexFactory, SF_Compute,		FHairStrandsVertexFactoryShaderParameters);
diff --git a/Engine/Plugins/Runtime/HairStrands/Source/HairStrandsCore/Public/HairCardsVertexFactory.h b/Engine/Plugins/Runtime/HairStrands/Source/HairStrandsCore/Public/HairCardsVertexFactory.h
--- a/Engine/Plugins/Runtime/HairStrands/Source/HairStrandsCore/Public/HairCardsVertexFactory.h
+++ b/Engine/Plugins/Runtime/HairStrands/Source/HairStrandsCore/Public/HairCardsVertexFactory.h
@@ -10,10 +10,10 @@
 #include "HairCardsDatas.h"
 #include "HairStrandsInterface.h"
 #include "PrimitiveSceneProxy.h"
+#include "MeshBatch.h"
 
 class FMaterial;
 class FSceneView;
-struct FMeshBatchElement;
 
 /**
  * A vertex factory which simply transforms explicit vertex attributes from local to world space.
@@ -39,6 +39,9 @@ public:
 	static void ValidateCompiledResult(const FVertexFactoryType* Type, EShaderPlatform Platform, const FShaderParameterMap& ParameterMap, TArray<FString>& OutErrors);
 	static void GetPSOPrecacheVertexFetchElements(EVertexInputStreamType VertexInputStreamType, FVertexDeclarationElementList& Elements);
 
+	// Return the primitive id supported by the VF
+	EPrimitiveIdMode GetPrimitiveIdMode(ERHIFeatureLevel::Type In) const;
+
 	/**
 	 * An implementation of the interface used by TSynchronizedResource to update the resource with new data from the game thread.
 	 */
diff --git a/Engine/Plugins/Runtime/HairStrands/Source/HairStrandsCore/Public/HairStrandsVertexFactory.h b/Engine/Plugins/Runtime/HairStrands/Source/HairStrandsCore/Public/HairStrandsVertexFactory.h
--- a/Engine/Plugins/Runtime/HairStrands/Source/HairStrandsCore/Public/HairStrandsVertexFactory.h
+++ b/Engine/Plugins/Runtime/HairStrands/Source/HairStrandsCore/Public/HairStrandsVertexFactory.h
@@ -10,10 +10,10 @@
 #include "HairStrandsDatas.h"
 #include "HairStrandsInterface.h"
 #include "PrimitiveSceneProxy.h"
+#include "MeshBatch.h"
 
 class FMaterial;
 class FSceneView;
-struct FMeshBatchElement;
 struct FHairGroupInstance;
 
 /**
@@ -44,6 +44,9 @@ public:
 	static void ValidateCompiledResult(const FVertexFactoryType* Type, EShaderPlatform Platform, const FShaderParameterMap& ParameterMap, TArray<FString>& OutErrors);
 	static void GetPSOPrecacheVertexFetchElements(EVertexInputStreamType VertexInputStreamType, FVertexDeclarationElementList& Elements);
 
+	// Return the primitive id supported by the VF
+	EPrimitiveIdMode GetPrimitiveIdMode(ERHIFeatureLevel::Type In) const;
+
 	/**
 	 * An implementation of the interface used by TSynchronizedResource to update the resource with new data from the game thread.
 	 */
