Fix incorrect function call into a check() causing the function to not be called in shipping build.

diff --git a/Engine/Plugins/Runtime/HairStrands/Source/HairStrandsCore/Private/GroomResources.cpp b/Engine/Plugins/Runtime/HairStrands/Source/HairStrandsCore/Private/GroomResources.cpp
--- a/Engine/Plugins/Runtime/HairStrands/Source/HairStrandsCore/Private/GroomResources.cpp
+++ b/Engine/Plugins/Runtime/HairStrands/Source/HairStrandsCore/Private/GroomResources.cpp
@@ -678,7 +678,8 @@ void FHairCommonResource::InitRHI(FRHICommandListBase&)
 
 	// 3. Loaded the data
 	const int32 DummyLODIndex = -1;
-	check(InternalGetOrRequestData(HAIR_MAX_NUM_CURVE_PER_GROUP, HAIR_MAX_NUM_POINT_PER_GROUP, DummyLODIndex));
+	const bool bReady = InternalGetOrRequestData(HAIR_MAX_NUM_CURVE_PER_GROUP, HAIR_MAX_NUM_POINT_PER_GROUP, DummyLODIndex);
+	check(bReady);
 
 	// 4. Allocate the resource, and update the data
 	if (bUseRenderGraph)
