Fix crash on platform not supporting strands when rendering grooms with bindings.

diff --git a/Engine/Plugins/Runtime/HairStrands/Source/HairStrandsCore/Private/GroomManager.cpp b/Engine/Plugins/Runtime/HairStrands/Source/HairStrandsCore/Private/GroomManager.cpp
--- a/Engine/Plugins/Runtime/HairStrands/Source/HairStrandsCore/Private/GroomManager.cpp
+++ b/Engine/Plugins/Runtime/HairStrands/Source/HairStrandsCore/Private/GroomManager.cpp
@@ -656,7 +656,7 @@ static void RunHairBindingSurfaceUpdate(
 		// 1. Guides
 		{
 			FHairStrandsRootBulkData* RootBulkData = nullptr;
-			if (bStrandSupported && Instance->Guides.IsValid())
+			if (Instance->Guides.IsValid())
 			{	
 				const bool bNeedSurfaceUpdate = 
 					Instance->BindingType == EHairBindingType::Skinning &&
