Reduce cooked Groom size by 35%.

This CL compress hair interpolation data by splitting per curve and per point data needed data.

diff --git a/Engine/Plugins/Runtime/HairStrands/Shaders/Private/DeformerDataInterfaceGroom.ush b/Engine/Plugins/Runtime/HairStrands/Shaders/Private/DeformerDataInterfaceGroom.ush
--- a/Engine/Plugins/Runtime/HairStrands/Shaders/Private/DeformerDataInterfaceGroom.ush
+++ b/Engine/Plugins/Runtime/HairStrands/Shaders/Private/DeformerDataInterfaceGroom.ush
@@ -120,15 +120,14 @@ float ReadCurveNumPoint_{DataInterfaceName}(uint CurveIndex)
 
 uint ReadGuideIndex_{DataInterfaceName}(uint ControlPointIndex)
 {
-	FGuideData GuideData;
-	if ({DataInterfaceName}_bSingleGuide)
-	{
-		GuideData = UnpackGuideData({DataInterfaceName}_InterpolationBuffer.Load(ControlPointIndex * HAIR_INTERPOLATION_1GUIDE_STRIDE));
-	}
-	else
-	{
-		GuideData = UnpackGuideData({DataInterfaceName}_InterpolationBuffer.Load4(ControlPointIndex * HAIR_INTERPOLATION_3GUIDE_STRIDE));
-	}
+	// Find the rendering curve index
+	const uint CurveIndex = {DataInterfaceName}_PointToCurveBuffer.Load(ControlPointIndex);
+
+	// Fint the curve and point interpolation data
+	const FGuideCurveData CurveData = UnpackGuideCurveData({DataInterfaceName}_CurveInterpolationBuffer, CurveIndex, {DataInterfaceName}_bSingleGuide);
+	const FGuidePointData PointData = UnpackGuidePointData({DataInterfaceName}_PointInterpolationBuffer, ControlPointIndex, {DataInterfaceName}_bSingleGuide);
 
-	return GuideData.VertexIndices.x;
+	// Build the sim point index
+	const uint2 GuidePointIndices = CurveData.RootPointIndices + PointData.PointIndices;
+	return GuidePointIndices.x;
 }
diff --git a/Engine/Plugins/Runtime/HairStrands/Source/HairStrandsCore/Private/Deformer/DeformerDataInterfaceGroom.cpp b/Engine/Plugins/Runtime/HairStrands/Source/HairStrandsCore/Private/Deformer/DeformerDataInterfaceGroom.cpp
--- a/Engine/Plugins/Runtime/HairStrands/Source/HairStrandsCore/Private/Deformer/DeformerDataInterfaceGroom.cpp
+++ b/Engine/Plugins/Runtime/HairStrands/Source/HairStrandsCore/Private/Deformer/DeformerDataInterfaceGroom.cpp
@@ -226,7 +226,8 @@ void FOptimusGroomDataProviderProxy::AllocateResources(FRDGBuilder& GraphBuilder
 
 			{
 				FHairStrandsInstanceInterpolationParameters& R = Interpolations.AddDefaulted_GetRef();
-				R.InterpolationBuffer = Instance->Strands.InterpolationResource ? RegisterAsSRV(GraphBuilder, Instance->Strands.InterpolationResource->InterpolationBuffer) : nullptr;
+				R.CurveInterpolationBuffer = Instance->Strands.InterpolationResource ? RegisterAsSRV(GraphBuilder, Instance->Strands.InterpolationResource->CurveInterpolationBuffer) : nullptr;
+				R.PointInterpolationBuffer = Instance->Strands.InterpolationResource ? RegisterAsSRV(GraphBuilder, Instance->Strands.InterpolationResource->PointInterpolationBuffer) : nullptr;
 			}
 
 			if (!FallbackSRV)
@@ -257,7 +258,8 @@ void FOptimusGroomDataProviderProxy::GatherDispatchData(FDispatchData const& InD
 		Parameters.Resources.PointAttributeBuffer	= FallbackSRV;
 		Parameters.Resources.PointToCurveBuffer		= FallbackSRV;
 		Parameters.Resources.CurveBuffer			= FallbackSRV;
-		Parameters.Interpolation.InterpolationBuffer= FallbackSRV;
+		Parameters.Interpolation.CurveInterpolationBuffer = FallbackSRV;
+		Parameters.Interpolation.PointInterpolationBuffer = FallbackSRV;
 
 		if (FHairGroupInstance* Instance = GroomComponent->GetGroupInstance(InvocationIndex))
 		{
@@ -268,7 +270,7 @@ void FOptimusGroomDataProviderProxy::GatherDispatchData(FDispatchData const& InD
 			{
 				Parameters.Resources = Resources[InvocationIndex];
 			}
-			if (Interpolations[InvocationIndex].InterpolationBuffer != nullptr)
+			if (Interpolations[InvocationIndex].CurveInterpolationBuffer != nullptr)
 			{
 				Parameters.Interpolation = Interpolations[InvocationIndex];
 			}
diff --git a/Engine/Plugins/Runtime/HairStrands/Source/HairStrandsCore/Private/GroomBuilder.cpp b/Engine/Plugins/Runtime/HairStrands/Source/HairStrandsCore/Private/GroomBuilder.cpp
--- a/Engine/Plugins/Runtime/HairStrands/Source/HairStrandsCore/Private/GroomBuilder.cpp
+++ b/Engine/Plugins/Runtime/HairStrands/Source/HairStrandsCore/Private/GroomBuilder.cpp
@@ -40,7 +40,18 @@ static FAutoConsoleVariableRef CVarHairGroupIndexBuilder_MaxVoxelResolution(TEXT
 
 FString FGroomBuilder::GetVersion()
 {
-	return TEXT("v9a");
+	return TEXT("v10a");
+}
+
+// For debug purpose
+template<typename T>
+void ReportSize(const TCHAR* InName, const TArray<T>& In)
+{
+#if 0
+	const uint32 TotalSize = In.GetTypeSize() * In.Num();
+	const uint32 TotalKBytes = TotalSize / 1024;
+	UE_LOG(LogGroomBuilder, Log, TEXT("Size: %10d Kbytes - Num: %8d - Name: %s"), TotalKBytes, In.Num(), InName);
+#endif
 }
 
 namespace FHairStrandsDecimation
@@ -470,6 +481,7 @@ namespace HairStrandsBuilder
 			}
 
 			CopyToBulkData<FHairStrandsAttributeFormat>(OutBulkData.Data.CurveAttributes, OutPackedAttributes);
+			ReportSize(TEXT("Curve Attributes"), OutPackedAttributes);
 		}
 
 		// Concatenate all point-attributes
@@ -548,19 +560,24 @@ namespace HairStrandsBuilder
 			{
 				OutBulkData.Header.Flags |= FHairStrandsBulkData::DataFlags_HasPointAttribute;
 				CopyToBulkData<FHairStrandsAttributeFormat>(OutBulkData.Data.PointAttributes, OutPackedAttributes);
+				ReportSize(TEXT("Point Attributes"), OutPackedAttributes);
 			}
 		}
 
 		CopyToBulkData<FHairStrandsPositionFormat>(OutBulkData.Data.Positions, OutPackedPositions);
+		ReportSize(TEXT("Positions"), OutPackedPositions);
 		CopyToBulkData<FHairStrandsCurveFormat>(OutBulkData.Data.Curves, OutPackedCurves);
+		ReportSize(TEXT("PackedCurves"), OutPackedCurves);
 		if (bUse16bitsCurveIndex)
 		{
 			OutBulkData.Header.Flags |= FHairStrandsBulkData::DataFlags_Has16bitsCurveIndex;
 			CopyToBulkData<FHairStrandsPointToCurveFormat16>(OutBulkData.Data.PointToCurve, OutPointToCurve16);
+			ReportSize(TEXT("PointToCurve"), OutPointToCurve16);
 		}
 		else
 		{
 			CopyToBulkData<FHairStrandsPointToCurveFormat32>(OutBulkData.Data.PointToCurve, OutPointToCurve32);
+			ReportSize(TEXT("PointToCurve"), OutPointToCurve32);
 		}
 
 		// Build curve to point count mapping for runtime CLOD
@@ -711,14 +728,14 @@ namespace HairInterpolationBuilder
 
 	struct FMetrics
 	{
-		static const uint32 Count = 3;
+		static const uint32 Count = HAIR_INTERPOLATION_MAX_GUIDE_COUNT;
 		float KMinMetrics[Count];
 		int32 KClosestGuideIndices[Count];
 	};
 
 	struct FClosestGuides
 	{
-		static const uint32 Count = 3;
+		static const uint32 Count = HAIR_INTERPOLATION_MAX_GUIDE_COUNT;
 		int32 Indices[Count];
 	};
 
@@ -732,6 +749,7 @@ namespace HairInterpolationBuilder
 	{
 		FMetrics Metric = InMetric;
 		check(Metric.KClosestGuideIndices[0] >= 0);
+		static_assert(HAIR_INTERPOLATION_MAX_GUIDE_COUNT == 2);
 
 		// If some indices are invalid (for instance, found a valid single guide, fill in the rest with the valid ones)
 		if (Metric.KClosestGuideIndices[1] < 0)
@@ -739,57 +757,35 @@ namespace HairInterpolationBuilder
 			Metric.KClosestGuideIndices[1] = Metric.KClosestGuideIndices[0];
 			Metric.KMinMetrics[1] = Metric.KMinMetrics[0];
 		}
-		if (Metric.KClosestGuideIndices[2] < 0)
-		{
-			Metric.KClosestGuideIndices[2] = Metric.KClosestGuideIndices[1];
-			Metric.KMinMetrics[2] = Metric.KMinMetrics[1];
-		}
 
 		uint32 RandIndex0 = 0;
 		uint32 RandIndex1 = 1;
-		uint32 RandIndex2 = 2;
 		if (bRandomizeInterpolation)
 		{
 			// This randomization makes certain strands being affected by 1, 2, or 3 guides
 			RandIndex0 = RandomIndices[0];
 			RandIndex1 = RandomIndices[1];
-			RandIndex2 = RandomIndices[2];
 		}
 
 		ClosestGuides.Indices[0] = Metric.KClosestGuideIndices[RandIndex0];
 		ClosestGuides.Indices[1] = Metric.KClosestGuideIndices[RandIndex1];
-		ClosestGuides.Indices[2] = Metric.KClosestGuideIndices[RandIndex2];
 
 		if (bUseUniqueGuide)
 		{
 			ClosestGuides.Indices[1] = Metric.KClosestGuideIndices[RandIndex0];
-			ClosestGuides.Indices[2] = Metric.KClosestGuideIndices[RandIndex0];
 			RandIndex1 = RandIndex0;
-			RandIndex2 = RandIndex0;
 		}
 
-
 		float MinMetrics[FMetrics::Count];
 		MinMetrics[0] = Metric.KMinMetrics[RandIndex0];
 		MinMetrics[1] = Metric.KMinMetrics[RandIndex1];
-		MinMetrics[2] = Metric.KMinMetrics[RandIndex2];
 
-
-		while (!(MinMetrics[0] <= MinMetrics[1] && MinMetrics[1] <= MinMetrics[2]))
+		if (MinMetrics[0] > MinMetrics[1])
 		{
-			if (MinMetrics[0] > MinMetrics[1])
-			{
-				SwapValue(MinMetrics[0], MinMetrics[1]);
-				SwapValue(ClosestGuides.Indices[0], ClosestGuides.Indices[1]);
-			}
-
-			if (MinMetrics[1] > MinMetrics[2])
-			{
-				SwapValue(MinMetrics[1], MinMetrics[2]);
-				SwapValue(ClosestGuides.Indices[1], ClosestGuides.Indices[2]);
-			}
+			SwapValue(MinMetrics[0], MinMetrics[1]);
+			SwapValue(ClosestGuides.Indices[0], ClosestGuides.Indices[1]);
 		}
-
+		
 		// If there less than 3 valid guides, fill the rest with existing valid guides
 		// This can happen due to the normal-orientation based rejection above
 		if (ClosestGuides.Indices[1] < 0)
@@ -797,14 +793,8 @@ namespace HairInterpolationBuilder
 			ClosestGuides.Indices[1] = ClosestGuides.Indices[0];
 			MinMetrics[1] = MinMetrics[0];
 		}
-		if (ClosestGuides.Indices[2] < 0)
-		{
-			ClosestGuides.Indices[2] = ClosestGuides.Indices[1];
-			MinMetrics[2] = MinMetrics[1];
-		}
 
 		check(MinMetrics[0] <= MinMetrics[1]);
-		check(MinMetrics[1] <= MinMetrics[2]);
 	}
 
 	// Simple closest distance metric
@@ -1026,7 +1016,6 @@ namespace HairInterpolationBuilder
 
 			check(ClosestGuides.Indices[0] >= 0);
 			check(ClosestGuides.Indices[1] >= 0);
-			check(ClosestGuides.Indices[2] >= 0);
 
 			return ClosestGuides;
 		}
@@ -1148,7 +1137,6 @@ namespace HairInterpolationBuilder
 
 			check(ClosestGuides.Indices[0] >= 0);
 			check(ClosestGuides.Indices[1] >= 0);
-			check(ClosestGuides.Indices[2] >= 0);
 
 			return ClosestGuides;
 		}
@@ -1183,7 +1171,6 @@ namespace HairInterpolationBuilder
 
 		check(ClosestGuides.Indices[0] >= 0);
 		check(ClosestGuides.Indices[1] >= 0);
-		check(ClosestGuides.Indices[2] >= 0);
 
 		return ClosestGuides;
 	}
@@ -1303,7 +1290,7 @@ namespace HairInterpolationBuilder
 	{
 		TRACE_CPUPROFILER_EVENT_SCOPE(HairInterpolationBuilder::BuildInterpolationData);
 
-		InterpolationData.SetNum(RenStrandsData.GetNumPoints());
+		InterpolationData.SetNum(RenStrandsData.GetNumCurves(), RenStrandsData.GetNumPoints());
 		InterpolationData.bUseUniqueGuide = Settings.bUseUniqueGuide;
 
 		typedef TArray<FHairRoot> FRoots;
@@ -1395,13 +1382,26 @@ namespace HairInterpolationBuilder
 			const uint32 RendPointCount	= RenStrandsData.StrandsCurves.CurvesCount[RenCurveIndex];
 			const uint32 RenOffset		= RenStrandsData.StrandsCurves.CurvesOffset[RenCurveIndex];
 			const FVector3f& RenPointPositionRoot = RenStrandsData.StrandsPoints.PointsPosition[RenOffset];
+
+			// Set per-curve index
+			float LocalCurveSimWeights[FClosestGuides::Count]; 
+			for (uint32 KIndex = 0; KIndex < FClosestGuides::Count; ++KIndex)
+			{
+				const uint32 SimCurveIndex = ClosestGuides.Indices[KIndex];
+				InterpolationData.CurveSimIndices[RenCurveIndex][KIndex] = SimCurveIndex;
+				LocalCurveSimWeights[KIndex] = 0;
+
+				const uint32 SimOffset = SimStrandsData.StrandsCurves.CurvesOffset[SimCurveIndex];
+				InterpolationData.CurveSimRootPointIndex[RenCurveIndex][KIndex] = SimOffset;
+			}
+
+			// Set per-point 'local' point index and lerp value
 			for (uint32 RenPointIndex = 0; RenPointIndex < RendPointCount; ++RenPointIndex)
 			{
 				const uint32 PointGlobalIndex = RenPointIndex + RenOffset;
 				const FVector3f& RenPointPosition = RenStrandsData.StrandsPoints.PointsPosition[PointGlobalIndex];
 				const float RenPointDistance = RenStrandsData.StrandsPoints.PointsCoordU[PointGlobalIndex] * RenStrandsData.StrandsCurves.CurvesLength[RenCurveIndex];
 
-				float TotalWeight = 0;
 				for (uint32 KIndex = 0; KIndex < FClosestGuides::Count; ++KIndex)
 				{
 					// Find the closest vertex on the guide which matches the strand vertex distance along its curve
@@ -1417,10 +1417,9 @@ namespace HairInterpolationBuilder
 						const FVector3f& SimPointPosition1 = SimStrandsData.StrandsPoints.PointsPosition[Desc.Index1 + SimOffset];
 						const float Weight = 1.0f / FMath::Max(MinWeightDistance, FVector3f::Distance(RenPointPosition, FMath::Lerp(SimPointPosition0, SimPointPosition1, Desc.T)));
 
-						InterpolationData.PointsSimCurvesIndex[PointGlobalIndex][KIndex] = SimCurveIndex;
-						InterpolationData.PointsSimCurvesVertexIndex[PointGlobalIndex][KIndex] = Desc.Index0 + SimOffset;
-						InterpolationData.PointsSimCurvesVertexLerp[PointGlobalIndex][KIndex] = Desc.T;
-						InterpolationData.PointsSimCurvesVertexWeights[PointGlobalIndex][KIndex] = Weight;
+						InterpolationData.PointSimIndices[PointGlobalIndex][KIndex] = Desc.Index0/* + SimOffset*/;
+						InterpolationData.PointSimLerps[PointGlobalIndex][KIndex] = Desc.T;
+						LocalCurveSimWeights[KIndex] += Weight;
 					}
 
 					// Use only the root as a *constant* weight for deformation along each vertex
@@ -1434,10 +1433,9 @@ namespace HairInterpolationBuilder
 						const float Weight = 1.0f / FMath::Max(MinWeightDistance, FVector3f::Distance(RenRootPointPosition, SimRootPointPosition));
 						const FVertexInterpolationDesc Desc = FindMatchingVertex(RenPointDistance, SimStrandsData, SimCurveIndex);
 
-						InterpolationData.PointsSimCurvesIndex[PointGlobalIndex][KIndex] = SimCurveIndex;
-						InterpolationData.PointsSimCurvesVertexIndex[PointGlobalIndex][KIndex] = Desc.Index0 + SimOffset;
-						InterpolationData.PointsSimCurvesVertexLerp[PointGlobalIndex][KIndex] = Desc.T;
-						InterpolationData.PointsSimCurvesVertexWeights[PointGlobalIndex][KIndex] = Weight;
+						InterpolationData.PointSimIndices[PointGlobalIndex][KIndex] = Desc.Index0/* + SimOffset*/;
+						InterpolationData.PointSimLerps[PointGlobalIndex][KIndex] = Desc.T;
+						LocalCurveSimWeights[KIndex] += Weight;
 					}
 
 					// Use the *same vertex index* to match guide vertex with strand vertex
@@ -1450,148 +1448,128 @@ namespace HairInterpolationBuilder
 						const FVector3f& SimPointPosition = SimStrandsData.StrandsPoints.PointsPosition[SimPointIndex + SimOffset];
 						const float Weight = 1.0f / FMath::Max(MinWeightDistance, FVector3f::Distance(RenPointPosition, SimPointPosition));
 
-						InterpolationData.PointsSimCurvesIndex[PointGlobalIndex][KIndex] = SimCurveIndex;
-						InterpolationData.PointsSimCurvesVertexIndex[PointGlobalIndex][KIndex] = SimPointIndex + SimOffset;
-						InterpolationData.PointsSimCurvesVertexLerp[PointGlobalIndex][KIndex] = 1;
-						InterpolationData.PointsSimCurvesVertexWeights[PointGlobalIndex][KIndex] = Weight;
+						InterpolationData.PointSimIndices[PointGlobalIndex][KIndex] = SimPointIndex/* + SimOffset*/;
+						InterpolationData.PointSimLerps[PointGlobalIndex][KIndex] = 1;
+						LocalCurveSimWeights[KIndex] += Weight;
 					}
-
-					TotalWeight += InterpolationData.PointsSimCurvesVertexWeights[PointGlobalIndex][KIndex];
 				}
 
-				for (int32 KIndex = 0; KIndex < FClosestGuides::Count; ++KIndex)
-				{
-					InterpolationData.PointsSimCurvesVertexWeights[PointGlobalIndex][KIndex] /= TotalWeight;
-				}
+			}
+
+			// Set per-curve weights
+			float TotalWeight = 0;
+			for (int32 KIndex = 0; KIndex < FClosestGuides::Count; ++KIndex)
+			{
+				TotalWeight += LocalCurveSimWeights[KIndex];
+			}
+			for (int32 KIndex = 0; KIndex < FClosestGuides::Count; ++KIndex)
+			{
+				InterpolationData.CurveSimWeights[RenCurveIndex][KIndex] = LocalCurveSimWeights[KIndex] / FMath::Max(0.001f, TotalWeight);
 			}
 		});
 	}
 
 	/** Build data for interpolation between simulation and rendering */
-	void BuildRenderData(
+	void BuildInterplationBulkData(
 		const FHairStrandsDatas& SimDatas, 
 		const FHairStrandsInterpolationDatas& HairInterpolation, 
 		FHairStrandsInterpolationBulkData& OutBulkData)
 	{
 		OutBulkData.Reset();
 
-		TRACE_CPUPROFILER_EVENT_SCOPE(HairInterpolationBuilder::BuildRenderData);
+		TRACE_CPUPROFILER_EVENT_SCOPE(HairInterpolationBuilder::BuildInterplationBulkData);
 
-		const uint32 PointCount = HairInterpolation.Num();
-		if (PointCount == 0)
+		const uint32 PointCount = HairInterpolation.GetPointCount();
+		const uint32 CurveCount = HairInterpolation.GetCurveCount();
+		if (PointCount == 0 || CurveCount == 0)
+		{
 			return;
-
-		auto LowerPart = [](uint32 Index) -> uint16 { return uint16(Index & 0xFFFF); };
-		auto UpperPart = [](uint32 Index) -> uint8  { return uint8((Index >> 16) & 0xFF); };
+		}
 
 		OutBulkData.Header.Flags = FHairStrandsInterpolationBulkData::DataFlags_HasData;
 		OutBulkData.Header.PointCount = PointCount;
-		OutBulkData.Header.Strides.SimRootPointIndexStride = FHairStrandsRootIndexFormat::SizeInByte;
+		OutBulkData.Header.CurveCount = CurveCount;
+		OutBulkData.Header.Strides.CurveInterpolationStride = 0;
+		OutBulkData.Header.Strides.PointInterpolationStride = 0;
 
-		if (HairInterpolation.bUseUniqueGuide)
+		struct FHairInterpolationCurve
 		{
-			struct FHairInterpolation1GuideVertex
-			{
-				// Guide's vertex index are stored onto 24bits: 1) VertexGuideIndex0 is the lower part, 2) VertexGuideIndex1 is the upper part
-				uint16 VertexGuideIndex0;
-				uint8  VertexGuideIndex1;
-				uint8  VertexLerp;
-			};
-			static_assert(sizeof(FHairInterpolation1GuideVertex) == HAIR_INTERPOLATION_1GUIDE_STRIDE);
-			static_assert(sizeof(FHairStrandsInterpolationFormat::Type) == 4u);
+			uint32 CurveIndex : 24;
+			uint32 CurveWeight : 8;
+			uint32 PointIndex;
+		};
 
-			OutBulkData.Header.Flags |= FHairStrandsInterpolationBulkData::DataFlags_HasSingleGuideData;
-			OutBulkData.Header.Strides.InterpolationStride = sizeof(FHairInterpolation1GuideVertex);
+		struct FHairInterpolationPoint
+		{
+			uint8 VertexIndex;
+			uint8 VertexLerp;
+		};
 
-			TArray<FHairStrandsInterpolationFormat::Type> OutPointsInterpolation;
-			OutPointsInterpolation.SetNum(PointCount);
-			FHairInterpolation1GuideVertex* Out = (FHairInterpolation1GuideVertex*)OutPointsInterpolation.GetData();
-			for (uint32 PointIndex = 0; PointIndex < PointCount; ++PointIndex)
-			{
-				const FIntVector& Indices 	= HairInterpolation.PointsSimCurvesVertexIndex[PointIndex];
-				const FVector3f& Weights 	= HairInterpolation.PointsSimCurvesVertexWeights[PointIndex];
-				const FVector3f& S 			= HairInterpolation.PointsSimCurvesVertexLerp[PointIndex];
+		static_assert(sizeof(FHairInterpolationCurve) == HAIR_INTERPOLATION_CURVE_STRIDE);
+		static_assert(sizeof(FHairInterpolationPoint) == HAIR_INTERPOLATION_POINT_STRIDE);
+		static_assert(HAIR_INTERPOLATION_MAX_GUIDE_COUNT == 2u);
 
-				FHairInterpolation1GuideVertex& OutInterp = Out[PointIndex];
-				OutInterp.VertexGuideIndex0 = LowerPart(Indices[0]);
-				OutInterp.VertexGuideIndex1 = UpperPart(Indices[0]);
-				OutInterp.VertexLerp		 = FMath::Clamp(S[0] * 255.f, 0, 255);
-			}
+		// Header
+		const uint32 KCount = HairInterpolation.bUseUniqueGuide ? 1u : HAIR_INTERPOLATION_MAX_GUIDE_COUNT;
+		if (HairInterpolation.bUseUniqueGuide)
+		{
 
-			HairStrandsBuilder::CopyToBulkData<FHairStrandsInterpolationFormat>(OutBulkData.Data.Interpolation, OutPointsInterpolation);
+			OutBulkData.Header.Flags |= FHairStrandsInterpolationBulkData::DataFlags_HasSingleGuideData;
+			OutBulkData.Header.Strides.CurveInterpolationStride = sizeof(FHairInterpolationCurve);
+			OutBulkData.Header.Strides.PointInterpolationStride = sizeof(FHairInterpolationPoint);
 		}
 		else
 		{
-			struct FHairInterpolation3GuidesVertex
-			{	
-				typedef FUintVector4 BulkType;
-
-				// 128 bits
-				uint16 Index0;       // X:16
-				uint16 Index1;       // X:16
-				uint16 Index2;       // Y:16
-				uint8  VertexWeight0;// Y:8
-				uint8  VertexWeight1;// Y:8
-				uint8  VertexIndex0; // Z:8
-				uint8  VertexIndex1; // Z:8
-				uint8  VertexIndex2; // Z:8
-				uint8  Pad0;         // Z:8
-				uint8  VertexLerp0;  // W:8
-				uint8  VertexLerp1;  // W:8
-				uint8  VertexLerp2;  // W:8
-				uint8  Pad1;         // W:8
-			};
-			static_assert(sizeof(FHairInterpolation3GuidesVertex) == HAIR_INTERPOLATION_3GUIDE_STRIDE);
-			static_assert(sizeof(FHairStrandsInterpolationFormat::Type) == 4u);
-
-			OutBulkData.Header.Strides.InterpolationStride = sizeof(FHairInterpolation3GuidesVertex);
 
-			TArray<FHairStrandsInterpolationFormat::Type> OutPointsInterpolation;
-			OutPointsInterpolation.SetNum(PointCount * 4);
-			FHairInterpolation3GuidesVertex* Out = (FHairInterpolation3GuidesVertex*)OutPointsInterpolation.GetData();
-
-			for (uint32 PointIndex = 0; PointIndex < PointCount; ++PointIndex)
+			OutBulkData.Header.Strides.CurveInterpolationStride = sizeof(FHairInterpolationCurve) * HAIR_INTERPOLATION_MAX_GUIDE_COUNT;
+			OutBulkData.Header.Strides.PointInterpolationStride = sizeof(FHairInterpolationPoint) * HAIR_INTERPOLATION_MAX_GUIDE_COUNT;
+		}
+		
+		// Data
+		{
+			// Fill curve data
+			TArray<FHairStrandsInterpolationFormat::Type> OutInterpolationCurves;
 			{
-				const FIntVector& Indices 	= HairInterpolation.PointsSimCurvesVertexIndex[PointIndex];
-				const FVector3f& Weights 	= (FVector3f)HairInterpolation.PointsSimCurvesVertexWeights[PointIndex];
-				const FVector3f& S 			= (FVector3f)HairInterpolation.PointsSimCurvesVertexLerp[PointIndex];
+				const uint32 uint32Count = FMath::DivideAndRoundUp(CurveCount * OutBulkData.Header.Strides.CurveInterpolationStride, 4u);
+				OutInterpolationCurves.Init(0u, uint32Count);
 
-				FHairInterpolation3GuidesVertex& OutInterp = Out[PointIndex];
-				OutInterp.Index0 		= LowerPart(Indices[0]);
-				OutInterp.Index1 		= LowerPart(Indices[1]);
-				OutInterp.Index2 		= LowerPart(Indices[2]);
-				OutInterp.VertexWeight0 = FMath::Clamp(Weights[0] * 255.f, 0, 255);
-				OutInterp.VertexWeight1 = FMath::Clamp(Weights[1] * 255.f, 0, 255);
-				OutInterp.VertexIndex0  = UpperPart(Indices[0]);
-				OutInterp.VertexIndex1  = UpperPart(Indices[1]);
-				OutInterp.VertexIndex2  = UpperPart(Indices[2]);
-				OutInterp.Pad0			= 0;
-				OutInterp.VertexLerp0   = FMath::Clamp(S[0] * 255.f, 0, 255);
-				OutInterp.VertexLerp1   = FMath::Clamp(S[1] * 255.f, 0, 255);
-				OutInterp.VertexLerp2   = FMath::Clamp(S[2] * 255.f, 0, 255);
-				OutInterp.Pad1			= 0;
+				FHairInterpolationCurve* Out = (FHairInterpolationCurve*)OutInterpolationCurves.GetData();
+				for (uint32 CurveIndex = 0; CurveIndex < CurveCount; ++CurveIndex)
+				{
+					for (uint32 KIndex = 0; KIndex < KCount; ++KIndex)
+					{
+						FHairInterpolationCurve& OutInterp = Out[CurveIndex*KCount + KIndex];
+						OutInterp.CurveIndex  = HairInterpolation.CurveSimIndices[CurveIndex][KIndex];
+						OutInterp.CurveWeight = FMath::Clamp(HairInterpolation.CurveSimWeights[CurveIndex][KIndex] * 255.f, 0, 255);
+						OutInterp.PointIndex  = HairInterpolation.CurveSimRootPointIndex[CurveIndex][KIndex];
+					}
+				}
 			}
-			HairStrandsBuilder::CopyToBulkData<FHairStrandsInterpolationFormat>(OutBulkData.Data.Interpolation, OutPointsInterpolation);
-		}
 
-		// Map point to root point
-		{
-			OutBulkData.Header.SimPointCount = SimDatas.GetNumPoints();
-			const uint32 SimCurveCount = SimDatas.GetNumCurves();
-
-			TArray<FHairStrandsRootIndexFormat::Type> SimRootPointIndex;
-			SimRootPointIndex.SetNum(OutBulkData.Header.SimPointCount);
-			for (uint32 CurveIndex = 0; CurveIndex < SimCurveCount; ++CurveIndex)
+			// Fill in point data
+			TArray<FHairStrandsInterpolationFormat::Type> OutInterpolationPoints;
 			{
-				const uint16 SimPointCount = SimDatas.StrandsCurves.CurvesCount[CurveIndex];
-				const uint32 SimPointOffset = SimDatas.StrandsCurves.CurvesOffset[CurveIndex];
-				for (uint32 PointIndex = 0; PointIndex < SimPointCount; ++PointIndex)
+				const uint32 uint32Count = FMath::DivideAndRoundUp(PointCount * OutBulkData.Header.Strides.PointInterpolationStride, 4u);
+				OutInterpolationPoints.Init(0u, uint32Count);
+
+				FHairInterpolationPoint* Out = (FHairInterpolationPoint*)OutInterpolationPoints.GetData();
+				for (uint32 PointIndex = 0; PointIndex < PointCount; ++PointIndex)
 				{
-					SimRootPointIndex[PointIndex + SimPointOffset] = SimPointOffset;
+					const FIntVector2& Indices 	= HairInterpolation.PointSimIndices[PointIndex];
+					const FVector2f& S 			= HairInterpolation.PointSimLerps[PointIndex];
+					for (uint32 KIndex = 0; KIndex < KCount; ++KIndex)
+					{		
+						FHairInterpolationPoint& OutInterp = Out[PointIndex*KCount + KIndex];
+						OutInterp.VertexIndex 		= Indices[KIndex];
+						OutInterp.VertexLerp		= FMath::Clamp(S[KIndex] * 255.f, 0, 255);
+					}
 				}
 			}
 
-			HairStrandsBuilder::CopyToBulkData<FHairStrandsRootIndexFormat>(OutBulkData.Data.SimRootPointIndex, SimRootPointIndex);
+			HairStrandsBuilder::CopyToBulkData<FHairStrandsInterpolationFormat>(OutBulkData.Data.CurveInterpolation, OutInterpolationCurves);
+			HairStrandsBuilder::CopyToBulkData<FHairStrandsInterpolationFormat>(OutBulkData.Data.PointInterpolation, OutInterpolationPoints);
+			ReportSize(TEXT("Interpolation Curves"), OutInterpolationCurves);
+			ReportSize(TEXT("Interpolation Points"), OutInterpolationPoints);
 		}
 	}
 
@@ -1608,22 +1586,29 @@ namespace HairInterpolationBuilder
 		check(SimData.GetNumCurves() > 0);
 		int32 SimCurveCount = SimData.GetNumCurves();
 
-		OutInterpolation.SetNum(RenData.GetNumPoints());
-		for (uint32 CurveIndex = 0; CurveIndex < RenData.GetNumCurves(); ++CurveIndex)
+		OutInterpolation.SetNum(RenData.GetNumCurves(), RenData.GetNumPoints());
+		for (uint32 RenCurveIndex = 0; RenCurveIndex < RenData.GetNumCurves(); ++RenCurveIndex)
 		{
-			const uint32 CurveOffset = RenData.StrandsCurves.CurvesOffset[CurveIndex];
-			const uint16 CurveNumVertices = RenData.StrandsCurves.CurvesCount[CurveIndex];
+			const uint32 RenCurveOffset = RenData.StrandsCurves.CurvesOffset[RenCurveIndex];
+			const uint16 RenCurveNumVertices = RenData.StrandsCurves.CurvesCount[RenCurveIndex];
+
+			const FIntVector StrandClosestGuides = RenData.StrandsCurves.CurvesClosestGuideIDs[RenCurveIndex];
+			const FVector3f StrandGuideWeights  = RenData.StrandsCurves.CurvesClosestGuideWeights[RenCurveIndex];
 
-			const FIntVector StrandClosestGuides = RenData.StrandsCurves.CurvesClosestGuideIDs[CurveIndex];
-			const FVector3f StrandGuideWeights  = RenData.StrandsCurves.CurvesClosestGuideWeights[CurveIndex];
-			for (uint16 VertexIndex = 0; VertexIndex < CurveNumVertices; ++VertexIndex)
+
+			float CurveSimWeights[FClosestGuides::Count];
+			for (int32 GuideIndex = 0; GuideIndex < FClosestGuides::Count; ++GuideIndex)
 			{
-				const uint32 PointGlobalIndex = VertexIndex + CurveOffset;
-				const FVector3f& RenPointPosition = RenData.StrandsPoints.PointsPosition[PointGlobalIndex];
-				const float RenPointDistance = RenData.StrandsPoints.PointsCoordU[PointGlobalIndex] * RenData.StrandsCurves.CurvesLength[CurveIndex];
+				CurveSimWeights[GuideIndex] = 0;
+			}
+
+			for (uint16 RenPointIndex = 0; RenPointIndex < RenCurveNumVertices; ++RenPointIndex)
+			{
+				const uint32 RenPointGlobalIndex = RenPointIndex + RenCurveOffset;
+				const FVector3f& RenPointPosition = RenData.StrandsPoints.PointsPosition[RenPointGlobalIndex];
+				const float RenPointDistance = RenData.StrandsPoints.PointsCoordU[RenPointGlobalIndex] * RenData.StrandsCurves.CurvesLength[RenCurveIndex];
 
 				bool bHasValidGuide = false;
-				float TotalWeight = 0;
 				for (uint32 GuideIndex = 0; GuideIndex < FClosestGuides::Count; ++GuideIndex)
 				{
 					int32 ImportedGroomID = StrandClosestGuides[GuideIndex];
@@ -1635,11 +1620,13 @@ namespace HairInterpolationBuilder
 						const uint32 SimOffset = SimData.StrandsCurves.CurvesOffset[*SimCurveIndex];
 						const FVertexInterpolationDesc Desc = FindMatchingVertex(RenPointDistance, SimData, *SimCurveIndex);
 
-						OutInterpolation.PointsSimCurvesIndex[PointGlobalIndex][GuideIndex] = *SimCurveIndex;
-						OutInterpolation.PointsSimCurvesVertexIndex[PointGlobalIndex][GuideIndex] = Desc.Index0 + SimOffset;
-						OutInterpolation.PointsSimCurvesVertexLerp[PointGlobalIndex][GuideIndex] = Desc.T;
-						OutInterpolation.PointsSimCurvesVertexWeights[PointGlobalIndex][GuideIndex] = StrandGuideWeights[GuideIndex];
-						TotalWeight += OutInterpolation.PointsSimCurvesVertexWeights[PointGlobalIndex][GuideIndex];
+						// Point data
+						OutInterpolation.PointSimIndices[RenPointGlobalIndex][GuideIndex] = Desc.Index0 + SimOffset;
+						OutInterpolation.PointSimLerps[RenPointGlobalIndex][GuideIndex] = Desc.T;
+
+						// Curve data
+						OutInterpolation.CurveSimIndices[RenCurveIndex][GuideIndex] = *SimCurveIndex;
+						CurveSimWeights[GuideIndex] += OutInterpolation.CurveSimWeights[RenCurveIndex][GuideIndex];
 						bHasValidGuide = true;
 					}
 				}
@@ -1655,23 +1642,27 @@ namespace HairInterpolationBuilder
 						const uint32 SimOffset = SimData.StrandsCurves.CurvesOffset[SimCurveIndex];
 						const FVertexInterpolationDesc Desc = FindMatchingVertex(RenPointDistance, SimData, SimCurveIndex);
 
-						OutInterpolation.PointsSimCurvesIndex[PointGlobalIndex][GuideIndex] = SimCurveIndex;
-						OutInterpolation.PointsSimCurvesVertexIndex[PointGlobalIndex][GuideIndex] = Desc.Index0 + SimOffset;
-						OutInterpolation.PointsSimCurvesVertexLerp[PointGlobalIndex][GuideIndex] = Desc.T;
-						OutInterpolation.PointsSimCurvesVertexWeights[PointGlobalIndex][GuideIndex] = 1.0f;
-						TotalWeight += OutInterpolation.PointsSimCurvesVertexWeights[PointGlobalIndex][GuideIndex];
-					}
-				}
+						// Point data
+						OutInterpolation.PointSimIndices[RenPointGlobalIndex][GuideIndex] = Desc.Index0 + SimOffset;
+						OutInterpolation.PointSimLerps[RenPointGlobalIndex][GuideIndex] = Desc.T;
 
-				// Normalize the weights
-				if (TotalWeight > 0.f)
-				{
-					for (int32 GuideIndex = 0; GuideIndex < FClosestGuides::Count; ++GuideIndex)
-					{
-						OutInterpolation.PointsSimCurvesVertexWeights[PointGlobalIndex][GuideIndex] /= TotalWeight;
+						// Curve data
+						OutInterpolation.CurveSimIndices[RenCurveIndex][GuideIndex] = SimCurveIndex;
+						CurveSimWeights[GuideIndex] += 1.0f;
 					}
 				}
 			}
+
+			// Normalize the curve weights
+			float TotalWeight = 0;
+			for (int32 GuideIndex = 0; GuideIndex < FClosestGuides::Count; ++GuideIndex)
+			{
+				TotalWeight += CurveSimWeights[GuideIndex];
+			}
+			for (int32 GuideIndex = 0; GuideIndex < FClosestGuides::Count; ++GuideIndex)
+			{
+				OutInterpolation.CurveSimWeights[RenCurveIndex][GuideIndex] = CurveSimWeights[GuideIndex] / FMath::Max(0.001f, TotalWeight);
+			}
 		}
 	}
 }
@@ -2215,7 +2206,7 @@ void FGroomBuilder::BuildInterplationData(
 	// Build Rendering data for InterpolationData
 	// Build InterpolationData from render and simulation HairStrandsDatas
 	// Skip building if interpolation data was provided by the source file
-	if (OutInterpolationData.Num() == 0)
+	if (OutInterpolationData.GetPointCount() == 0)
 	{
 		// If there's usable closest guides and guide weights attributes, fill them into the asset
 		// This step requires the HairSimulationData (guides) to be filled prior to this
@@ -2252,7 +2243,7 @@ void FGroomBuilder::BuildInterplationBulkData(
 	const FHairStrandsInterpolationDatas& InInterpolationData,
 	FHairStrandsInterpolationBulkData& OutInterpolationData)
 {
-	HairInterpolationBuilder::BuildRenderData(InSimData, InInterpolationData, OutInterpolationData);
+	HairInterpolationBuilder::BuildInterplationBulkData(InSimData, InInterpolationData, OutInterpolationData);
 }
 
 namespace FHairStrandsDecimation
@@ -3268,6 +3259,7 @@ static void BuildClusterBulkData(
 	check(Out.Header.CurveCount		== uint32(In.CurveToClusterIds.Num()));
 	
 	HairStrandsBuilder::CopyToBulkData<FHairClusterIndexFormat>(Out.Data.CurveToClusterIds, In.CurveToClusterIds);
+	ReportSize(TEXT("CurveToClusterIds"), In.CurveToClusterIds);
 
 	// Pack point LOD data
 	// |                          32bits                               |
@@ -3287,6 +3279,7 @@ static void BuildClusterBulkData(
 		PackedPointLODs[PackedIt] = Packed;
 	}
 	HairStrandsBuilder::CopyToBulkData<FHairClusterIndexFormat>(Out.Data.PointLODs, PackedPointLODs);
+	ReportSize(TEXT("PointLODs"), PackedPointLODs);
 
 	// Pack LODInfo into GPU format
 	{
@@ -3368,6 +3361,7 @@ static void BuildClusterBulkData(
 		}
 
 		HairStrandsBuilder::CopyToBulkData<FHairClusterInfoFormat>(Out.Data.PackedClusterInfos, PackedClusterInfos);
+		ReportSize(TEXT("ClusterInfos"), PackedClusterInfos);
 	}
 }
 
diff --git a/Engine/Plugins/Runtime/HairStrands/Source/HairStrandsCore/Private/GroomManager.cpp b/Engine/Plugins/Runtime/HairStrands/Source/HairStrandsCore/Private/GroomManager.cpp
--- a/Engine/Plugins/Runtime/HairStrands/Source/HairStrandsCore/Private/GroomManager.cpp
+++ b/Engine/Plugins/Runtime/HairStrands/Source/HairStrandsCore/Private/GroomManager.cpp
@@ -1354,6 +1354,11 @@ static void RunHairStrandsInterpolation_Strands(
 	}
 
 	{
+		DECLARE_GPU_STAT(HairStrandsInterpolationCurve);
+		RDG_EVENT_SCOPE(GraphBuilder, "HairStrandsInterpolationCurve");
+		TRACE_CPUPROFILER_EVENT_SCOPE(HairStrandsInterpolationCurve);
+		RDG_GPU_STAT_SCOPE(GraphBuilder, HairStrandsInterpolationCurve);
+
 		TArray<FRDGBufferSRVRef> Transitions;
 		Transitions.Reserve(InstanceDatas.Num());
 		for (FInstanceData& InstanceData : InstanceDatas)
@@ -1370,6 +1375,8 @@ static void RunHairStrandsInterpolation_Strands(
 					ShaderPrintData,
 					InstanceData.Instance,
 					InstanceData.ActivePointCount,
+					InstanceData.ActiveCurveCount,
+					InstanceData.Instance->Strands.RestResource->BulkData.Header.MaxPointPerCurve,
 					InstanceData.MeshLODIndex,
 					InstanceData.Instance->Strands.Modifier.HairLengthScale,
 					InstanceData.Instance->Strands.HairInterpolationType,
@@ -1385,11 +1392,13 @@ static void RunHairStrandsInterpolation_Strands(
 					bHasSkinning && bValidGuide ? InstanceData.Instance->Guides.DeformedRootResource : nullptr,
 					RegisterAsSRV(GraphBuilder, InstanceData.Instance->Strands.RestResource->PositionBuffer),
 					RegisterAsSRV(GraphBuilder, InstanceData.Instance->Strands.RestResource->PointToCurveBuffer),
+					RegisterAsSRV(GraphBuilder, InstanceData.Instance->Strands.RestResource->CurveBuffer),
+					bValidGuide ? RegisterAsSRV(GraphBuilder, InstanceData.Instance->Guides.RestResource->CurveBuffer) : nullptr,
 					bUseSingleGuide,
-					bValidGuide ? RegisterAsSRV(GraphBuilder, InstanceData.Instance->Strands.InterpolationResource->InterpolationBuffer) : nullptr,
+					bValidGuide ? RegisterAsSRV(GraphBuilder, InstanceData.Instance->Strands.InterpolationResource->CurveInterpolationBuffer) : nullptr,
+					bValidGuide ? RegisterAsSRV(GraphBuilder, InstanceData.Instance->Strands.InterpolationResource->PointInterpolationBuffer) : nullptr,
 					bValidGuide ? RegisterAsSRV(GraphBuilder, InstanceData.Instance->Guides.RestResource->PositionBuffer) : nullptr,
 					bValidGuide ? RegisterAsSRV(GraphBuilder, InstanceData.Instance->Guides.DeformedResource->GetBuffer(FHairStrandsDeformedResource::Current)) : nullptr,
-					bValidGuide ? RegisterAsSRV(GraphBuilder, InstanceData.Instance->Strands.InterpolationResource->SimRootPointIndexBuffer) : nullptr,
 					bValidGuide ? RegisterAsSRV(GraphBuilder, InstanceData.Instance->Guides.RestResource->PointToCurveBuffer) : nullptr,
 					InstanceData.RDGResources.DeformerPositionSRV,
 					InstanceData.RDGResources.PositionUAV,
@@ -1530,7 +1539,8 @@ static void RunHairStrandsInterpolation_Strands(
 					Register(GraphBuilder, InstanceData.Instance->Strands.RestResource->CurveAttributeBuffer, ERDGImportedBufferFlags::CreateSRV).Buffer,
 					RegisterAsSRV(GraphBuilder, InstanceData.Instance->Strands.RestResource->CurveBuffer),
 					RegisterAsSRV(GraphBuilder, InstanceData.Instance->Strands.ClusterResource->CurveToClusterIdBuffer),
-					bValidGuide ? RegisterAsSRV(GraphBuilder, InstanceData.Instance->Strands.InterpolationResource->InterpolationBuffer) : nullptr,
+					bValidGuide ? RegisterAsSRV(GraphBuilder, InstanceData.Instance->Strands.InterpolationResource->CurveInterpolationBuffer) : nullptr,
+					bValidGuide ? RegisterAsSRV(GraphBuilder, InstanceData.Instance->Strands.InterpolationResource->PointInterpolationBuffer) : nullptr,
 					OutRenCurveAttributeBuffer);
 			}
 		}
@@ -1924,6 +1934,8 @@ static void RunHairStrandsInterpolation_Cards(
 					ShaderPrintData,
 					InstanceData.Instance,
 					LOD.Guides.RestResource->GetPointCount(),
+					LOD.Guides.RestResource->GetCurveCount(),
+					LOD.Guides.RestResource->BulkData.Header.MaxPointPerCurve,
 					InstanceData.MeshLODIndex,
 					1.0f,
 					LOD.Guides.HairInterpolationType,
@@ -1939,11 +1951,13 @@ static void RunHairStrandsInterpolation_Cards(
 					InstanceData.bHasSkinning && InstanceData.bValidGuide ? InstanceData.Instance->Guides.DeformedRootResource : nullptr,
 					RegisterAsSRV(GraphBuilder, LOD.Guides.RestResource->PositionBuffer),
 					RegisterAsSRV(GraphBuilder, LOD.Guides.RestResource->PointToCurveBuffer),
+					RegisterAsSRV(GraphBuilder, LOD.Guides.RestResource->CurveBuffer),
+					InstanceData.bValidGuide ? RegisterAsSRV(GraphBuilder, InstanceData.Instance->Guides.RestResource->CurveBuffer) : nullptr,
 					bUseSingleGuide,
-					InstanceData.bValidGuide ? RegisterAsSRV(GraphBuilder, LOD.Guides.InterpolationResource->InterpolationBuffer) : nullptr,
+					InstanceData.bValidGuide ? RegisterAsSRV(GraphBuilder, LOD.Guides.InterpolationResource->CurveInterpolationBuffer) : nullptr,
+					InstanceData.bValidGuide ? RegisterAsSRV(GraphBuilder, LOD.Guides.InterpolationResource->PointInterpolationBuffer) : nullptr,
 					InstanceData.bValidGuide ? RegisterAsSRV(GraphBuilder, InstanceData.Instance->Guides.RestResource->PositionBuffer) : nullptr,
 					InstanceData.bValidGuide ? RegisterAsSRV(GraphBuilder, InstanceData.Instance->Guides.DeformedResource->GetBuffer(FHairStrandsDeformedResource::Current)) : nullptr,
-					RegisterAsSRV(GraphBuilder, LOD.Guides.InterpolationResource->SimRootPointIndexBuffer),
 					InstanceData.bValidGuide ? RegisterAsSRV(GraphBuilder, InstanceData.Instance->Guides.RestResource->PointToCurveBuffer) : nullptr,
 					nullptr,
 					Guides_DeformedPositionUAV,
diff --git a/Engine/Plugins/Runtime/HairStrands/Source/HairStrandsCore/Private/GroomResources.cpp b/Engine/Plugins/Runtime/HairStrands/Source/HairStrandsCore/Private/GroomResources.cpp
--- a/Engine/Plugins/Runtime/HairStrands/Source/HairStrandsCore/Private/GroomResources.cpp
+++ b/Engine/Plugins/Runtime/HairStrands/Source/HairStrandsCore/Private/GroomResources.cpp
@@ -1563,7 +1563,7 @@ void FHairStrandsDeformedRootResource::InternalRelease()
 
 FHairStrandsInterpolationResource::FHairStrandsInterpolationResource(FHairStrandsInterpolationBulkData& InBulkData, const FHairResourceName& InResourceName, const FName& InOwnerName) :
 	FHairCommonResource(EHairStrandsAllocationType::Deferred, InResourceName, InOwnerName),
-	InterpolationBuffer(), BulkData(InBulkData)
+CurveInterpolationBuffer(), PointInterpolationBuffer(),BulkData(InBulkData)
 {
 	MaxAvailableCurveCount = 0;
 
@@ -1585,14 +1585,14 @@ void FHairStrandsInterpolationResource::InternalAllocate(FRDGBuilder& GraphBuild
 	// If we enter this function, the request need to be completed
 	check(StreamingRequest.IsCompleted());
 
-	InternalCreateByteAddressBufferRDG_FromHairBulkData(GraphBuilder, BulkData.Data.Interpolation, InterpolationBuffer, ToHairResourceDebugName(TEXT("Hair.StrandsInterpolation_InterpolationBuffer"), ResourceName), OwnerName, EHairResourceUsageType::Static);
-	InternalCreateVertexBufferRDG_FromHairBulkData<FHairStrandsRootIndexFormat>(GraphBuilder, BulkData.Data.SimRootPointIndex, BulkData.Header.SimPointCount, SimRootPointIndexBuffer, ToHairResourceDebugName(TEXT("Hair.StrandsInterpolation_SimRootPointIndex"), ResourceName), OwnerName, EHairResourceUsageType::Static);
+	InternalCreateByteAddressBufferRDG_FromHairBulkData(GraphBuilder, BulkData.Data.CurveInterpolation, CurveInterpolationBuffer, ToHairResourceDebugName(TEXT("Hair.StrandsInterpolation_CurveInterpolationBuffer"), ResourceName), OwnerName, EHairResourceUsageType::Static);
+	InternalCreateByteAddressBufferRDG_FromHairBulkData(GraphBuilder, BulkData.Data.PointInterpolation, PointInterpolationBuffer, ToHairResourceDebugName(TEXT("Hair.StrandsInterpolation_PointInterpolationBuffer"), ResourceName), OwnerName, EHairResourceUsageType::Static);
 }
 
 void FHairStrandsInterpolationResource::InternalRelease()
 {
-	InterpolationBuffer.Release();
-	SimRootPointIndexBuffer.Release();
+	CurveInterpolationBuffer.Release();
+	PointInterpolationBuffer.Release();
 	MaxAvailableCurveCount = 0;
 }
 
diff --git a/Engine/Plugins/Runtime/HairStrands/Source/HairStrandsCore/Private/HairCardsBuilder.cpp b/Engine/Plugins/Runtime/HairStrands/Source/HairStrandsCore/Private/HairCardsBuilder.cpp
--- a/Engine/Plugins/Runtime/HairStrands/Source/HairStrandsCore/Private/HairCardsBuilder.cpp
+++ b/Engine/Plugins/Runtime/HairStrands/Source/HairStrandsCore/Private/HairCardsBuilder.cpp
@@ -610,6 +610,7 @@ static bool InternalImportGeometry_WithGeneratedGuides(
 		InternalCreateCardsInterpolation(Out.Cards, OutGuides, InterpolationData);
 
 		// Fill out the interpolation data
+		static_assert(sizeof(FHairCardsInterpolationVertex) == HAIR_INTERPOLATION_CARDS_GUIDE_STRIDE);
 		OutInterpolationBulkData.Interpolation.SetNum(PointCount);
 		for (uint32 PointIt = 0; PointIt < PointCount; ++PointIt)
 		{
diff --git a/Engine/Plugins/Runtime/HairStrands/Source/HairStrandsCore/Private/HairStrandsClusterCulling.cpp b/Engine/Plugins/Runtime/HairStrands/Source/HairStrandsCore/Private/HairStrandsClusterCulling.cpp
--- a/Engine/Plugins/Runtime/HairStrands/Source/HairStrandsCore/Private/HairStrandsClusterCulling.cpp
+++ b/Engine/Plugins/Runtime/HairStrands/Source/HairStrandsCore/Private/HairStrandsClusterCulling.cpp
@@ -44,7 +44,7 @@ class FHairClusterCullCS: public FGlobalShader
 		END_SHADER_PARAMETER_STRUCT()
 
 public:
-		static uint32 GetGroupSize() { return 64u; }
+	static uint32 GetGroupSize() { return 64u; }
 	static bool ShouldCompilePermutation(const FGlobalShaderPermutationParameters& Parameters) { return IsHairStrandsSupported(EHairStrandsShaderType::Strands, Parameters.Platform); }
 	static void ModifyCompilationEnvironment(const FGlobalShaderPermutationParameters& Parameters, FShaderCompilerEnvironment& OutEnvironment)
 	{
@@ -73,7 +73,7 @@ class FHairClusterCullArgsCS: public FGlobalShader
 		END_SHADER_PARAMETER_STRUCT()
 
 public:
-		static bool ShouldCompilePermutation(const FGlobalShaderPermutationParameters& Parameters) { return IsHairStrandsSupported(EHairStrandsShaderType::Strands, Parameters.Platform); }
+	static bool ShouldCompilePermutation(const FGlobalShaderPermutationParameters& Parameters) { return IsHairStrandsSupported(EHairStrandsShaderType::Strands, Parameters.Platform); }
 	static void ModifyCompilationEnvironment(const FGlobalShaderPermutationParameters& Parameters, FShaderCompilerEnvironment& OutEnvironment)
 	{
 		FGlobalShader::ModifyCompilationEnvironment(Parameters, OutEnvironment);
@@ -130,33 +130,19 @@ void AddClusterCullingPass(
 			ShaderPrint::SetParameters(GraphBuilder, *ShaderPrintData, Parameters->ShaderPrintUniformBuffer);
 		}
 
-		// Pick permutation based on groom max number of point per curve
-		//const uint32 PointPerCurve = FMath::Clamp(uint32(FMath::Pow( FMath::RoundFromZero(FMath::Log2(float(ClusterData.MaxPointPerCurve))),2u)), 4u, 32u);
-		const uint32 PointPerCurve = FMath::Clamp(uint32(FMath::Pow(2u, FMath::RoundFromZero(FMath::Log2(float(ClusterData.MaxPointPerCurve))))), 4u, 64u);
-		check(FMath::IsPowerOfTwo(PointPerCurve));
-
-		const uint32 CurvePerGroup = FHairClusterCullCS::GetGroupSize() / PointPerCurve;
-		const uint32 LinearGroupCount = FMath::DivideAndRoundUp(Parameters->CurveCount, CurvePerGroup);
-
-		FIntVector DispatchCount(LinearGroupCount, 1, 1);
-		if (DispatchCount.X > 0xFFFFu)
-		{
-			DispatchCount.X = 64;
-			DispatchCount.Y = FMath::DivideAndRoundUp(LinearGroupCount, uint32(DispatchCount.X));
-		}
-		Parameters->DispatchCountX = DispatchCount.X;
-		check(DispatchCount.X <= 0xFFFFu);
-		check(DispatchCount.Y <= 0xFFFFu);
+		// Compute the dispatch information for pass dispatching work per curve
+		const FPointPerCurveDispatchInfo DispatchInfo = GetPointPerCurveDispatchInfo(ClusterData.MaxPointPerCurve, Parameters->CurveCount, FHairClusterCullCS::GetGroupSize());
+		Parameters->DispatchCountX = DispatchInfo.DispatchCount.X;
 
 		FHairClusterCullCS::FPermutationDomain Permutation;
-		Permutation.Set<FHairClusterCullCS::FPointPerCurve>(PointPerCurve);
+		Permutation.Set<FHairClusterCullCS::FPointPerCurve>(DispatchInfo.PointPerCurve);
 		TShaderMapRef<FHairClusterCullCS> ComputeShader(ShaderMap, Permutation);
 		FComputeShaderUtils::AddPass(
 			GraphBuilder,
 			RDG_EVENT_NAME("HairStrands::ClusterCullPass2"),
 			ComputeShader,
 			Parameters,
-			DispatchCount);
+			DispatchInfo.DispatchCount);
 	}
 
 	// 2. Prepare indirect draw/dispatch args buffers
diff --git a/Engine/Plugins/Runtime/HairStrands/Source/HairStrandsCore/Private/HairStrandsDatas.cpp b/Engine/Plugins/Runtime/HairStrands/Source/HairStrandsCore/Private/HairStrandsDatas.cpp
--- a/Engine/Plugins/Runtime/HairStrands/Source/HairStrandsCore/Private/HairStrandsDatas.cpp
+++ b/Engine/Plugins/Runtime/HairStrands/Source/HairStrandsCore/Private/HairStrandsDatas.cpp
@@ -99,18 +99,20 @@ void FHairStrandsDatas::Reset()
 
 void FHairStrandsInterpolationDatas::Reset()
 {
-	PointsSimCurvesVertexWeights.Reset();
-	PointsSimCurvesVertexLerp.Reset();
-	PointsSimCurvesVertexIndex.Reset();
-	PointsSimCurvesIndex.Reset();
+	CurveSimWeights.Reset();
+	CurveSimIndices.Reset();
+	CurveSimRootPointIndex.Reset();
+	PointSimLerps.Reset();
+	PointSimIndices.Reset();
 }
 
-void FHairStrandsInterpolationDatas::SetNum(const uint32 NumCurves)
+void FHairStrandsInterpolationDatas::SetNum(const uint32 NumCurves, const uint32 NumPoints)
 {
-	PointsSimCurvesVertexWeights.SetNum(NumCurves);
-	PointsSimCurvesVertexLerp.SetNum(NumCurves);
-	PointsSimCurvesVertexIndex.SetNum(NumCurves);
-	PointsSimCurvesIndex.SetNum(NumCurves);
+	CurveSimIndices.SetNum(NumCurves);
+	CurveSimWeights.SetNum(NumCurves);
+	CurveSimRootPointIndex.SetNum(NumCurves);
+	PointSimIndices.SetNum(NumPoints);
+	PointSimLerps.SetNum(NumPoints);
 }
 
 void FHairStrandsClusterData::Reset()
@@ -613,6 +615,17 @@ void FHairStrandsBulkData::GetResources(FHairStrandsBulkCommon::FQuery& Out)
 	}
 }
 
+uint32 FHairStrandsBulkData::GetSize() const
+{
+	uint32 Out = 0;
+	Out += Data.Positions.GetBulkDataSize();
+	Out += Data.CurveAttributes.GetBulkDataSize();
+	Out += Data.PointAttributes.GetBulkDataSize();
+	Out += Data.PointToCurve.GetBulkDataSize();
+	Out += Data.Curves.GetBulkDataSize();
+	return Out;
+}
+
 void FHairStrandsBulkData::Reset()
 {
 	Header.CurveCount = 0;
@@ -660,25 +673,25 @@ void FHairStrandsInterpolationBulkData::Reset()
 {
 	Header.Flags = 0;
 	Header.PointCount = 0;
-	Header.SimPointCount = 0;
+	Header.CurveCount = 0;
 	
 	// Deallocate memory if needed
-	Data.Interpolation.RemoveBulkData();
-	Data.SimRootPointIndex.RemoveBulkData();
+	Data.CurveInterpolation.RemoveBulkData();
+	Data.PointInterpolation.RemoveBulkData();
 
 	// Reset the bulk byte buffer to ensure the (serialize) data size is reset to 0
-	Data.Interpolation		= FHairBulkContainer();
-	Data.SimRootPointIndex	= FHairBulkContainer();
+	Data.CurveInterpolation	= FHairBulkContainer();
+	Data.PointInterpolation	= FHairBulkContainer();
 }
 
 void FHairStrandsInterpolationBulkData::SerializeHeader(FArchive& Ar, UObject* Owner)
 {
 	Ar << Header.Flags;
 	Ar << Header.PointCount;
-	Ar << Header.SimPointCount;
+	Ar << Header.CurveCount;
 
-	Ar << Header.Strides.InterpolationStride;
-	Ar << Header.Strides.SimRootPointIndexStride;
+	Ar << Header.Strides.CurveInterpolationStride;
+	Ar << Header.Strides.PointInterpolationStride;
 }
 
 uint32 FHairStrandsInterpolationBulkData::GetResourceCount() const
@@ -694,20 +707,30 @@ void FHairStrandsInterpolationBulkData::GetResources(FHairStrandsBulkCommon::FQu
 	{
 		// Translate requested curve count into chunk/offset/size to be read
 		uint32 PointCount = 0;
+		uint32 CurveCount = 0;
 		if (Out.Type == FHairStrandsBulkCommon::FQuery::ReadIO || Out.Type == FHairStrandsBulkCommon::FQuery::ReadDDC || Out.Type == FHairStrandsBulkCommon::FQuery::UnloadData)
 		{
+			CurveCount = FMath::Min(Header.CurveCount, Out.GetCurveCount());
 			PointCount = FMath::Min(Header.PointCount, Out.GetPointCount());
 		}
 
-		Out.Add(Data.Interpolation, 	TEXT("_Interpolation"), 	Data.Interpolation.LoadedSize, 		PointCount * Header.Strides.InterpolationStride);
-		Out.Add(Data.SimRootPointIndex, TEXT("_SimRootPointIndex"), Data.SimRootPointIndex.LoadedSize, 	Header.SimPointCount * Header.Strides.SimRootPointIndexStride); // Load all data at once (guide data)
+		Out.Add(Data.CurveInterpolation, 	TEXT("_CurveInterpolation"), 	Data.CurveInterpolation.LoadedSize, 		CurveCount * Header.Strides.CurveInterpolationStride);
+		Out.Add(Data.PointInterpolation, 	TEXT("_PointInterpolation"), 	Data.PointInterpolation.LoadedSize, 		PointCount * Header.Strides.PointInterpolationStride);
 	}
 }
 
+uint32 FHairStrandsInterpolationBulkData::GetSize() const
+{
+	uint32 Out = 0;
+	Out += Data.CurveInterpolation.GetBulkDataSize();
+	Out += Data.PointInterpolation.GetBulkDataSize();
+	return Out;
+}
+
 void FHairStrandsInterpolationBulkData::ResetLoadedSize()
 {
-	Data.Interpolation.LoadedSize		= 0;
-	Data.SimRootPointIndex.LoadedSize	= 0;
+	Data.CurveInterpolation.LoadedSize = 0;
+	Data.PointInterpolation.LoadedSize = 0;
 }
 
 /////////////////////////////////////////////////////////////////////////////////////////
@@ -795,6 +818,15 @@ void FHairStrandsClusterBulkData::GetResources(FHairStrandsBulkCommon::FQuery &
 	}
 }
 
+uint32 FHairStrandsClusterBulkData::GetSize() const
+{
+	uint32 Out = 0;;
+	Out += Data.PackedClusterInfos.GetBulkDataSize();
+	Out += Data.CurveToClusterIds.GetBulkDataSize();
+	Out += Data.PointLODs.GetBulkDataSize();
+	return Out;
+}
+
 uint32 FHairStrandsClusterBulkData::GetCurveCount(float InLODIndex) const
 {
 	const TArray<FHairLODInfo>& LODInfos = Header.LODInfos;
@@ -926,6 +958,23 @@ void FHairStrandsRootBulkData::GetResources(FQuery& Out)
 	}
 }
 
+uint32 FHairStrandsRootBulkData::GetSize() const
+{
+	uint32 Out = 0;
+	for (const auto& LOD : Data.LODs)
+	{
+		Out += LOD.RootToUniqueTriangleIndexBuffer.GetBulkDataSize();
+		Out += LOD.RootBarycentricBuffer.GetBulkDataSize();
+		Out += LOD.UniqueTriangleIndexBuffer.GetBulkDataSize();
+		Out += LOD.RestUniqueTrianglePositionBuffer.GetBulkDataSize();
+
+		Out += LOD.MeshInterpolationWeightsBuffer.GetBulkDataSize();
+		Out += LOD.MeshSampleIndicesAndSectionsBuffer.GetBulkDataSize();
+		Out += LOD.RestSamplePositionsBuffer.GetBulkDataSize();
+	}
+	return Out;
+}
+
 void FHairStrandsRootBulkData::Reset()
 {
 	// Header
diff --git a/Engine/Plugins/Runtime/HairStrands/Source/HairStrandsCore/Private/HairStrandsInterpolation.cpp b/Engine/Plugins/Runtime/HairStrands/Source/HairStrandsCore/Private/HairStrandsInterpolation.cpp
--- a/Engine/Plugins/Runtime/HairStrands/Source/HairStrandsCore/Private/HairStrandsInterpolation.cpp
+++ b/Engine/Plugins/Runtime/HairStrands/Source/HairStrandsCore/Private/HairStrandsInterpolation.cpp
@@ -452,21 +452,25 @@ class FHairInterpolationCS : public FGlobalShader
 	DECLARE_GLOBAL_SHADER(FHairInterpolationCS);
 	SHADER_USE_PARAMETER_STRUCT(FHairInterpolationCS, FGlobalShader);
 
+	class FPointPerCurve : SHADER_PERMUTATION_SPARSE_INT("PERMUTATION_POINT_PER_CURVE", 4, 8, 16, 32);
 	class FDynamicGeometry : SHADER_PERMUTATION_INT("PERMUTATION_DYNAMIC_GEOMETRY", 5);
 	class FSimulation : SHADER_PERMUTATION_BOOL("PERMUTATION_SIMULATION");
 	class FSingleGuide : SHADER_PERMUTATION_BOOL("PERMUTATION_USE_SINGLE_GUIDE");
-	class FCulling : SHADER_PERMUTATION_BOOL("PERMUTATION_CULLING");
 	class FDeformer : SHADER_PERMUTATION_BOOL("PERMUTATION_DEFORMER");
-	using FPermutationDomain = TShaderPermutationDomain<FDynamicGeometry, FSimulation, FSingleGuide, FCulling, FDeformer>;
+	using FPermutationDomain = TShaderPermutationDomain<FDynamicGeometry, FSimulation, FSingleGuide, FDeformer, FPointPerCurve>;
 
 	BEGIN_SHADER_PARAMETER_STRUCT(FParameters, )
 		SHADER_PARAMETER_STRUCT_INCLUDE(ShaderPrint::FShaderParameters, ShaderDrawParameters)
 		SHADER_PARAMETER(uint32, VertexCount)
+		SHADER_PARAMETER(uint32, CurveCount)
 		SHADER_PARAMETER(float, HairLengthScale)
 		SHADER_PARAMETER(FVector3f, InRenHairPositionOffset)
 		SHADER_PARAMETER(FVector3f, InSimHairPositionOffset)
-		SHADER_PARAMETER(uint32,  HairStrandsVFTODO_bCullingEnable)
 		SHADER_PARAMETER(FMatrix44f, LocalToWorldMatrix)
+		SHADER_PARAMETER(uint32,  DispatchCountX)
+	
+		SHADER_PARAMETER_RDG_BUFFER_SRV(Buffer, RenCurveBuffer)
+		SHADER_PARAMETER_RDG_BUFFER_SRV(Buffer, SimCurveBuffer)
 
 		SHADER_PARAMETER_RDG_BUFFER_SRV(Buffer, RenRestPosePositionBuffer)
 		SHADER_PARAMETER_RDG_BUFFER_SRV(Buffer, RenDeformerPositionBuffer)
@@ -475,7 +479,8 @@ class FHairInterpolationCS : public FGlobalShader
 		SHADER_PARAMETER_RDG_BUFFER_SRV(Buffer, SimRestPosePositionBuffer)
 		SHADER_PARAMETER_RDG_BUFFER_SRV(Buffer, SimDeformedPositionBuffer)
 
-		SHADER_PARAMETER_RDG_BUFFER_SRV(ByteAddressBuffer, InterpolationBuffer)
+		SHADER_PARAMETER_RDG_BUFFER_SRV(ByteAddressBuffer, CurveInterpolationBuffer)
+		SHADER_PARAMETER_RDG_BUFFER_SRV(ByteAddressBuffer, PointInterpolationBuffer)
 
 		SHADER_PARAMETER_RDG_BUFFER_SRV(Buffer<float4>, RenRootRestPositionBuffer)
 		SHADER_PARAMETER_RDG_BUFFER_SRV(Buffer<float4>, RenRootDeformedPositionBuffer)
@@ -488,7 +493,6 @@ class FHairInterpolationCS : public FGlobalShader
 		SHADER_PARAMETER_RDG_BUFFER_SRV(Buffer<float4>, SimRootDeformedPositionBuffer)
 
 		SHADER_PARAMETER_RDG_BUFFER_SRV(Buffer<uint>, SimRootBarycentricBuffer)
-		SHADER_PARAMETER_RDG_BUFFER_SRV(Buffer<uint>, SimPointToCurveBuffer)
 		SHADER_PARAMETER_RDG_BUFFER_SRV(Buffer<uint>, SimRootToUniqueTriangleIndexBuffer)
 
 		SHADER_PARAMETER_RDG_BUFFER_SRV(Buffer<uint>, SimRootPointIndexBuffer)
@@ -504,13 +508,14 @@ class FHairInterpolationCS : public FGlobalShader
 		END_SHADER_PARAMETER_STRUCT()
 
 public:
-	static uint32 GetGroupSize() { return HAIR_VERTEXCOUNT_GROUP_SIZE; }
+	static uint32 GetGroupSize() { return 32;}
 	static bool ShouldCompilePermutation(const FGlobalShaderPermutationParameters& Parameters) { return IsHairStrandsSupported(EHairStrandsShaderType::All, Parameters.Platform); }
 	static void ModifyCompilationEnvironment(const FGlobalShaderPermutationParameters& Parameters, FShaderCompilerEnvironment& OutEnvironment)
 	{
 		FGlobalShader::ModifyCompilationEnvironment(Parameters, OutEnvironment);
 		OutEnvironment.SetDefine(TEXT("SHADER_HAIRINTERPOLATION"), 1);
 		OutEnvironment.SetDefine(TEXT("GROUP_SIZE"), GetGroupSize());
+		OutEnvironment.CompilerFlags.Add(CFLAG_Wave32);
 	}
 };
 
@@ -530,6 +535,8 @@ void AddHairStrandsInterpolationPass(
 	const FShaderPrintData* ShaderPrintData,
 	const FHairGroupInstance* Instance,
 	const uint32 VertexCount,
+	const uint32 CurveCount,
+	const uint32 MaxPointPerCurve,
 	const int32 MeshLODIndex,
 	const float HairLengthScale,
 	const EHairInterpolationType HairInterpolationType,
@@ -545,12 +552,14 @@ void AddHairStrandsInterpolationPass(
 	const FHairStrandsDeformedRootResource* SimDeformedRootResources,
 	const FRDGBufferSRVRef& RenRestPosePositionBuffer,
 	const FRDGBufferSRVRef& RenPointToCurveBuffer,
+	const FRDGBufferSRVRef& RenCurveBuffer,
+	const FRDGBufferSRVRef& SimCurveBuffer,
 	const bool bUseSingleGuide,
-	const FRDGBufferSRVRef& InterpolationBuffer,
+	const FRDGBufferSRVRef& CurveInterpolationBuffer,
+	const FRDGBufferSRVRef& PointInterpolationBuffer,
 	const FRDGBufferSRVRef& SimRestPosePositionBuffer,
 	const FRDGBufferSRVRef& SimDeformedPositionBuffer,
 	const FRDGBufferSRVRef& SimRootPointIndexBuffer,
-	const FRDGBufferSRVRef& SimPointToCurveBuffer,
 	const FRDGBufferSRVRef& RenDeformerPositionBuffer,
 	FRDGBufferUAVRef& OutRenPositionBuffer,
 	const FHairStrandsDeformedRootResource::FLOD::EFrameType DeformedFrame)
@@ -559,10 +568,12 @@ void AddHairStrandsInterpolationPass(
 	Parameters->RenRestPosePositionBuffer = RenRestPosePositionBuffer;
 	Parameters->SimRestPosePositionBuffer = SimRestPosePositionBuffer;
 	Parameters->SimDeformedPositionBuffer = SimDeformedPositionBuffer;
-	Parameters->InterpolationBuffer = InterpolationBuffer;
+	Parameters->CurveInterpolationBuffer = CurveInterpolationBuffer;
+	Parameters->PointInterpolationBuffer = PointInterpolationBuffer;
 	Parameters->OutRenDeformedPositionBuffer = OutRenPositionBuffer;
 
 	Parameters->VertexCount = VertexCount;
+	Parameters->CurveCount = CurveCount;
 	Parameters->InRenHairPositionOffset = (FVector3f)InRenHairWorldOffset;
 	Parameters->InSimHairPositionOffset = (FVector3f)InSimHairWorldOffset;
 
@@ -570,19 +581,18 @@ void AddHairStrandsInterpolationPass(
 	Parameters->OutRenHairPositionOffsetBuffer = OutRenHairPositionOffsetBuffer;
 
 	Parameters->SimRootPointIndexBuffer = SimRootPointIndexBuffer;
-	Parameters->SimPointToCurveBuffer = SimPointToCurveBuffer;
 	Parameters->RenPointToCurveBuffer = RenPointToCurveBuffer;
 	
 	Parameters->LocalToWorldMatrix = FMatrix44f(Instance->LocalToWorld.ToMatrixWithScale());		// LWC_TODO: Precision loss
 	Parameters->HairLengthScale = HairLengthScale;
 
-	// Only needed if DEBUG_ENABLE is manually enabledin HairStrandsInterpolation.usf. This is for manual debug purpose only
-	#if 0
-	if (ShaderPrintData && ShaderPrintData->IsValid())
+	Parameters->RenCurveBuffer = RenCurveBuffer;
+	Parameters->SimCurveBuffer = SimCurveBuffer;
+
+	if (ShaderPrintData)
 	{
 		ShaderPrint::SetParameters(GraphBuilder, *ShaderPrintData, Parameters->ShaderDrawParameters);
 	}
-	#endif
 
 	const bool bSupportDynamicMesh = 
 		RenRestRootResources &&
@@ -632,8 +642,6 @@ void AddHairStrandsInterpolationPass(
 	}
 
 	const bool bHasLocalDeformation = Instance->Guides.bIsSimulationEnable || bSupportGlobalInterpolation || Instance->Guides.bIsDeformationEnable || Instance->Guides.bIsSimulationCacheEnable;
-	const bool bCullingEnable = IsHairStrandContinuousDecimationReorderingEnabled() ? false : (InstanceGeometryType == EHairGeometryType::Strands && CullingData.bCullingResultAvailable); 	// TODO: improve reordering so that culling can be used effectively
-	Parameters->HairStrandsVFTODO_bCullingEnable = bCullingEnable ? 1 : 0;
 
 	// Select dynamic geometry permutation, based on the Simulation/RBF/InterpolationType
 	int32 DynamicGeometryType = 0;
@@ -652,40 +660,23 @@ void AddHairStrandsInterpolationPass(
 		if (!bSupportDynamicMesh &&  bHasLocalDeformation) DynamicGeometryType = 0; // INTERPOLATION_RIGID
 	}
 
+	// Compute the dispatch information for pass dispatching work per curve
+	const FPointPerCurveDispatchInfo DispatchInfo = GetPointPerCurveDispatchInfo(MaxPointPerCurve, CurveCount, FHairInterpolationCS::GetGroupSize());
+	Parameters->DispatchCountX = DispatchInfo.DispatchCount.X;
+
 	FHairInterpolationCS::FPermutationDomain PermutationVector;
 	PermutationVector.Set<FHairInterpolationCS::FDynamicGeometry>(DynamicGeometryType);
 	PermutationVector.Set<FHairInterpolationCS::FSimulation>(bHasLocalDeformation);
 	PermutationVector.Set<FHairInterpolationCS::FSingleGuide>(bUseSingleGuide);
-	PermutationVector.Set<FHairInterpolationCS::FCulling>(bCullingEnable);
 	PermutationVector.Set<FHairInterpolationCS::FDeformer>(bSupportDeformer);
-
+	PermutationVector.Set<FHairInterpolationCS::FPointPerCurve>(DispatchInfo.PointPerCurve);
 	TShaderMapRef<FHairInterpolationCS> ComputeShader(ShaderMap, PermutationVector);
-
-	if (bCullingEnable)
-	{
-		Parameters->HairStrandsVFTODO_CullingIndirectBuffer = CullingData.HairStrandsVF_CullingIndirectBuffer.SRV;
-		Parameters->HairStrandsVFTODO_CullingIndexBuffer = CullingData.HairStrandsVF_CullingIndexBuffer.SRV;
-		Parameters->HairStrandsVFTODO_CullingRadiusScaleBuffer = CullingData.HairStrandsVF_CullingRadiusScaleBuffer.SRV;
-		Parameters->HairStrandsVFTODO_CullingIndirectBufferArgs = CullingData.HairStrandsVF_CullingIndirectBuffer.Buffer;
-
-		FComputeShaderUtils::AddPass(
-			GraphBuilder,
-			RDG_EVENT_NAME("HairStrands::Interpolation(culling=on)"),
-			ComputeShader, 
-			Parameters,
-			CullingData.HairStrandsVF_CullingIndirectBuffer.Buffer,
-			0);
-	}
-	else
-	{
-		const FIntVector DispatchCount(FMath::DivideAndRoundUp(VertexCount, FHairInterpolationCS::GetGroupSize()), 1, 1);
-		FComputeShaderUtils::AddPass(
-			GraphBuilder,
-			RDG_EVENT_NAME("HairStrands::Interpolation(culling=off)"),
-			ComputeShader,
-			Parameters,
-			DispatchCount);
-	}
+	FComputeShaderUtils::AddPass(
+		GraphBuilder,
+		RDG_EVENT_NAME("HairStrands::Interpolation(culling=off, PerCurve=%d, SingleGuide=%d, DynGeom:%d)", DispatchInfo.PointPerCurve, bUseSingleGuide ? 1u : 0u, DynamicGeometryType > 0 ? 1u : 0u),
+		ComputeShader,
+		Parameters,
+		DispatchInfo.DispatchCount);
 }
 
 ///////////////////////////////////////////////////////////////////////////////////////////////////
@@ -701,19 +692,23 @@ class FHairPatchAttributeCS : public FGlobalShader
 
 	BEGIN_SHADER_PARAMETER_STRUCT(FParameters, )
 		SHADER_PARAMETER(uint32, CurveCount)
+		SHADER_PARAMETER(uint32, bUseSingleGuide)
 		SHADER_PARAMETER_STRUCT_INCLUDE(FHairStrandsInstanceAttributeParameters, RenAttributes)
 		SHADER_PARAMETER_RDG_BUFFER_SRV(Buffer, RenCurveBuffer)
-		SHADER_PARAMETER_RDG_BUFFER_SRV(ByteAddressBuffer, InterpolationBuffer)
+		SHADER_PARAMETER_RDG_BUFFER_SRV(ByteAddressBuffer, CurveInterpolationBuffer)
+		SHADER_PARAMETER_RDG_BUFFER_SRV(ByteAddressBuffer, PointInterpolationBuffer)
 		SHADER_PARAMETER_RDG_BUFFER_SRV(Buffer, RenCurveToClusterIdBuffer)
 		SHADER_PARAMETER_RDG_BUFFER_UAV(RWByteAddressBuffer, OutRenAttributeBuffer)
 	END_SHADER_PARAMETER_STRUCT()
 
 public:
+	static uint32 GetGroupSize() { return 1024u; }
 	static bool ShouldCompilePermutation(const FGlobalShaderPermutationParameters& Parameters) { return IsHairStrandsSupported(EHairStrandsShaderType::Strands, Parameters.Platform); }
 	static void ModifyCompilationEnvironment(const FGlobalShaderPermutationParameters& Parameters, FShaderCompilerEnvironment& OutEnvironment)
 	{
 		FGlobalShader::ModifyCompilationEnvironment(Parameters, OutEnvironment);
 		OutEnvironment.SetDefine(TEXT("SHADER_PATCHATTRIBUTE"), 1);
+		OutEnvironment.SetDefine(TEXT("GROUP_SIZE"), GetGroupSize());
 	}
 };
 
@@ -730,19 +725,21 @@ void AddPatchAttributePass(
 	const FRDGBufferRef& RenAttributeBuffer,
 	const FRDGBufferSRVRef& RenCurveBuffer,
 	const FRDGBufferSRVRef& RenCurveToClusterIdBuffer,
-	const FRDGBufferSRVRef& InterpolationBuffer,
+	const FRDGBufferSRVRef& CurveInterpolationBuffer,
+	const FRDGBufferSRVRef& PointInterpolationBuffer,
 	FRDGImportedBuffer& OutRenAttributeBuffer)
 {
 	// Sanity check
 	check(Mode == EHairPatchAttribute::ClusterInfluence || Mode == EHairPatchAttribute::GuideInflucence)
 	if (Mode == EHairPatchAttribute::GuideInflucence && !bSimulation) return;
 
-	const FIntVector DispatchCount = FIntVector(FMath::DivideAndRoundUp(CurveCount, 1024u), 1, 1);
+	const FIntVector DispatchCount = FIntVector(FMath::DivideAndRoundUp(CurveCount, FHairPatchAttributeCS::GetGroupSize()), 1, 1);
 
 	// First copy all the rendering attributes into the new bufffer, before overriding some of them
 	AddCopyBufferPass(GraphBuilder, OutRenAttributeBuffer.Buffer, RenAttributeBuffer);
 
 	FHairPatchAttributeCS::FParameters* Parameters = GraphBuilder.AllocParameters<FHairPatchAttributeCS::FParameters>();
+	Parameters->bUseSingleGuide = bUseSingleGuide ? 1u : 0u;
 	Parameters->CurveCount = CurveCount;
 	Parameters->RenCurveBuffer = RenCurveBuffer;
 	Parameters->RenCurveToClusterIdBuffer = RenCurveToClusterIdBuffer;
@@ -750,7 +747,8 @@ void AddPatchAttributePass(
 
 	if (bSimulation)
 	{
-		Parameters->InterpolationBuffer  = InterpolationBuffer;
+		Parameters->CurveInterpolationBuffer  = CurveInterpolationBuffer;
+		Parameters->PointInterpolationBuffer  = PointInterpolationBuffer;
 	}
 
 	Parameters->OutRenAttributeBuffer = OutRenAttributeBuffer.UAV;
diff --git a/Engine/Plugins/Runtime/HairStrands/Source/HairStrandsCore/Public/GroomAsset.h b/Engine/Plugins/Runtime/HairStrands/Source/HairStrandsCore/Public/GroomAsset.h
--- a/Engine/Plugins/Runtime/HairStrands/Source/HairStrandsCore/Public/GroomAsset.h
+++ b/Engine/Plugins/Runtime/HairStrands/Source/HairStrandsCore/Public/GroomAsset.h
@@ -245,8 +245,8 @@ FORCEINLINE uint32 GetDataSize(const FHairStrandsBulkData& BulkData)
 FORCEINLINE uint32 GetDataSize(const FHairStrandsInterpolationBulkData& InterpolationBulkData) 	
 {
 	uint32 Total = 0;
-	Total += InterpolationBulkData.Data.Interpolation.IsBulkDataLoaded()	? InterpolationBulkData.Data.Interpolation.GetBulkDataSize() : 0;
-	Total += InterpolationBulkData.Data.SimRootPointIndex.IsBulkDataLoaded()? InterpolationBulkData.Data.SimRootPointIndex.GetBulkDataSize() : 0;
+	Total += InterpolationBulkData.Data.CurveInterpolation.IsBulkDataLoaded()? InterpolationBulkData.Data.CurveInterpolation.GetBulkDataSize() : 0;
+	Total += InterpolationBulkData.Data.PointInterpolation.IsBulkDataLoaded()? InterpolationBulkData.Data.PointInterpolation.GetBulkDataSize() : 0;
 	return Total;
 }
 
diff --git a/Engine/Plugins/Runtime/HairStrands/Source/HairStrandsCore/Public/GroomResources.h b/Engine/Plugins/Runtime/HairStrands/Source/HairStrandsCore/Public/GroomResources.h
--- a/Engine/Plugins/Runtime/HairStrands/Source/HairStrandsCore/Public/GroomResources.h
+++ b/Engine/Plugins/Runtime/HairStrands/Source/HairStrandsCore/Public/GroomResources.h
@@ -297,21 +297,21 @@ struct FHairStrandsInterpolationResource : public FHairCommonResource
 	virtual FHairStrandsBulkCommon* InternalGetBulkData() override { return &BulkData; }
 
 	/* Get the resource name */
-	virtual FString GetFriendlyName() const override { return TEXT("FHairStrandsInterplationResource"); }
+	virtual FString GetFriendlyName() const override { return TEXT("FHairStrandsInterpolationResource"); }
 
 	/* Return the memory size for GPU resources */
 	uint32 GetResourcesSize() const
 	{
 		uint32 Total = 0;
-		Total += GetBufferTotalNumBytes(InterpolationBuffer);
-		Total += GetBufferTotalNumBytes(SimRootPointIndexBuffer);
+		Total += GetBufferTotalNumBytes(CurveInterpolationBuffer);
+		Total += GetBufferTotalNumBytes(PointInterpolationBuffer);
 		return Total;
 	}
 
 	bool UseSingleGuide() const { return (BulkData.Header.Flags & FHairStrandsInterpolationBulkData::DataFlags_HasSingleGuideData) != 0; }
 
-	FRDGExternalBuffer InterpolationBuffer;
-	FRDGExternalBuffer SimRootPointIndexBuffer;
+	FRDGExternalBuffer CurveInterpolationBuffer;
+	FRDGExternalBuffer PointInterpolationBuffer;
 
 	/* Reference to the hair strands interpolation render data */
 	FHairStrandsInterpolationBulkData& BulkData;
diff --git a/Engine/Plugins/Runtime/HairStrands/Source/HairStrandsCore/Public/HairStrandsDatas.h b/Engine/Plugins/Runtime/HairStrands/Source/HairStrandsCore/Public/HairStrandsDatas.h
--- a/Engine/Plugins/Runtime/HairStrands/Source/HairStrandsCore/Public/HairStrandsDatas.h
+++ b/Engine/Plugins/Runtime/HairStrands/Source/HairStrandsCore/Public/HairStrandsDatas.h
@@ -585,6 +585,7 @@ struct HAIRSTRANDSCORE_API FHairStrandsBulkData : FHairStrandsBulkCommon
 	float  GetMaxRadius() const { return Header.MaxRadius; }
 	FVector GetPositionOffset() const { return Header.BoundingBox.GetCenter(); }
 	const FBox& GetBounds() const { return Header.BoundingBox; }
+	uint32 GetSize() const;
 
 	uint32 GetCurveAttributeSizeInBytes(uint32 InCurveCount=HAIR_MAX_NUM_CURVE_PER_GROUP) const	{ return InCurveCount > 0 ? FMath::DivideAndRoundUp(FMath::Min(Header.CurveCount, InCurveCount), Header.Strides.CurveAttributeChunkElementCount) * Header.Strides.CurveAttributeChunkStride : 0; }
 	uint32 GetPointAttributeSizeInBytes(uint32 InPointCount=HAIR_MAX_NUM_POINT_PER_GROUP) const	{ return InPointCount > 0 ? FMath::DivideAndRoundUp(FMath::Min(Header.PointCount, InPointCount), Header.Strides.PointAttributeChunkElementCount) * Header.Strides.PointAttributeChunkStride : 0; }
@@ -642,27 +643,31 @@ struct HAIRSTRANDSCORE_API FHairStrandsBulkData : FHairStrandsBulkCommon
 struct HAIRSTRANDSCORE_API FHairStrandsInterpolationDatas
 {
 	/** Set the number of interpolated points */
-	void SetNum(const uint32 NumPoints);
+	void SetNum(const uint32 InCurveCount, const uint32 InPointCount);
 
 	/** Reset the interpolated points to 0 */
 	void Reset();
 
 	/** Get the number of interpolated points */
-	uint32 Num() const { return PointsSimCurvesVertexIndex.Num(); }
+	uint32 GetPointCount() const { return PointSimIndices.Num(); }
+	uint32 GetCurveCount() const { return CurveSimIndices.Num(); }
 
-	bool IsValid() const { return PointsSimCurvesIndex.Num() > 0; }
+	bool IsValid() const { return CurveSimIndices.Num() > 0; }
 
-	/** Simulation curve indices, ordered by closest influence */
-	TArray<FIntVector> PointsSimCurvesIndex;
+	/** Simulation curve indices, ordered by closest influence (per-curve data)*/
+	TArray<FIntVector2> CurveSimIndices;
 
-	/** Closest vertex indices on simulation curve, ordered by closest influence */
-	TArray<FIntVector> PointsSimCurvesVertexIndex;
+	/** Weight of simulation curve, ordered by closest influence (per-curve data)*/
+	TArray<FVector2f> CurveSimWeights;
 
-	/** Lerp value between the closest vertex indices and the next one, ordered by closest influence */
-	TArray<FVector3f> PointsSimCurvesVertexLerp;
+	/** (Global) Index of the root point of the curve*/
+	TArray<FIntVector2> CurveSimRootPointIndex;
 
-	/** Weight of vertex indices on simulation curve, ordered by closest influence */
-	TArray<FVector3f>	PointsSimCurvesVertexWeights;
+	/** Closest vertex (local, 0..255) indices on simulation curve, ordered by closest influence (per-point data)*/
+	TArray<FIntVector2> PointSimIndices;
+
+	/** Lerp value between the closest vertex indices and the next one, ordered by closest influence (per-point data)*/
+	TArray<FVector2f> PointSimLerps;
 
 	/** True, if interpolation data are built using a single guide */
 	bool bUseUniqueGuide = false;
@@ -683,24 +688,25 @@ struct HAIRSTRANDSCORE_API FHairStrandsInterpolationBulkData : FHairStrandsBulkC
 	virtual uint32 GetResourceCount() const override;
 	virtual void GetResources(FQuery& Out) override;
 	uint32 GetPointCount() const { return Header.PointCount; };
+	uint32 GetSize() const;
 
 	struct FHeader
 	{
 		uint32 Flags = 0;
 		uint32 PointCount = 0;
-		uint32 SimPointCount = 0;
+		uint32 CurveCount = 0;
 
 		struct FStrides
 		{
-			uint32 InterpolationStride = 0;
-			uint32 SimRootPointIndexStride = 0;
+			uint32 CurveInterpolationStride = 0;
+			uint32 PointInterpolationStride = 0;
 		} Strides;
 	} Header;
 
 	struct FData
 	{
-		FHairBulkContainer Interpolation;		// FHairStrandsInterpolationFormat  - Per-rendering-vertex interpolation data (closest guides, weight factors, ...). Data for a 1 or 3 guide(s))
-		FHairBulkContainer SimRootPointIndex;	// FHairStrandsRootIndexFormat      - Per-rendering-vertex index of the sim-root vertex
+		FHairBulkContainer CurveInterpolation;	// FHairStrandsInterpolationFormat  - Per-rendering-curve interpolation data (closest guides, weight factors, ...). Data for a 1 or 2 guide(s))
+		FHairBulkContainer PointInterpolation;	// FHairStrandsInterpolationFormat  - Per-rendering-point interpolation data (local point index, lerp factors,...). Data for a 1 or 2 guide(s))
 	} Data;
 };
 
@@ -733,6 +739,7 @@ struct HAIRSTRANDSCORE_API FHairStrandsClusterBulkData : FHairStrandsBulkCommon
 	virtual uint32 GetResourceCount() const override;
 	virtual void GetResources(FQuery& Out) override;
 	uint32 GetCurveCount(float InLODIndex) const;
+	uint32 GetSize() const;
 
 	bool IsValid() const { return Header.ClusterCount > 0 && Header.PointCount > 0; }
 
@@ -826,6 +833,7 @@ struct FHairStrandsRootBulkData : FHairStrandsBulkCommon
 	virtual void SerializeHeader(FArchive& Ar, UObject* Owner) override;
 	virtual uint32 GetResourceCount() const override;
 	virtual void GetResources(FQuery& Out) override;
+	uint32 GetSize() const;
 
 	void Reset();
 	virtual void ResetLoadedSize() override;
diff --git a/Engine/Plugins/Runtime/HairStrands/Source/HairStrandsCore/Public/HairStrandsInterpolation.h b/Engine/Plugins/Runtime/HairStrands/Source/HairStrandsCore/Public/HairStrandsInterpolation.h
--- a/Engine/Plugins/Runtime/HairStrands/Source/HairStrandsCore/Public/HairStrandsInterpolation.h
+++ b/Engine/Plugins/Runtime/HairStrands/Source/HairStrandsCore/Public/HairStrandsInterpolation.h
@@ -119,6 +119,8 @@ void AddHairStrandsInterpolationPass(
 	const FShaderPrintData* ShaderPrintData,
 	const FHairGroupInstance* Instance,
 	const uint32 VertexCount,
+	const uint32 CurveCount,
+	const uint32 MaxPointPerCurve,
 	const int32 MeshLODIndex,
 	const float HairLengthScale,
 	const EHairInterpolationType HairInterpolationType,
@@ -134,12 +136,14 @@ void AddHairStrandsInterpolationPass(
 	const FHairStrandsDeformedRootResource* SimDeformedRootResources,
 	const FRDGBufferSRVRef& RenRestPosePositionBuffer,
 	const FRDGBufferSRVRef& RenPointToCurveBuffer,
+	const FRDGBufferSRVRef& RenCurveBuffer,
+	const FRDGBufferSRVRef& SimCurveBuffer,
 	const bool bUseSingleGuide,
-	const FRDGBufferSRVRef& InterpolationBuffer,
+	const FRDGBufferSRVRef& CurveInterpolationBuffer,
+	const FRDGBufferSRVRef& PoinInterpolationBuffer,
 	const FRDGBufferSRVRef& SimRestPosePositionBuffer,
 	const FRDGBufferSRVRef& SimDeformedPositionBuffer,
 	const FRDGBufferSRVRef& SimRootPointIndexBuffer,
-	const FRDGBufferSRVRef& SimPointToCurveBuffer,
 	const FRDGBufferSRVRef& RenDeformerPositionBuffer,
 	FRDGBufferUAVRef& OutRenPositionBuffer,
 	const FHairStrandsDeformedRootResource::FLOD::EFrameType DeformedFrame);
@@ -172,7 +176,8 @@ void AddPatchAttributePass(
 	const FRDGBufferRef& RenAttributeBuffer,
 	const FRDGBufferSRVRef& RenCurveBuffer,
 	const FRDGBufferSRVRef& RenCurveToClusterIdBuffer,
-	const FRDGBufferSRVRef& InterpolationBuffer,
+	const FRDGBufferSRVRef& CurveInterpolationBuffer,
+	const FRDGBufferSRVRef& PointInterpolationBuffer,
 	FRDGImportedBuffer& OutRenAttributeBuffer);
 
 void AddTransferPositionPass(
diff --git a/Engine/Shaders/Private/HairStrands/HairCardsDeformation.usf b/Engine/Shaders/Private/HairStrands/HairCardsDeformation.usf
--- a/Engine/Shaders/Private/HairStrands/HairCardsDeformation.usf
+++ b/Engine/Shaders/Private/HairStrands/HairCardsDeformation.usf
@@ -89,9 +89,8 @@ void MainCS(uint DispatchThreadId : SV_DispatchThreadID)
 	const float3 GuideDeformedPositionOffset = ReadCardPositionOffset(GuideDeformedPositionOffsetBuffer, InstanceRegisteredIndex, CardLODIndex);
 
 #if PERMUTATION_DYNAMIC_GEOMETRY == 0
-	// Simple deformation based on world space offset
 
-	const FGuideData GuideData = UnpackGuideData(CardsInterpolationBuffer.Load(VertexId * HAIR_INTERPOLATION_1GUIDE_STRIDE));
+	const FCardsGuideData GuideData = UnpackCardGuideData(CardsInterpolationBuffer, VertexId);
 	const uint GuideIndex0 = GuideData.VertexIndices.x;
 	const uint GuideIndex1 = min(GuideIndex0+1, GuideVertexCount-1);
 	const float GuideLerp = GuideData.VertexLerps.x;
@@ -138,7 +137,7 @@ void MainCS(uint DispatchThreadId : SV_DispatchThreadID)
 #elif PERMUTATION_DYNAMIC_GEOMETRY == 1
 	// Compute the deformation offset in triangle local space
 
-	const FGuideData GuideData = UnpackGuideData(CardsInterpolationBuffer.Load(VertexId * HAIR_INTERPOLATION_1GUIDE_STRIDE));
+	const FCardsGuideData GuideData = UnpackCardGuideData(CardsInterpolationBuffer, VertexId);
 	const uint GuideIndex0 = GuideData.VertexIndices.x;
 	const uint GuideIndex1 = min(GuideIndex0 + 1, GuideVertexCount - 1);
 	const float GuideLerp = GuideData.VertexLerps.x;
diff --git a/Engine/Shaders/Private/HairStrands/HairStrandsInterpolation.usf b/Engine/Shaders/Private/HairStrands/HairStrandsInterpolation.usf
--- a/Engine/Shaders/Private/HairStrands/HairStrandsInterpolation.usf
+++ b/Engine/Shaders/Private/HairStrands/HairStrandsInterpolation.usf
@@ -22,40 +22,32 @@
 
 #pragma warning (default:7203) // downgrade 'not enough registers available for the entire program.' to a warning instead of error
 
-
+// Currently disable multi-guide interpolation to reduce VGPR pressure and increase occupancy
+#define FORCE_USE_SINGLE_GUIDE 1
 
 ///////////////////////////////////////////////////////////////////////////////////////////////////
 // Guides
 
 #if SHADER_HAIRINTERPOLATION || SHADER_PATCHATTRIBUTE
-ByteAddressBuffer InterpolationBuffer;
+ByteAddressBuffer CurveInterpolationBuffer;
+ByteAddressBuffer PointInterpolationBuffer;
 #endif
 
 struct FGuideDataWithOffset
 {
-	uint3  VertexIndices;
-	float3 VertexLerps;
-	float3 VertexWeights;
+	uint2  VertexIndices;
+	float2 VertexLerps;
+	float2 VertexWeights;
 	float3 OutPositionOffset;
 };
 
-FGuideDataWithOffset GetGuideData(uint InVertexIndex, bool bSingleGuide, float3 OutPositionOffset)
+FGuideDataWithOffset GetGuideData(FGuideCurveData InCurveData, FGuidePointData InPointData, float3 OutPositionOffset)
 {
 #if PERMUTATION_SIMULATION > 0
-	FGuideData GuideData;
-	if (bSingleGuide)
-	{
-		GuideData = UnpackGuideData(InterpolationBuffer.Load(InVertexIndex * HAIR_INTERPOLATION_1GUIDE_STRIDE));
-	}
-	else
-	{
-		GuideData = UnpackGuideData(InterpolationBuffer.Load4(InVertexIndex * HAIR_INTERPOLATION_3GUIDE_STRIDE));
-	}
-
 	FGuideDataWithOffset Out;
-	Out.VertexIndices     = GuideData.VertexIndices;
-	Out.VertexLerps       = GuideData.VertexLerps;
-	Out.VertexWeights     = GuideData.VertexWeights;
+	Out.VertexIndices     = InCurveData.RootPointIndices + InPointData.PointIndices;
+	Out.VertexLerps       = InPointData.PointLerps;
+	Out.VertexWeights     = InCurveData.CurveWeights;
 	Out.OutPositionOffset = OutPositionOffset;
 	return Out;
 #else
@@ -72,8 +64,13 @@ FGuideDataWithOffset GetGuideData(uint InVertexIndex, bool bSingleGuide, float3
 #define INTERPOLATION_SKINNING_ROTATION 3
 #define INTERPOLATION_SKINNING_TRANSLATION_AND_ROTATION 4
 
+#ifndef GROUP_SIZE
+#error GROUP_SIZE needs to be Undefined
+#endif
+
 #if SHADER_HAIRINTERPOLATION
 
+uint DispatchCountX;
 uint VertexCount;
 uint HairDebugMode;
 float HairLengthScale;
@@ -124,10 +121,6 @@ StructuredBuffer<float4> OutRenHairPositionOffsetBuffer;
 float3 GetSimPositionOffset() { return ReadSimPositionOffset(OutSimHairPositionOffsetBuffer, InstanceRegisteredIndex); }
 float3 GetRenPositionOffset() { return ReadRenPositionOffset(OutRenHairPositionOffsetBuffer, InstanceRegisteredIndex); }
 
-#ifndef GROUP_SIZE
-	#error Unknown group size	
-#endif
-
 ///////////////////////////////////////////////////////////////////////////////////////////////////
 // Deformation offset
 
@@ -208,32 +201,28 @@ float3 ComputeRotateGeometryOffset(
 
 struct FRootTriangleData
 {
-	uint RootIndex;
-	uint TriangleIndex;
 	float3 RootBarycentric;
 	FHairMeshTriangle RestTriangle;
 	FHairMeshTriangle DeformedTriangle;
 };
 
-FRootTriangleData GetRootTriangleData(uint VertexIndex, bool bSim)
+FRootTriangleData GetRootTriangleData(uint CurveIndex, bool bSim)
 {
 	FRootTriangleData Out = (FRootTriangleData)0;
 	#if PERMUTATION_DYNAMIC_GEOMETRY >= 1 
 	if (bSim)
 	{
-		Out.RootIndex		 = SimPointToCurveBuffer[VertexIndex];
-		Out.TriangleIndex	 = SimRootToUniqueTriangleIndexBuffer[Out.RootIndex];
-		Out.RootBarycentric  = UnpackBarycentrics(SimRootBarycentricBuffer[Out.RootIndex]);
-		Out.RestTriangle	 = GetTriangleTransformation(Out.TriangleIndex, SimRootRestPositionBuffer);
-		Out.DeformedTriangle = GetTriangleTransformation(Out.TriangleIndex, SimRootDeformedPositionBuffer);
+		const uint TriangleIndex = SimRootToUniqueTriangleIndexBuffer[CurveIndex];
+		Out.RootBarycentric  = UnpackBarycentrics(SimRootBarycentricBuffer[CurveIndex]);
+		Out.RestTriangle	 = GetTriangleTransformation(TriangleIndex, SimRootRestPositionBuffer);
+		Out.DeformedTriangle = GetTriangleTransformation(TriangleIndex, SimRootDeformedPositionBuffer);
 	}
 	else
 	{
-		Out.RootIndex		 = RenPointToCurveBuffer[VertexIndex];
-		Out.TriangleIndex	 = RenRootToUniqueTriangleIndexBuffer[Out.RootIndex];
-		Out.RootBarycentric  = UnpackBarycentrics(RenRootBarycentricBuffer[Out.RootIndex]);
-		Out.RestTriangle	 = GetTriangleTransformation(Out.TriangleIndex, RenRootRestPositionBuffer);
-		Out.DeformedTriangle = GetTriangleTransformation(Out.TriangleIndex, RenRootDeformedPositionBuffer);
+		const uint TriangleIndex = RenRootToUniqueTriangleIndexBuffer[CurveIndex];
+		Out.RootBarycentric  = UnpackBarycentrics(RenRootBarycentricBuffer[CurveIndex]);
+		Out.RestTriangle	 = GetTriangleTransformation(TriangleIndex, RenRootRestPositionBuffer);
+		Out.DeformedTriangle = GetTriangleTransformation(TriangleIndex, RenRootDeformedPositionBuffer);
 	}
 	#endif
 	return Out;
@@ -245,11 +234,11 @@ float3 ApplyRigidInterpolation(FGuideDataWithOffset GuideData, float3 InPosition
 {
 	float3 CurrOffset = 0;
 #if PERMUTATION_SIMULATION == 1
-	#if PERMUTATION_USE_SINGLE_GUIDE == 1
+	#if PERMUTATION_USE_SINGLE_GUIDE == 1 || FORCE_USE_SINGLE_GUIDE
 	const uint K = 0;
 	#else
 	[unroll]
-	for (uint K = 0; K < 3; ++K)
+	for (uint K = 0; K < HAIR_INTERPOLATION_MAX_GUIDE_COUNT; ++K)
 	#endif
 	{
 		const uint GuideIndex0 = GuideData.VertexIndices[K];
@@ -267,23 +256,23 @@ float3 ApplyRigidInterpolation(FGuideDataWithOffset GuideData, float3 InPosition
 
 ///////////////////////////////////////////////////////////////////////////////////////////////////
 // Skinning with offset
-float3 ApplySkinningWithOffset(const FGuideDataWithOffset GuideData, const FRootTriangleData RenTri, float3 InPosition)
+float3 ApplySkinningWithOffset(const FGuideDataWithOffset GuideData, const FRootTriangleData RenTri, const FRootTriangleData SimTris[2], float3 InPosition)
 {
 	float3 CurrOffset = 0;
 
 	// Compute the simulation offset in hair local space (i.e., triangle)
 #if PERMUTATION_SIMULATION == 1
-	#if PERMUTATION_USE_SINGLE_GUIDE == 1
+	#if PERMUTATION_USE_SINGLE_GUIDE == 1 || FORCE_USE_SINGLE_GUIDE
 	const uint K = 0;
 	#else
 	[unroll]
-	for (uint K = 0; K < 3; ++K)
+	for (uint K = 0; K < HAIR_INTERPOLATION_MAX_GUIDE_COUNT; ++K)
 	#endif
 	{
 		const uint GuideIndex0 = GuideData.VertexIndices[K];
 		const uint GuideIndex1 = GuideIndex0 + 1;
 
-		const FRootTriangleData SimTri = GetRootTriangleData(GuideIndex0, true);
+		const FRootTriangleData SimTri = SimTris[K];
 
 		const float3 Offset0 = ComputeDynamicGeometryOffset(GuideIndex0, GuideData.VertexWeights[K], SimTri.RestTriangle, SimTri.DeformedTriangle, SimTri.RootBarycentric, GuideData.OutPositionOffset);
 		const float3 Offset1 = ComputeDynamicGeometryOffset(GuideIndex1, GuideData.VertexWeights[K], SimTri.RestTriangle, SimTri.DeformedTriangle, SimTri.RootBarycentric, GuideData.OutPositionOffset);
@@ -307,11 +296,11 @@ float3 ApplySkinningWithTranslation(const FGuideDataWithOffset GuideData, const
 	float3 CurrOffset = 0;
 
 #if PERMUTATION_SIMULATION == 1
-	#if PERMUTATION_USE_SINGLE_GUIDE == 1
+#if PERMUTATION_USE_SINGLE_GUIDE == 1 || FORCE_USE_SINGLE_GUIDE
 	const uint K = 0;
 	#else
 	[unroll]
-	for (uint K = 0; K < 3; ++K)
+	for (uint K = 0; K < HAIR_INTERPOLATION_MAX_GUIDE_COUNT; ++K)
 	#endif
 	{
 		const uint GuideIndex0 = GuideData.VertexIndices[K];
@@ -330,24 +319,24 @@ float3 ApplySkinningWithTranslation(const FGuideDataWithOffset GuideData, const
 
 ///////////////////////////////////////////////////////////////////////////////////////////////////
 // Skinning with rotation
-float3 ApplySkinningWithRotation(const FGuideDataWithOffset GuideData, const FRootTriangleData RenTri, float3 InPosition)
+float3 ApplySkinningWithRotation(const FGuideDataWithOffset GuideData, const FRootTriangleData RenTri, const FRootTriangleData SimTris[2], float3 InPosition)
 {
 	InPosition -= RenTri.RootBarycentric.x * RenTri.RestTriangle.P0 + RenTri.RootBarycentric.y * RenTri.RestTriangle.P1 + RenTri.RootBarycentric.z * RenTri.RestTriangle.P2;
 
 	float3 CurrOffset = 0;
 
 #if PERMUTATION_SIMULATION == 1
-	#if PERMUTATION_USE_SINGLE_GUIDE == 1
+	#if PERMUTATION_USE_SINGLE_GUIDE == 1 || FORCE_USE_SINGLE_GUIDE
 	const uint K = 0;
 	#else
 	[unroll]
-	for (uint K = 0; K < 3; ++K)
+	for (uint K = 0; K < HAIR_INTERPOLATION_MAX_GUIDE_COUNT; ++K)
 	#endif
 	{
 		const uint GuideIndex0 = GuideData.VertexIndices[K];
 		const uint GuideIndex1 = GuideIndex0 + 1;
 
-		const FRootTriangleData SimTri = GetRootTriangleData(GuideIndex0, true);
+		const FRootTriangleData SimTri = SimTris[K];
 
 		const float3 LocalPoint = VectorToTriangle(InPosition, SimTri.RestTriangle);
 
@@ -364,22 +353,22 @@ float3 ApplySkinningWithRotation(const FGuideDataWithOffset GuideData, const FRo
 
 ///////////////////////////////////////////////////////////////////////////////////////////////////
 // Skinning with translation and rotation
-float3 ApplySkinningWithTranslationAndRotation(const FGuideDataWithOffset GuideData, const FRootTriangleData RenTri, float3 InPosition)
+float3 ApplySkinningWithTranslationAndRotation(const FGuideDataWithOffset GuideData, const FRootTriangleData RenTri, const FRootTriangleData SimTris[2], float3 InPosition)
 {
 	float3 CurrOffset = 0;
 
 #if PERMUTATION_SIMULATION == 1
-	#if PERMUTATION_USE_SINGLE_GUIDE == 1
+	#if PERMUTATION_USE_SINGLE_GUIDE == 1 || FORCE_USE_SINGLE_GUIDE
 	const uint K = 0;
 	#else
 	[unroll]
-	for (uint K = 0; K < 3; ++K)
+	for (uint K = 0; K < HAIR_INTERPOLATION_MAX_GUIDE_COUNT; ++K)
 	#endif
 	{
 		const uint GuideIndex0 = GuideData.VertexIndices[K];
 		const uint GuideIndex1 = GuideIndex0 + 1;
 
-		const FRootTriangleData SimTri = GetRootTriangleData(GuideIndex0, true);
+		const FRootTriangleData SimTri = SimTris[K];  //GetRootTriangleData(GuideIndex0, true);
 
 		const float3 LocalPoint = ToTriangle(InPosition, SimTri.RootBarycentric, SimTri.RestTriangle);
 
@@ -416,90 +405,334 @@ float3 ApplyDeformer(const FRootTriangleData RenTri, float3 InRestPostion, float
 ///////////////////////////////////////////////////////////////////////////////////////////////////
 // Interpolation
 
+uint CurveCount;
+uint CurvePerGroup;
+
+Buffer<uint> SimCurveBuffer;
+Buffer<uint> RenCurveBuffer;
+
+#define TOTALMAXPOINTPERCURVE 32
+//-------------------------------------------------------------------------------------------------
+
+#define CURVE_PER_GROUP (GROUP_SIZE / PERMUTATION_POINT_PER_CURVE)
+#define MAX_CURVE CURVE_PER_GROUP
+#define MAX_POINT PERMUTATION_POINT_PER_CURVE
+
+uint GetGlobalCurveIndex(uint ThreadIndex, uint2 GroupIndex)
+{
+	const uint LinearGroupIndex = GroupIndex.x + GroupIndex.y * DispatchCountX;
+#if PERMUTATION_POINT_PER_CURVE == TOTALMAXPOINTPERCURVE
+	return LinearGroupIndex;
+	#else
+	return ThreadIndex + LinearGroupIndex * CURVE_PER_GROUP;
+	#endif
+}
+
+//-------------------------------------------------------------------------------------------------
+#if PERMUTATION_POINT_PER_CURVE < TOTALMAXPOINTPERCURVE
+
+groupshared uint CurvePointIndex[MAX_CURVE];
+groupshared uint CurvePointCount[MAX_CURVE];
+
+groupshared FRootTriangleData g_RenTri[MAX_CURVE];
+groupshared float3 RenPositionOffset;
+
+groupshared FGuideCurveData SimCurveData[MAX_CURVE];
+groupshared FRootTriangleData g_SimTris[MAX_CURVE][2];
+
 [numthreads(GROUP_SIZE, 1, 1)]
-void MainCS(uint2 DispatchThreadId : SV_DispatchThreadID)
+void MainCS(uint3 DispatchThreadId : SV_DispatchThreadID, uint ThreadIndex : SV_GroupIndex, uint2 GroupIndex : SV_GroupID)
 {
-	// 1. Find vertex index
-	uint VertexIndex = DispatchThreadId.x;
-#if PERMUTATION_CULLING
-	if (HairStrandsVFTODO_bCullingEnable)
+	// 1. Read the Curve data (1 thread per curve)
+	if (ThreadIndex < CURVE_PER_GROUP)
 	{
-		const uint FetchIndex = DispatchThreadId.x;
-		const uint VertexCountAfterCulling = HairStrandsVFTODO_CullingIndirectBuffer[3];
-		if (FetchIndex >= VertexCountAfterCulling)
+		const uint LocalCurveIndex  = ThreadIndex;
+		const uint GlobalCurveIndex = GetGlobalCurveIndex(ThreadIndex, GroupIndex);
+		const bool bValidCurveIndex = GlobalCurveIndex < CurveCount;
+
+		CurvePointIndex[LocalCurveIndex] = 0;
+		CurvePointCount[LocalCurveIndex] = 0;
+		if (bValidCurveIndex)
 		{
-			return;
+			// Curve's point count/offset
+			const FHairCurve Curve = UnpackHairCurve(RenCurveBuffer[GlobalCurveIndex]);
+			CurvePointIndex[LocalCurveIndex] = Curve.PointIndex;
+			CurvePointCount[LocalCurveIndex] = Curve.PointCount;	
+
+			// Load curve interpolation data
+			#if PERMUTATION_SIMULATION
+			SimCurveData[LocalCurveIndex] = UnpackGuideCurveData(CurveInterpolationBuffer, GlobalCurveIndex, PERMUTATION_USE_SINGLE_GUIDE);
+			if (FORCE_USE_SINGLE_GUIDE)
+			{
+				SimCurveData[LocalCurveIndex].CurveWeights.x = 1;
+				SimCurveData[LocalCurveIndex].CurveWeights.y = 0;
+			}
+			#else
+			SimCurveData[LocalCurveIndex] = (FGuideCurveData)0;
+			#endif
+
+			// Load ren triangles
+			g_RenTri[LocalCurveIndex] = GetRootTriangleData(GlobalCurveIndex, false/*bSim*/);
+
+			// Load sim triangles
+			g_SimTris[LocalCurveIndex][0] = (FRootTriangleData)0;
+			g_SimTris[LocalCurveIndex][1] = (FRootTriangleData)0;
+			#if PERMUTATION_SIMULATION
+			g_SimTris[LocalCurveIndex][0] = GetRootTriangleData(SimCurveData[LocalCurveIndex].CurveIndices[0], true/*bSim*/);
+			if (!PERMUTATION_USE_SINGLE_GUIDE)
+			{
+				g_SimTris[LocalCurveIndex][1] = GetRootTriangleData(SimCurveData[LocalCurveIndex].CurveIndices[1], true/*bSim*/);
+			}
+			#endif
 		}
-		VertexIndex = HairStrandsVFTODO_CullingIndexBuffer[FetchIndex];
 	}
-#endif
+	if (ThreadIndex == 0)
+	{
+		RenPositionOffset = GetRenPositionOffset();
+	}
+	GroupMemoryBarrierWithGroupSync();
+	
+	// 2. Apply deformation (1 thread per point)
+	{
+		const uint LocalCurveIndex   = ThreadIndex / uint(MAX_POINT);
+		const uint CurrentPointIndex = ThreadIndex % uint(MAX_POINT);
 
+		const uint GlobalCurveIndex = GetGlobalCurveIndex(ThreadIndex, GroupIndex);
+		if (GlobalCurveIndex < CurveCount)
+		{		
+			if (CurrentPointIndex < CurvePointCount[LocalCurveIndex])
+			{
+				const uint VertexIndex = CurrentPointIndex + CurvePointIndex[LocalCurveIndex];
+
+				#if PERMUTATION_SIMULATION
+				const FGuidePointData SimPointData = UnpackGuidePointData(PointInterpolationBuffer, VertexIndex, PERMUTATION_USE_SINGLE_GUIDE);
+				const FGuideDataWithOffset GuideData = GetGuideData(SimCurveData[LocalCurveIndex], SimPointData, GetSimPositionOffset());
+				#else
+				const FGuidePointData SimPointData = (FGuidePointData)0;
+				const FGuideDataWithOffset GuideData = (FGuideDataWithOffset)0;;
+				#endif
+
+				// Remap from LDS to local variable to ease the setup
+				FRootTriangleData RenTri = g_RenTri[LocalCurveIndex];
+				FRootTriangleData SimTris[2];
+				#if PERMUTATION_SIMULATION
+				SimTris[0] = g_SimTris[LocalCurveIndex][0];
+				SimTris[1] = g_SimTris[LocalCurveIndex][1];
+				#else
+				SimTris[0] = (FRootTriangleData)0;
+				SimTris[1] = (FRootTriangleData)0;
+				#endif
+
+				// Manual decoding of the rest position
+				uint4 PackedRestPosition = RenRestPosePositionBuffer[VertexIndex];
+				const float3 RestPosition = UnpackHairControlPointPosition(PackedRestPosition, InRenHairPositionOffset);
+				
+				float3 OutPosition = 0.0f;
+
+				// 1. Rigid transformation
+				// ControlPoint is in the local hair referential
+				// CurrOffset takes only translation component into account, and is done in object local space (vs. triangle/hair local space)
+				#if PERMUTATION_DYNAMIC_GEOMETRY == INTERPOLATION_RIGID
+				{
+					OutPosition = ApplyRigidInterpolation(GuideData, RestPosition);
+				}
+				// 2. Skin transformation
+				// Apply dynamic mesh deformation (translation / rotation)
+				#elif PERMUTATION_DYNAMIC_GEOMETRY == INTERPOLATION_SKINNING_OFFSET
+				{
+					OutPosition = ApplySkinningWithOffset(GuideData, RenTri, SimTris, RestPosition);
+				}
+				// 3. Linear blend skinning (translation)
+				#elif PERMUTATION_DYNAMIC_GEOMETRY == INTERPOLATION_SKINNING_TRANSLATION
+				{
+					OutPosition = ApplySkinningWithTranslation(GuideData, RenTri, RestPosition);
+				}
+				// 4. Linear blend skinning (rotation)
+				#elif PERMUTATION_DYNAMIC_GEOMETRY == INTERPOLATION_SKINNING_ROTATION
+				{
+					OutPosition = ApplySkinningWithRotation(GuideData, RenTri, SimTris, RestPosition);
+				}
+				// 5. Linear blend skinning (translation + rotation)
+				#elif PERMUTATION_DYNAMIC_GEOMETRY == INTERPOLATION_SKINNING_TRANSLATION_AND_ROTATION
+				{
+					OutPosition = ApplySkinningWithTranslationAndRotation(GuideData, RenTri, SimTris, RestPosition);
+				}
+				#endif
+
+				// 3. Apply deformer
+				#if PERMUTATION_DEFORMER
+				{
+					const uint4 PackedDeformerPosition = RenDeformerPositionBuffer[VertexIndex];
+					const float3 DeformerPosition = UnpackHairControlPointPosition(PackedDeformerPosition, InRenHairPositionOffset);
+					OutPosition = ApplyDeformer(RenTri, RestPosition, DeformerPosition, OutPosition);
+					PackedRestPosition.w = PackedDeformerPosition.w;
+				}
+				#endif
+
+				// 4. Write out the final position
+				// Optionally trim hair
+				if (HairLengthScale < 1.0f)
+				{
+					const float CoordU = UnpackHairControlPointCoordU(PackedRestPosition);
+					if (HairLengthScale < CoordU)
+					{
+						OutPosition = INFINITE_FLOAT;
+					}
+				}
+
+				OutRenDeformedPositionBuffer[VertexIndex] = PackHairControlPointPosition(PackedRestPosition, OutPosition, RenPositionOffset);
+			}
+		}
+	}
+}
+#endif // PERMUTATION_POINT_PER_CURVE < TOTALMAXPOINTPERCURVE
+//-------------------------------------------------------------------------------------------------
+#if PERMUTATION_POINT_PER_CURVE == TOTALMAXPOINTPERCURVE
 
-	// 2. Appply deformation
-	if (VertexIndex < VertexCount)
-	{
-		const FGuideDataWithOffset GuideData = GetGuideData(VertexIndex, PERMUTATION_USE_SINGLE_GUIDE, GetSimPositionOffset());
-		const FRootTriangleData RenTri = GetRootTriangleData(VertexIndex, false);
+groupshared uint CurvePointIndex;
+groupshared uint CurvePointCount;
 
-		// Manual decoding of the rest position
-		uint4 PackedRestPosition = RenRestPosePositionBuffer[VertexIndex];
-		const float3 RestPosition = UnpackHairControlPointPosition(PackedRestPosition, InRenHairPositionOffset);
+groupshared FRootTriangleData RenTri;
+groupshared float3 RenPositionOffset;
 
-		float3 OutPosition = 0.0f;
+groupshared FGuideCurveData SimCurveData;
+groupshared FRootTriangleData SimTris[2];
 
-		// 1. Rigid transformation
-		// ControlPoint is in the local hair referential
-		// CurrOffset takes only translation component into account, and is done in object local space (vs. triangle/hair local space)
-		#if PERMUTATION_DYNAMIC_GEOMETRY == INTERPOLATION_RIGID
-		{
-			OutPosition = ApplyRigidInterpolation(GuideData, RestPosition);
-		}
-		// 2. Skin transformation
-		// Apply dynamic mesh deformation (translation / rotation)
-		#elif PERMUTATION_DYNAMIC_GEOMETRY == INTERPOLATION_SKINNING_OFFSET
-		{
-			OutPosition = ApplySkinningWithOffset(GuideData, RenTri, RestPosition);
-		}
-		// 3. Linear blend skinning (translation)
-		#elif PERMUTATION_DYNAMIC_GEOMETRY == INTERPOLATION_SKINNING_TRANSLATION
-		{
-			OutPosition = ApplySkinningWithTranslation(GuideData, RenTri, RestPosition);
-		}
-		// 4. Linear blend skinning (rotation)
-		#elif PERMUTATION_DYNAMIC_GEOMETRY == INTERPOLATION_SKINNING_ROTATION
-		{
-			OutPosition = ApplySkinningWithRotation(GuideData, RenTri, RestPosition);
-		}
-		// 5. Linear blend skinning (translation + rotation)
-		#elif PERMUTATION_DYNAMIC_GEOMETRY == INTERPOLATION_SKINNING_TRANSLATION_AND_ROTATION
+[numthreads(GROUP_SIZE, 1, 1)]
+void MainCS(uint3 DispatchThreadId : SV_DispatchThreadID, uint ThreadIndex : SV_GroupIndex, uint2 GroupIndex : SV_GroupID)
+{
+	// 1. Read the Curve data
+	const uint CurveIndex = GetGlobalCurveIndex(ThreadIndex, GroupIndex);
+	const bool bIsCurveValid = CurveIndex < CurveCount;
+	if (!bIsCurveValid)
+	{
+		return;
+	}
+	
+	// 2. Curve's point count/offset
+	if (ThreadIndex == 0)
+	{
+		CurvePointIndex = 0;
+		CurvePointCount = 0;
+		if (bIsCurveValid)
 		{
-			OutPosition = ApplySkinningWithTranslationAndRotation(GuideData, RenTri, RestPosition);
-		}
-		#endif
+			const FHairCurve Curve = UnpackHairCurve(RenCurveBuffer[CurveIndex]);
+			CurvePointIndex = Curve.PointIndex;
+			CurvePointCount = Curve.PointCount;
+
+			// Load curve interpolation data
+			#if PERMUTATION_SIMULATION
+			SimCurveData = UnpackGuideCurveData(CurveInterpolationBuffer, CurveIndex, PERMUTATION_USE_SINGLE_GUIDE);
+			if (FORCE_USE_SINGLE_GUIDE)
+			{
+				SimCurveData.CurveWeights.x = 1;
+				SimCurveData.CurveWeights.y = 0;
+			}
+			#else
+			SimCurveData = (FGuideCurveData)0;
+			#endif
 
-		// 3. Apply deformer
-		#if PERMUTATION_DEFORMER
-		{
-			const uint4 PackedDeformerPosition = RenDeformerPositionBuffer[VertexIndex];
-			const float3 DeformerPosition = UnpackHairControlPointPosition(PackedDeformerPosition, InRenHairPositionOffset);
-			OutPosition = ApplyDeformer(RenTri, RestPosition, DeformerPosition, OutPosition);
-			PackedRestPosition.w = PackedDeformerPosition.w;
+			// Load sim/ren triangles
+			RenTri = GetRootTriangleData(CurveIndex, false/*bSim*/);
+
+			#if PERMUTATION_SIMULATION
+			SimTris[0] = GetRootTriangleData(SimCurveData.CurveIndices[0], true/*bSim*/);
+			SimTris[1] = SimTris[0];
+			if (!PERMUTATION_USE_SINGLE_GUIDE && !FORCE_USE_SINGLE_GUIDE)
+			{
+				SimTris[1] = GetRootTriangleData(SimCurveData.CurveIndices[1], true/*bSim*/);
+			}
+			#else
+			SimTris[0] = (FRootTriangleData)0;
+			SimTris[1] = (FRootTriangleData)0;
+			#endif
 		}
-		#endif
+	}
+	RenPositionOffset = GetRenPositionOffset();
+	GroupMemoryBarrierWithGroupSync();
 
-		// 4. Write out the final position
-		// Optionally trim hair
-		if (HairLengthScale < 1.0f)
-		{
-			const float CoordU = UnpackHairControlPointCoordU(PackedRestPosition);
-			if (HairLengthScale < CoordU)
+	// 2. Apply deformation
+	{
+		const uint LoopCount = DivideAndRoundUp(CurvePointCount, GROUP_SIZE);
+		for (uint LoopIt = 0; LoopIt < LoopCount; ++LoopIt)
+		{		
+			const uint CurrentPointIndex = ThreadIndex + LoopIt * GROUP_SIZE;
+	
+			if (CurrentPointIndex < CurvePointCount)
 			{
-				OutPosition = INFINITE_FLOAT;
+				const uint VertexIndex = CurrentPointIndex + CurvePointIndex;
+
+				#if PERMUTATION_SIMULATION
+				const FGuidePointData SimPointData = UnpackGuidePointData(PointInterpolationBuffer, VertexIndex, PERMUTATION_USE_SINGLE_GUIDE);
+				const FGuideDataWithOffset GuideData = GetGuideData(SimCurveData, SimPointData, GetSimPositionOffset());
+				#else
+				const FGuidePointData SimPointData = (FGuidePointData)0;
+				const FGuideDataWithOffset GuideData = (FGuideDataWithOffset)0;;
+				#endif
+
+				// Manual decoding of the rest position
+				uint4 PackedRestPosition = RenRestPosePositionBuffer[VertexIndex];
+				const float3 RestPosition = UnpackHairControlPointPosition(PackedRestPosition, InRenHairPositionOffset);
+
+				float3 OutPosition = 0.0f;
+
+				// 1. Rigid transformation
+				// ControlPoint is in the local hair referential
+				// CurrOffset takes only translation component into account, and is done in object local space (vs. triangle/hair local space)
+				#if PERMUTATION_DYNAMIC_GEOMETRY == INTERPOLATION_RIGID
+				{
+					OutPosition = ApplyRigidInterpolation(GuideData, RestPosition);
+				}
+				// 2. Skin transformation
+				// Apply dynamic mesh deformation (translation / rotation)
+				#elif PERMUTATION_DYNAMIC_GEOMETRY == INTERPOLATION_SKINNING_OFFSET
+				{
+					OutPosition = ApplySkinningWithOffset(GuideData, RenTri, SimTris, RestPosition);
+				}
+				// 3. Linear blend skinning (translation)
+				#elif PERMUTATION_DYNAMIC_GEOMETRY == INTERPOLATION_SKINNING_TRANSLATION
+				{
+					OutPosition = ApplySkinningWithTranslation(GuideData, RenTri, RestPosition);
+				}
+				// 4. Linear blend skinning (rotation)
+				#elif PERMUTATION_DYNAMIC_GEOMETRY == INTERPOLATION_SKINNING_ROTATION
+				{
+					OutPosition = ApplySkinningWithRotation(GuideData, RenTri, SimTris, RestPosition);
+				}
+				// 5. Linear blend skinning (translation + rotation)
+				#elif PERMUTATION_DYNAMIC_GEOMETRY == INTERPOLATION_SKINNING_TRANSLATION_AND_ROTATION
+				{
+					OutPosition = ApplySkinningWithTranslationAndRotation(GuideData, RenTri, SimTris, RestPosition);
+				}
+				#endif
+
+				// 3. Apply deformer
+				#if PERMUTATION_DEFORMER
+				{
+					const uint4 PackedDeformerPosition = RenDeformerPositionBuffer[VertexIndex];
+					const float3 DeformerPosition = UnpackHairControlPointPosition(PackedDeformerPosition, InRenHairPositionOffset);
+					OutPosition = ApplyDeformer(RenTri, RestPosition, DeformerPosition, OutPosition);
+					PackedRestPosition.w = PackedDeformerPosition.w;
+				}
+				#endif
+
+				// 4. Write out the final position
+				// Optionally trim hair
+				if (HairLengthScale < 1.0f)
+				{
+					const float CoordU = UnpackHairControlPointCoordU(PackedRestPosition);
+					if (HairLengthScale < CoordU)
+					{
+						OutPosition = INFINITE_FLOAT;
+					}
+				}
+
+				OutRenDeformedPositionBuffer[VertexIndex] = PackHairControlPointPosition(PackedRestPosition, OutPosition, RenPositionOffset);
 			}
 		}
-		OutRenDeformedPositionBuffer[VertexIndex] = PackHairControlPointPosition(PackedRestPosition, OutPosition, GetRenPositionOffset());
 	}
 }
+#endif //PERMUTATION_POINT_PER_CURVE == TOTALMAXPOINTPERCURVE
 #endif // SHADER_HAIRINTERPOLATION
 
 ///////////////////////////////////////////////////////////////////////////////////////////////////
@@ -510,6 +743,7 @@ void MainCS(uint2 DispatchThreadId : SV_DispatchThreadID)
 #include "HairStrandsAttributeCommon.ush"
 
 uint				CurveCount;
+uint				bUseSingleGuide;
 
 Buffer<uint>		RenCurveBuffer;
 Buffer<uint>		RenCurveToClusterIdBuffer;
@@ -539,7 +773,7 @@ uint ToSeedHash(uint In)
 	return (Hash & 0xFF);
 }
 
-[numthreads(1024, 1, 1)]
+[numthreads(GROUP_SIZE, 1, 1)]
 void MainCS(uint2 DispatchThreadId : SV_DispatchThreadID, uint GroupThread1D : SV_GroupIndex)
 {
 	SeedValues[GroupThread1D] = 0;
@@ -556,8 +790,13 @@ void MainCS(uint2 DispatchThreadId : SV_DispatchThreadID, uint GroupThread1D : S
 		}
 		#else // if (PERMUTATION_SIMULATION > 0)
 		{
-			const uint VertexIndex0		= UnpackHairCurve(RenCurveBuffer[CurveIndex]).PointIndex;
-			const FGuideDataWithOffset GuideData = GetGuideData(VertexIndex0, PERMUTATION_SIMULATION == 1, 0 /*OutPositionOffset*/);
+
+			const uint VertexIndex0 = UnpackHairCurve(RenCurveBuffer[CurveIndex]).PointIndex;
+
+			const FGuideCurveData SimCurveData = UnpackGuideCurveData(CurveInterpolationBuffer, CurveIndex, bUseSingleGuide);
+			const FGuidePointData SimPointData = UnpackGuidePointData(PointInterpolationBuffer, VertexIndex0, bUseSingleGuide); 
+			const FGuideDataWithOffset GuideData = GetGuideData(SimCurveData, SimPointData, 0 /*OutPositionOffset*/);
+
 			const uint GuideIndex		= GuideData.VertexIndices[0]; // Take the closest guide
 			SeedValues[GroupThread1D]	= ToSeedHash(GuideIndex);
 		}
diff --git a/Engine/Shaders/Private/HairStrands/HairStrandsPack.ush b/Engine/Shaders/Private/HairStrands/HairStrandsPack.ush
--- a/Engine/Shaders/Private/HairStrands/HairStrandsPack.ush
+++ b/Engine/Shaders/Private/HairStrands/HairStrandsPack.ush
@@ -169,37 +169,107 @@ float UnpackHairAO(uint In)
 ////////////////////////////////////////////////////////////////////////////////////////////////////////////////////////////////
 // Guide
 
-struct FGuideData
+struct FGuideCurveData
 {
-	uint3  VertexIndices;
-	float3 VertexLerps;
-	float3 VertexWeights;
+	uint2  CurveIndices;		// Index of the sim curve
+	float2 CurveWeights;		// Weight of the sim curve
+	uint2  RootPointIndices;		// Index of the first/root point of the sim curve (stored to avoid indirect fetch on CurveBuffer)
 };
 
-FGuideData UnpackGuideData(uint In)
+FGuideCurveData UnpackGuideCurveData(ByteAddressBuffer InBuffer, uint InCurveIndex, bool bUseSingleGuide)
+{
+	// This code assumes:
+	// * HAIR_INTERPOLATION_CURVE_STRIDE == 8
+	// * HAIR_INTERPOLATION_MAX_GUIDE_COUNT == 2
+
+	FGuideCurveData Out = (FGuideCurveData)0;
+	if (bUseSingleGuide)
+	{
+		const uint AddressInBytes = InCurveIndex * HAIR_INTERPOLATION_CURVE_STRIDE;
+		const uint2 Data = InBuffer.Load2(AddressInBytes);
+
+		Out.CurveIndices = Data.x & 0xFFFFFF;
+		Out.CurveWeights = UnpackR8(Data.x>>24);
+		Out.RootPointIndices = Data.y & 0xFFFFFF;
+	}
+	else
+	{
+		const uint AddressInBytes = InCurveIndex * HAIR_INTERPOLATION_CURVE_STRIDE * HAIR_INTERPOLATION_MAX_GUIDE_COUNT;
+		const uint4 Data = InBuffer.Load4(AddressInBytes);
+
+		Out.CurveIndices.x = Data.x & 0xFFFFFF;
+		Out.CurveWeights.x = UnpackR8(Data.x>>24);
+		Out.RootPointIndices.x = Data.y & 0xFFFFFF;
+
+		Out.CurveIndices.y = Data.z & 0xFFFFFF;
+		Out.CurveWeights.y = UnpackR8(Data.z>>24);
+		Out.RootPointIndices.y = Data.w & 0xFFFFFF;
+	}
+	return Out;
+}
+
+struct FGuidePointData
 {
-	FGuideData Out = (FGuideData)0;
-	Out.VertexWeights		= float3(1, 0, 0);
-	Out.VertexIndices		= 0;
-	Out.VertexLerps			= 0;
-	Out.VertexIndices.x		= In & 0xFFFFFF;
-	Out.VertexLerps.x		= UnpackR8(In >> 24);
+	uint2  PointIndices;
+	float2 PointLerps;
+};
+
+FGuidePointData UnpackGuidePointData(ByteAddressBuffer InBuffer, uint InPointIndex, bool bUseSingleGuide)
+{
+	// This code assumes:
+	// * HAIR_INTERPOLATION_POINT_STRIDE == 2
+	// * HAIR_INTERPOLATION_MAX_GUIDE_COUNT == 2
+
+	FGuidePointData Out = (FGuidePointData)0;
+	if (bUseSingleGuide)
+	{
+		// Ensure address is aligned on 4
+		// Since the data is 16bit wide, odd index needs to be shifted out of the 32bit input
+		// [    uint 32bits   ][    uint 32bits   ] ...
+		// [ Point0 ][ Point1 ][ Point2 ][ Point3 ] ...
+		const uint AddressInBytes = InPointIndex * HAIR_INTERPOLATION_POINT_STRIDE;
+		const uint AlignedAddressInBytes = AddressInBytes & 0xFFFFFFFE;
+		uint Data = InBuffer.Load(AlignedAddressInBytes);
+		Data = InPointIndex & 1 ? (Data>>16u) : Data;
+		
+		Out.PointIndices	= Data & 0xFF;
+		Out.PointLerps		= UnpackR8(Data >> 8);
+	}
+	else
+	{
+		const uint AddressInBytes = InPointIndex * HAIR_INTERPOLATION_POINT_STRIDE * HAIR_INTERPOLATION_MAX_GUIDE_COUNT;
+		const uint Data = InBuffer.Load(AddressInBytes);
+
+		Out.PointIndices.x	= Data & 0xFF;
+		Out.PointLerps.x	= UnpackR8(Data >> 8);
+
+		Out.PointIndices.y	= (Data>>16) & 0xFF;
+		Out.PointLerps.y	= UnpackR8(Data >> 24);
+	}
 	return Out;
 }
 
-FGuideData UnpackGuideData(uint4 In)
+
+////////////////////////////////////////////////////////////////////////////////////////////////////////////////////////////////
+// Cards guides
+
+struct FCardsGuideData
 {
-	FGuideData Out = (FGuideData)0;
-	Out.VertexIndices.x = (In.x & 0xFFFF) | (((In.z      ) & 0xFF) << 16);
-	Out.VertexIndices.y = (In.x >> 16   ) | (((In.z >> 8 ) & 0xFF) << 16);
-	Out.VertexIndices.z = (In.y & 0xFFFF) | (((In.z >> 16) & 0xFF) << 16);
+	uint3  VertexIndices;
+	float3 VertexLerps;
+	float3 VertexWeights;
+};
 
-	Out.VertexWeights	= float3(UnpackR8(In.y>>16), UnpackR8(In.y>>24), 0);
-	Out.VertexWeights.z = saturate(1.0f - (Out.VertexWeights.x + Out.VertexWeights.y));
+FCardsGuideData UnpackCardGuideData(ByteAddressBuffer InBuffer, uint InVertexId)
+{
+	const uint Data = InBuffer.Load(InVertexId * HAIR_INTERPOLATION_CARDS_GUIDE_STRIDE);
 
-	Out.VertexLerps.x	= UnpackR8(In.w);
-	Out.VertexLerps.y	= UnpackR8(In.w >> 8);
-	Out.VertexLerps.z	= UnpackR8(In.w >> 16);
+	FCardsGuideData Out = (FCardsGuideData)0;
+	Out.VertexWeights	= float3(1, 0, 0);
+	Out.VertexIndices	= 0;
+	Out.VertexLerps		= 0;
+	Out.VertexIndices.x	= Data & 0xFFFFFF;
+	Out.VertexLerps.x	= UnpackR8(Data >> 24);
 	return Out;
 }
 
diff --git a/Engine/Shaders/Shared/HairStrandsDefinitions.h b/Engine/Shaders/Shared/HairStrandsDefinitions.h
--- a/Engine/Shaders/Shared/HairStrandsDefinitions.h
+++ b/Engine/Shaders/Shared/HairStrandsDefinitions.h
@@ -32,8 +32,12 @@
 #define HAIR_ATTRIBUTE_INVALID_OFFSET 0xFFFFFFFF
 
 // Hair interpolation (in bytes)
-#define HAIR_INTERPOLATION_1GUIDE_STRIDE 4
-#define HAIR_INTERPOLATION_3GUIDE_STRIDE 16
+
+#define HAIR_INTERPOLATION_CURVE_STRIDE 8
+#define HAIR_INTERPOLATION_POINT_STRIDE 2
+#define HAIR_INTERPOLATION_MAX_GUIDE_COUNT 2
+
+#define HAIR_INTERPOLATION_CARDS_GUIDE_STRIDE 4
 
 // Max number of discrete LOD that a hair group can have
 #define MAX_HAIR_LOD 8
diff --git a/Engine/Source/Runtime/Renderer/Private/HairStrands/HairStrandsUtils.cpp b/Engine/Source/Runtime/Renderer/Private/HairStrands/HairStrandsUtils.cpp
--- a/Engine/Source/Runtime/Renderer/Private/HairStrands/HairStrandsUtils.cpp
+++ b/Engine/Source/Runtime/Renderer/Private/HairStrands/HairStrandsUtils.cpp
@@ -505,4 +505,32 @@ void AddTransitionPass(
 			FlushArray(PermutationIndex, ArrayCounts[PermutationIndex]);
 		}
 	}
+}
+
+FPointPerCurveDispatchInfo GetPointPerCurveDispatchInfo(uint32 InAssetMaxPointPerCurve, uint32 InAssetCurveCount, uint32 InGroupSize)
+{
+	FPointPerCurveDispatchInfo Out;
+	Out.SourcePoinPerCurve = InAssetMaxPointPerCurve;
+	Out.SourceCurveCount = InAssetCurveCount;
+	Out.GroupSize = InGroupSize;
+
+	// Compute the rounded point-per-curve count, based on the asset and the shader's requirement
+	Out.PointPerCurve = FMath::Clamp(uint32(FMath::Pow(2u, FMath::RoundFromZero(FMath::Log2(float(InAssetMaxPointPerCurve))))), 4u, Out.GroupSize);
+	check(FMath::IsPowerOfTwo(Out.PointPerCurve));
+
+	// Compute the number of curve per group
+	Out.CurvePerGroup = Out.GroupSize / Out.PointPerCurve;
+
+	// Compute dispatch count
+	const uint32 LinearGroupCount = FMath::DivideAndRoundUp(Out.SourceCurveCount, Out.CurvePerGroup);
+	Out.DispatchCount = FIntVector(LinearGroupCount, 1, 1);
+	if (Out.DispatchCount.X > 0xFFFFu)
+	{
+		Out.DispatchCount.X = 64;
+		Out.DispatchCount.Y = FMath::DivideAndRoundUp(LinearGroupCount, uint32(Out.DispatchCount.X));
+	}
+	check(Out.DispatchCount.X <= 0xFFFFu);
+	check(Out.DispatchCount.Y <= 0xFFFFu);
+
+	return Out;
 }
\ No newline at end of file
diff --git a/Engine/Source/Runtime/Renderer/Public/HairStrandsInterface.h b/Engine/Source/Runtime/Renderer/Public/HairStrandsInterface.h
--- a/Engine/Source/Runtime/Renderer/Public/HairStrandsInterface.h
+++ b/Engine/Source/Runtime/Renderer/Public/HairStrandsInterface.h
@@ -90,7 +90,8 @@ END_SHADER_PARAMETER_STRUCT()
 
 // Instance interpolation resources (RDG)
 BEGIN_SHADER_PARAMETER_STRUCT(FHairStrandsInstanceInterpolationParameters, RENDERER_API)
-	SHADER_PARAMETER_RDG_BUFFER_SRV(ByteAddressBuffer, InterpolationBuffer)
+	SHADER_PARAMETER_RDG_BUFFER_SRV(ByteAddressBuffer, CurveInterpolationBuffer)
+	SHADER_PARAMETER_RDG_BUFFER_SRV(ByteAddressBuffer, PointInterpolationBuffer)
 END_SHADER_PARAMETER_STRUCT()
 
 // Instance resources (Raw)
@@ -440,6 +441,24 @@ struct FHairTransientResources
 	const FCachedGeometry& GetMeshLODData(uint32 InRegisteredIndex, bool bSim) const { if (bSim) { check(SimMeshDatas.IsValidIndex(InRegisteredIndex)); return SimMeshDatas[InRegisteredIndex]; } else {check(RenMeshDatas.IsValidIndex(InRegisteredIndex)); return RenMeshDatas[InRegisteredIndex];} }
 };
 
+////////////////////////////////////////////////////////////////////////////////////////////////////
+// Helpers
+struct FPointPerCurveDispatchInfo
+{
+	// From the asset
+	uint32 SourceCurveCount = 0;
+	uint32 SourcePoinPerCurve = 0;
+
+	// Derived
+	uint32 CurvePerGroup = 0;
+	uint32 PointPerCurve = 0;
+	uint32 GroupSize = 0;
+	FIntVector DispatchCount = FIntVector::ZeroValue;
+};
+
+// Compute the dispatch information for pass dispatching work per curve
+RENDERER_API FPointPerCurveDispatchInfo GetPointPerCurveDispatchInfo(uint32 InAssetMaxPointPerCurve, uint32 InAssetCurveCount, uint32 InGroupSize);
+
 ////////////////////////////////////////////////////////////////////////////////////////////////////
 // API for enabling/disabling the various geometry representation
 enum class EHairStrandsShaderType
