Fix missing hair tangent buffer initialization when building RT geometry on static grooms.

diff --git a/Engine/Plugins/Runtime/HairStrands/Source/HairStrandsCore/Private/GroomManager.cpp b/Engine/Plugins/Runtime/HairStrands/Source/HairStrandsCore/Private/GroomManager.cpp
--- a/Engine/Plugins/Runtime/HairStrands/Source/HairStrandsCore/Private/GroomManager.cpp
+++ b/Engine/Plugins/Runtime/HairStrands/Source/HairStrandsCore/Private/GroomManager.cpp
@@ -1500,7 +1500,8 @@ static void RunHairStrandsInterpolation_Strands(
 		}
 		else
 		{
-			InstanceData.Instance->Strands.RestResource->GetTangentBuffer(GraphBuilder, ShaderMap, InstanceData.ActivePointCount, InstanceData.ActiveCurveCount);
+			FRDGExternalBuffer TangentBuffer = InstanceData.Instance->Strands.RestResource->GetTangentBuffer(GraphBuilder, ShaderMap, InstanceData.ActivePointCount, InstanceData.ActiveCurveCount);
+			InstanceData.RDGResources.TangentSRV = RegisterAsSRV(GraphBuilder, TangentBuffer);
 		}
 	}
 
