* Avoid to build invalid groom bindings
* Clearer error message when groom binding is invalid

diff --git a/Engine/Plugins/Runtime/HairStrands/Source/HairStrandsCore/Private/GroomBindingAsset.cpp b/Engine/Plugins/Runtime/HairStrands/Source/HairStrandsCore/Private/GroomBindingAsset.cpp
--- a/Engine/Plugins/Runtime/HairStrands/Source/HairStrandsCore/Private/GroomBindingAsset.cpp
+++ b/Engine/Plugins/Runtime/HairStrands/Source/HairStrandsCore/Private/GroomBindingAsset.cpp
@@ -1029,7 +1029,7 @@ void UGroomBindingAsset::BeginCacheForCookedPlatformData(const ITargetPlatform*
 	UGroomBindingAsset::FCachedCookedPlatformData* TargetPlatformData = FindCachedCookedPlatformData(GroupDerivedDataKeys, CachedCookedPlatformDatas);
 
 	// 3. If the target cooked data does not already exist, we build it
-	if (TargetPlatformData == nullptr)
+	if (TargetPlatformData == nullptr && GetGroom() != nullptr)
 	{
 		// 3.1 Build cooked derived data
 		const uint32 GroupCount = GroupDerivedDataKeys.Num();
diff --git a/Engine/Plugins/Runtime/HairStrands/Source/HairStrandsCore/Private/GroomBindingBuilder.cpp b/Engine/Plugins/Runtime/HairStrands/Source/HairStrandsCore/Private/GroomBindingBuilder.cpp
--- a/Engine/Plugins/Runtime/HairStrands/Source/HairStrandsCore/Private/GroomBindingBuilder.cpp
+++ b/Engine/Plugins/Runtime/HairStrands/Source/HairStrandsCore/Private/GroomBindingBuilder.cpp
@@ -2147,10 +2147,12 @@ static void BuildRootBulkData(
 static bool InternalBuildBinding_CPU(const FGroomBindingBuilder::FInput& In, uint32 InGroupIndex, const ITargetPlatform* TargetPlatform, UGroomBindingAsset::FHairGroupPlatformData& OutPlatformData)
 {
 #if WITH_EDITORONLY_DATA
-	if (!In.GroomAsset || !In.bHasValidTarget || In.GroomAsset->GetNumHairGroups() == 0)
+	const bool bIsAssetValid = In.GroomAsset && In.bHasValidTarget && In.GroomAsset->GetNumHairGroups() > 0;
+	if (!bIsAssetValid)
 	{
-		UE_LOG(LogHairStrands, Error, TEXT("[Groom] Binding asset cannot be created/rebuilt."));
-		return false;
+		if (!In.GroomAsset)							{ UE_LOG(LogHairStrands, Error, TEXT("[Groom] Binding asset cannot be created/rebuilt - The groom binding has no groom asset.")); return false; }
+		if (!In.bHasValidTarget)					{ UE_LOG(LogHairStrands, Error, TEXT("[Groom] Binding asset cannot be created/rebuilt - The groom binding has no valid skel./geom cache. target")); return false; }
+		if (In.GroomAsset->GetNumHairGroups() == 0)	{ UE_LOG(LogHairStrands, Error, TEXT("[Groom] Binding asset cannot be created/rebuilt - The groom asset has no groups.")); return false; }
 	}
 
 	// 1. Build groom root data
