Fix invalid warning when hair strands are used on a project without skin cache and hair manual cache is enabled.

diff --git a/Engine/Plugins/Runtime/HairStrands/Source/HairStrandsCore/Private/GroomComponent.cpp b/Engine/Plugins/Runtime/HairStrands/Source/HairStrandsCore/Private/GroomComponent.cpp
--- a/Engine/Plugins/Runtime/HairStrands/Source/HairStrandsCore/Private/GroomComponent.cpp
+++ b/Engine/Plugins/Runtime/HairStrands/Source/HairStrandsCore/Private/GroomComponent.cpp
@@ -81,6 +81,11 @@ static FAutoConsoleVariableRef CVarHairStrands_ForceVelocityOutput(TEXT("r.HairS
 
 #define LOCTEXT_NAMESPACE "GroomComponent"
 
+////////////////////////////////////////////////////////////////////////////////////////////////////////////////////////
+// Forward declarations
+
+bool IsHairManualSkinCacheEnabled();
+
 ////////////////////////////////////////////////////////////////////////////////////////////////////////////////////////
 
 template<typename T>
@@ -2580,7 +2585,7 @@ void UGroomComponent::InitResources(bool bIsBindingReloading)
 						const EGroomBindingType BindingType = GroomAsset->GetBindingType(GroupIt, LODIt);
 						const bool bIsVisible = GroomAsset->IsVisible(GroupIt, LODIt);
 
-						if (BindingType == EGroomBindingType::Skinning && !bSupportSkinCache && bIsVisible)
+						if (BindingType == EGroomBindingType::Skinning && (!bSupportSkinCache && !IsHairManualSkinCacheEnabled()) && bIsVisible)
 						{
 							UE_LOG(LogHairStrands, Warning, TEXT("[Groom] Groom asset (Group:%d/%d) is set to use Skinning at LOD %d/%d while the parent skel. mesh does not support skin. cache at this LOD - Groom:%s - Skel.Mesh:%s"),
 								GroupIt,
diff --git a/Engine/Plugins/Runtime/HairStrands/Source/HairStrandsCore/Private/GroomManager.cpp b/Engine/Plugins/Runtime/HairStrands/Source/HairStrandsCore/Private/GroomManager.cpp
--- a/Engine/Plugins/Runtime/HairStrands/Source/HairStrandsCore/Private/GroomManager.cpp
+++ b/Engine/Plugins/Runtime/HairStrands/Source/HairStrandsCore/Private/GroomManager.cpp
@@ -91,6 +91,11 @@ void AddHairClusterAABBPass(
 ///////////////////////////////////////////////////////////////////////////////////////////////////
 // Utils
 
+bool IsHairManualSkinCacheEnabled()
+{
+	return GHairStrands_ManualSkinCache > 0;
+}
+
 EHairBufferSwapType GetHairSwapBufferType()
 {
 	switch (GHairStrands_SwapBufferType)
@@ -455,7 +460,7 @@ static void AddHairSkinCacheDebugPass(
 	FRDGBufferRef InfoBuffer = CreateVertexBuffer(GraphBuilder, TEXT("Hair.Debug.UniqueMeshNames"), FRDGBufferDesc::CreateBufferDesc(4, InfoInUints * Infos.Num()), Infos.GetData(), InfoInBytes * Infos.Num());
 
 	FHairDebugPrintHairSkinCacheCS::FParameters* Parameters = GraphBuilder.AllocParameters<FHairDebugPrintHairSkinCacheCS::FParameters>();
-	Parameters->HairSkinCacheEnable = GHairStrands_ManualSkinCache > 0 ? 1u : 0u;
+	Parameters->HairSkinCacheEnable = IsHairManualSkinCacheEnabled() ? 1u : 0u;
 	Parameters->GPUSkinCacheEnable	= bIsGPUSkinCacheEnable ? 1u : 0u;
 	Parameters->InstanceCount 		= InstanceCount;
 	Parameters->UniqueMeshCount 	= UniqueMeshCount;
@@ -600,7 +605,7 @@ static FCachedGeometry GetCacheGeometryForHair(
 					OutHairGeometryCache.AddDebug(Instance, SceneProxy, Out, PositionUpdateType, FHairGeometryCache::ECacheType::SkinCache);
 				}
 				// 2. If no cached geometry is  extracted from the skel. mesh proxy, compute it using the manual hair skin cache
-				else if (GHairStrands_ManualSkinCache > 0)
+				else if (IsHairManualSkinCacheEnabled())
 				{
 					uint32 OutTotalSectionCount = 0;
 					GetOrAllocateCachedGeometry(GraphBuilder, ShaderMap, SceneProxy, RootBulkData, bOutputTriangleData, Out, OutHairGeometryCache, OutTotalSectionCount);
