Fix crash on CitySample due to groom bulk data not being 4byte aligned.

diff --git a/Engine/Plugins/Runtime/HairStrands/Source/HairStrandsCore/Private/GroomBuilder.cpp b/Engine/Plugins/Runtime/HairStrands/Source/HairStrandsCore/Private/GroomBuilder.cpp
--- a/Engine/Plugins/Runtime/HairStrands/Source/HairStrandsCore/Private/GroomBuilder.cpp
+++ b/Engine/Plugins/Runtime/HairStrands/Source/HairStrandsCore/Private/GroomBuilder.cpp
@@ -40,7 +40,7 @@ static FAutoConsoleVariableRef CVarHairGroupIndexBuilder_MaxVoxelResolution(TEXT
 
 FString FGroomBuilder::GetVersion()
 {
-	return TEXT("v10a");
+	return TEXT("v10b");
 }
 
 // For debug purpose
@@ -1513,17 +1513,10 @@ namespace HairInterpolationBuilder
 		const uint32 KCount = HairInterpolation.bUseUniqueGuide ? 1u : HAIR_INTERPOLATION_MAX_GUIDE_COUNT;
 		if (HairInterpolation.bUseUniqueGuide)
 		{
-
 			OutBulkData.Header.Flags |= FHairStrandsInterpolationBulkData::DataFlags_HasSingleGuideData;
-			OutBulkData.Header.Strides.CurveInterpolationStride = sizeof(FHairInterpolationCurve);
-			OutBulkData.Header.Strides.PointInterpolationStride = sizeof(FHairInterpolationPoint);
-		}
-		else
-		{
-
-			OutBulkData.Header.Strides.CurveInterpolationStride = sizeof(FHairInterpolationCurve) * HAIR_INTERPOLATION_MAX_GUIDE_COUNT;
-			OutBulkData.Header.Strides.PointInterpolationStride = sizeof(FHairInterpolationPoint) * HAIR_INTERPOLATION_MAX_GUIDE_COUNT;
 		}
+		OutBulkData.Header.Strides.CurveInterpolationStride = sizeof(FHairInterpolationCurve) * KCount;
+		OutBulkData.Header.Strides.PointInterpolationStride = sizeof(FHairInterpolationPoint) * KCount;
 		
 		// Data
 		{
@@ -1566,6 +1559,32 @@ namespace HairInterpolationBuilder
 				}
 			}
 
+			// Sanity check
+			{
+				// Check on data size aligned on 4-byte address
+				const uint32 PointTotalSize = OutInterpolationPoints.GetTypeSize() * OutInterpolationPoints.Num();
+				check((PointTotalSize & 0x3) == 0);
+
+				// Check on point count to ensure address is 4-byte aligned
+				// This is needed as resources are allocated from header description. This optionally adds dummy point.
+				uint32 PointSizeFromHeader = OutBulkData.Header.PointCount * OutBulkData.Header.Strides.PointInterpolationStride;
+				if ((PointSizeFromHeader & 0x3) != 0)
+				{
+					if (HairInterpolation.bUseUniqueGuide)
+					{
+						OutBulkData.Header.PointCount = FMath::DivideAndRoundUp(OutBulkData.Header.PointCount, 2u) * 2u;
+					}
+					else
+					{
+						// This shouldn't happen
+						check(false);
+					}
+				}
+				PointSizeFromHeader = OutBulkData.Header.PointCount * OutBulkData.Header.Strides.PointInterpolationStride;
+				check(PointSizeFromHeader == PointTotalSize);
+			}
+			
+
 			HairStrandsBuilder::CopyToBulkData<FHairStrandsInterpolationFormat>(OutBulkData.Data.CurveInterpolation, OutInterpolationCurves);
 			HairStrandsBuilder::CopyToBulkData<FHairStrandsInterpolationFormat>(OutBulkData.Data.PointInterpolation, OutInterpolationPoints);
 			ReportSize(TEXT("Interpolation Curves"), OutInterpolationCurves);
diff --git a/Engine/Plugins/Runtime/HairStrands/Source/HairStrandsCore/Private/HairStrandsDatas.cpp b/Engine/Plugins/Runtime/HairStrands/Source/HairStrandsCore/Private/HairStrandsDatas.cpp
--- a/Engine/Plugins/Runtime/HairStrands/Source/HairStrandsCore/Private/HairStrandsDatas.cpp
+++ b/Engine/Plugins/Runtime/HairStrands/Source/HairStrandsCore/Private/HairStrandsDatas.cpp
@@ -710,8 +710,12 @@ void FHairStrandsInterpolationBulkData::GetResources(FHairStrandsBulkCommon::FQu
 		uint32 CurveCount = 0;
 		if (Out.Type == FHairStrandsBulkCommon::FQuery::ReadIO || Out.Type == FHairStrandsBulkCommon::FQuery::ReadDDC || Out.Type == FHairStrandsBulkCommon::FQuery::UnloadData)
 		{
+			// Aligned requested point to get data aligned on 4-bytes
+			// See GroomBuilder.cpp for details
+			const uint32 RequestedPointCount = (Header.Flags & DataFlags_HasSingleGuideData) ? FMath::DivideAndRoundUp(Out.GetPointCount(),2u)*2u : Out.GetPointCount();
+
 			CurveCount = FMath::Min(Header.CurveCount, Out.GetCurveCount());
-			PointCount = FMath::Min(Header.PointCount, Out.GetPointCount());
+			PointCount = FMath::Min(Header.PointCount, RequestedPointCount);
 		}
 
 		Out.Add(Data.CurveInterpolation, 	TEXT("_CurveInterpolation"), 	Data.CurveInterpolation.LoadedSize, 		CurveCount * Header.Strides.CurveInterpolationStride);
@@ -1024,4 +1028,4 @@ void FHairStrandsRootBulkData::ResetLoadedSize()
 		LOD.MeshSampleIndicesAndSectionsBuffer.LoadedSize= 0;
 		LOD.RestSamplePositionsBuffer.LoadedSize 		 = 0;
 	}
-}
\ No newline at end of file
+}
