Pack hair strands point to curve bufffer.

This reduces its size by 8, and save around 10% of groom asset size.

diff --git a/Engine/Plugins/Runtime/HairStrands/Shaders/Private/DeformerDataInterfaceGroom.ush b/Engine/Plugins/Runtime/HairStrands/Shaders/Private/DeformerDataInterfaceGroom.ush
--- a/Engine/Plugins/Runtime/HairStrands/Shaders/Private/DeformerDataInterfaceGroom.ush
+++ b/Engine/Plugins/Runtime/HairStrands/Shaders/Private/DeformerDataInterfaceGroom.ush
@@ -121,7 +121,7 @@ float ReadCurveNumPoint_{DataInterfaceName}(uint CurveIndex)
 uint ReadGuideIndex_{DataInterfaceName}(uint ControlPointIndex)
 {
 	// Find the rendering curve index
-	const uint CurveIndex = {DataInterfaceName}_PointToCurveBuffer.Load(ControlPointIndex);
+	const uint CurveIndex = UnpackHairPointToCurveIndex({DataInterfaceName}_PointToCurveBuffer, ControlPointIndex);
 
 	// Fint the curve and point interpolation data
 	const FGuideCurveData CurveData = UnpackGuideCurveData({DataInterfaceName}_CurveInterpolationBuffer, CurveIndex, {DataInterfaceName}_bSingleGuide);
diff --git a/Engine/Plugins/Runtime/HairStrands/Shaders/Private/NiagaraInterpolateGroomGuides.usf b/Engine/Plugins/Runtime/HairStrands/Shaders/Private/NiagaraInterpolateGroomGuides.usf
--- a/Engine/Plugins/Runtime/HairStrands/Shaders/Private/NiagaraInterpolateGroomGuides.usf
+++ b/Engine/Plugins/Runtime/HairStrands/Shaders/Private/NiagaraInterpolateGroomGuides.usf
@@ -71,7 +71,7 @@ void MainCS(uint3 DispatchThreadId : SV_DispatchThreadID)
 {
 	if (DispatchThreadId.x < uint(NumPoints))
 	{
-		const int StrandIndex = PointToCurveIndexBuffer[DispatchThreadId.x];
+		const int StrandIndex = UnpackHairPointToCurveIndex(PointToCurveIndexBuffer, DispatchThreadId.x);
 
 		const FHairCurve Curve = UnpackHairCurve(CurvesOffsetsBuffer[StrandIndex]);
 		const int PointBegin = Curve.PointIndex;
diff --git a/Engine/Plugins/Runtime/HairStrands/Source/HairStrandsCore/Private/GroomBuilder.cpp b/Engine/Plugins/Runtime/HairStrands/Source/HairStrandsCore/Private/GroomBuilder.cpp
--- a/Engine/Plugins/Runtime/HairStrands/Source/HairStrandsCore/Private/GroomBuilder.cpp
+++ b/Engine/Plugins/Runtime/HairStrands/Source/HairStrandsCore/Private/GroomBuilder.cpp
@@ -40,7 +40,7 @@ static FAutoConsoleVariableRef CVarHairGroupIndexBuilder_MaxVoxelResolution(TEXT
 
 FString FGroomBuilder::GetVersion()
 {
-	return TEXT("v10b");
+	return TEXT("v11b");
 }
 
 // For debug purpose
@@ -183,17 +183,9 @@ namespace HairStrandsBuilder
 		OutPackedPositions.SetNum(NumPoints * FHairStrandsPositionFormat::ComponentCount);
 		OutPackedCurves.SetNum(NumCurves * FHairStrandsCurveFormat::ComponentCount);
 
-		const bool bUse16bitsCurveIndex = NumCurves < 65536;
-		TArray<FHairStrandsPointToCurveFormat16::Type> OutPointToCurve16;
-		TArray<FHairStrandsPointToCurveFormat32::Type> OutPointToCurve32;
-		if (bUse16bitsCurveIndex)
-		{
-			OutPointToCurve16.SetNum(NumPoints);
-		}
-		else
-		{
-			OutPointToCurve32.SetNum(NumPoints);
-		}
+		const uint32 PointToCurveChunkElementCount = 8u;
+		TArray<FHairStrandsPointToCurveFormat::Type> OutPointToCurve;
+		OutPointToCurve.SetNum(FMath::DivideAndRoundUp(NumPoints, PointToCurveChunkElementCount));
 
 		const uint32 Attributes = HairStrands.GetAttributes();
 		const uint32 AttributeFlags = HairStrands.GetAttributeFlags();
@@ -282,14 +274,48 @@ namespace HairStrandsBuilder
 					PackedPosition.UCoord = uint8(FMath::Clamp(CoordU * 255.f, 0.f, 255.f));
 				}
 
-				// Vertex to Curve
-				if (bUse16bitsCurveIndex)
+				// Point to Curve
 				{
-					OutPointToCurve16[PointIndex + IndexOffset] = CurveIndex;
-				}
-				else
-				{
-					OutPointToCurve32[PointIndex + IndexOffset] = CurveIndex;
+					static_assert(PointToCurveChunkElementCount == 8u);
+
+					const uint32 CurrentIndex    = PointIndex + IndexOffset;
+					const uint32 CurrentIndex4   = CurrentIndex >> 3u;
+					const uint32 LocalPointIndex = CurrentIndex & 0x7;
+					const bool bIsFirstPoint     = LocalPointIndex == 0;
+
+					// Format
+					// Encode 8 points per uint with a base curve index + delta bits
+					// [                         32bits                         ]
+					// [     24bits     ][                8bits                 ]
+					// [ BaseCurveIndex ][Pt0][Pt0][Pt2][Pt3][Pt4][Pt5][Pt6][Pt7]
+					if (bIsFirstPoint)
+					{
+						OutPointToCurve[CurrentIndex4] = CurveIndex;
+					}
+					else
+					{
+						const uint32 BaseCurveIndex = (OutPointToCurve[CurrentIndex4] & 0xFFFFFF);
+						uint32 PackedDeltas   = (OutPointToCurve[CurrentIndex4]>>24u) & 0xFF;
+
+						// Rebuilt curve indices from the detla encoding
+						uint32 CurveIndices[8] = {0,0,0,0,0,0,0,0};
+						CurveIndices[0] = BaseCurveIndex;
+						for (uint32 PointIt = 1; PointIt <= LocalPointIndex; ++PointIt)
+						{
+							CurveIndices[PointIt] = CurveIndices[PointIt-1] + ((PackedDeltas >> PointIt) & 0x1);
+						}
+
+						// Sanity check
+						check(CurveIndices[LocalPointIndex] == CurveIndex || CurveIndices[LocalPointIndex]+1 == CurveIndex);
+
+						// Add curve bit if needed
+						if (CurveIndices[LocalPointIndex] != CurveIndex)
+						{
+							PackedDeltas |= 1u << LocalPointIndex;
+						}
+
+						OutPointToCurve[CurrentIndex4] = BaseCurveIndex | PackedDeltas<<24u;
+					}
 				}
 
 				// Per-Vertex Color
@@ -538,7 +564,6 @@ namespace HairStrandsBuilder
 			const uint32 ChunkElementCount = FMath::Min(PointAttributeChunkElementCount, OutBulkData.Header.PointCount);
 			for (uint32 ChunkIt=0; ChunkIt<ChunkCount; ++ChunkIt)
 			{
-
 				if (HasHairAttribute(Attributes, EHairAttribute::Color))
 				{
 					AppendAttribute(HAIR_POINT_ATTRIBUTE_COLOR, Stride_Color, AttributeColor, ChunkIt, ChunkCount);
@@ -568,17 +593,8 @@ namespace HairStrandsBuilder
 		ReportSize(TEXT("Positions"), OutPackedPositions);
 		CopyToBulkData<FHairStrandsCurveFormat>(OutBulkData.Data.Curves, OutPackedCurves);
 		ReportSize(TEXT("PackedCurves"), OutPackedCurves);
-		if (bUse16bitsCurveIndex)
-		{
-			OutBulkData.Header.Flags |= FHairStrandsBulkData::DataFlags_Has16bitsCurveIndex;
-			CopyToBulkData<FHairStrandsPointToCurveFormat16>(OutBulkData.Data.PointToCurve, OutPointToCurve16);
-			ReportSize(TEXT("PointToCurve"), OutPointToCurve16);
-		}
-		else
-		{
-			CopyToBulkData<FHairStrandsPointToCurveFormat32>(OutBulkData.Data.PointToCurve, OutPointToCurve32);
-			ReportSize(TEXT("PointToCurve"), OutPointToCurve32);
-		}
+		CopyToBulkData<FHairStrandsPointToCurveFormat>(OutBulkData.Data.PointToCurve, OutPointToCurve);
+		ReportSize(TEXT("PointToCurve"), OutPointToCurve);
 
 		// Build curve to point count mapping for runtime CLOD
 		OutBulkData.Header.CurveToPointCount.Reserve(OutPackedCurves.Num());
@@ -590,7 +606,9 @@ namespace HairStrandsBuilder
 		// Stride datas
 		OutBulkData.Header.Strides.PositionStride = FHairStrandsPositionFormat::SizeInByte;
 		OutBulkData.Header.Strides.CurveStride = FHairStrandsCurveFormat::SizeInByte;
-		OutBulkData.Header.Strides.PointToCurveStride = (bUse16bitsCurveIndex ? FHairStrandsPointToCurveFormat16::SizeInByte : FHairStrandsPointToCurveFormat32::SizeInByte);
+
+		OutBulkData.Header.Strides.PointToCurveChunkStride = FHairStrandsPointToCurveFormat::SizeInByte;
+		OutBulkData.Header.Strides.PointToCurveChunkElementCount = PointToCurveChunkElementCount;
 
 		OutBulkData.Header.Strides.CurveAttributeChunkStride = CurveAttributeChunkStride;
 		OutBulkData.Header.Strides.PointAttributeChunkStride = PointAttributeChunkStride;
diff --git a/Engine/Plugins/Runtime/HairStrands/Source/HairStrandsCore/Private/GroomResources.cpp b/Engine/Plugins/Runtime/HairStrands/Source/HairStrandsCore/Private/GroomResources.cpp
--- a/Engine/Plugins/Runtime/HairStrands/Source/HairStrandsCore/Private/GroomResources.cpp
+++ b/Engine/Plugins/Runtime/HairStrands/Source/HairStrandsCore/Private/GroomResources.cpp
@@ -1054,14 +1054,9 @@ void FHairStrandsRestResource::InternalAllocate(FRDGBuilder& GraphBuilder)
 		GraphBuilder.UseExternalAccessMode(Register(GraphBuilder, PointAttributeBuffer, ERDGImportedBufferFlags::CreateSRV).Buffer, ERHIAccess::SRVMask);
 	}
 	InternalCreateVertexBufferRDG_FromHairBulkData<FHairStrandsCurveFormat>(GraphBuilder, BulkData.Data.Curves, CurveCount, CurveBuffer, ToHairResourceDebugName(HAIRSTRANDS_RESOUCE_NAME(CurveType, Hair.StrandsRest_CurveBuffer), ResourceName), OwnerName, EHairResourceUsageType::Static);
-	if (!!(BulkData.Header.Flags & FHairStrandsBulkData::DataFlags_Has16bitsCurveIndex))
-	{
-		InternalCreateVertexBufferRDG_FromHairBulkData<FHairStrandsPointToCurveFormat16>(GraphBuilder, BulkData.Data.PointToCurve, PointCount, PointToCurveBuffer, ToHairResourceDebugName(HAIRSTRANDS_RESOUCE_NAME(CurveType, Hair.StrandsRest_PointToCurveBuffer), ResourceName), OwnerName, EHairResourceUsageType::Static);
-	}
-	else
-	{
-		InternalCreateVertexBufferRDG_FromHairBulkData<FHairStrandsPointToCurveFormat32>(GraphBuilder, BulkData.Data.PointToCurve, PointCount, PointToCurveBuffer, ToHairResourceDebugName(HAIRSTRANDS_RESOUCE_NAME(CurveType, Hair.StrandsRest_PointToCurveBuffer), ResourceName), OwnerName, EHairResourceUsageType::Static);
-	}
+
+	const uint32 PointToCurveIndexCount = FMath::DivideAndRoundUp(PointCount, BulkData.Header.Strides.PointToCurveChunkElementCount);
+	InternalCreateVertexBufferRDG_FromHairBulkData<FHairStrandsPointToCurveFormat>(GraphBuilder, BulkData.Data.PointToCurve, PointToCurveIndexCount, PointToCurveBuffer, ToHairResourceDebugName(HAIRSTRANDS_RESOUCE_NAME(CurveType, Hair.StrandsRest_PointToCurveBuffer), ResourceName), OwnerName, EHairResourceUsageType::Static);
 
 	// If the buffer is not null, it means it was been lazily allocated once. In such a case, The tangent buffer is recreated the new required size
 	const bool bStaticTangentNeeded = TangentBuffer.Buffer != nullptr;
diff --git a/Engine/Plugins/Runtime/HairStrands/Source/HairStrandsCore/Private/HairStrandsDatas.cpp b/Engine/Plugins/Runtime/HairStrands/Source/HairStrandsCore/Private/HairStrandsDatas.cpp
--- a/Engine/Plugins/Runtime/HairStrands/Source/HairStrandsCore/Private/HairStrandsDatas.cpp
+++ b/Engine/Plugins/Runtime/HairStrands/Source/HairStrandsCore/Private/HairStrandsDatas.cpp
@@ -564,7 +564,10 @@ void FHairStrandsBulkData::SerializeHeader(FArchive& Ar, UObject* Owner)
 
 	Ar << Header.Strides.PositionStride;
 	Ar << Header.Strides.CurveStride;
-	Ar << Header.Strides.PointToCurveStride;
+
+	Ar << Header.Strides.PointToCurveChunkStride;
+	Ar << Header.Strides.PointToCurveChunkElementCount;
+
 	Ar << Header.Strides.CurveAttributeChunkStride;
 	Ar << Header.Strides.PointAttributeChunkStride;
 	Ar << Header.Strides.CurveAttributeChunkElementCount;
@@ -586,8 +589,7 @@ void FHairStrandsBulkData::GetResources(FHairStrandsBulkCommon::FQuery& Out)
 {
 	static_assert(sizeof(FHairStrandsPositionFormat::BulkType) == sizeof(FHairStrandsPositionFormat::Type));
 	static_assert(sizeof(FHairStrandsAttributeFormat::BulkType) == sizeof(FHairStrandsAttributeFormat::Type));
-	static_assert(sizeof(FHairStrandsPointToCurveFormat16::BulkType) == sizeof(FHairStrandsPointToCurveFormat16::Type));
-	static_assert(sizeof(FHairStrandsPointToCurveFormat32::BulkType) == sizeof(FHairStrandsPointToCurveFormat32::Type));
+	static_assert(sizeof(FHairStrandsPointToCurveFormat::BulkType) == sizeof(FHairStrandsPointToCurveFormat::Type));
 	static_assert(sizeof(FHairStrandsRootIndexFormat::BulkType) == sizeof(FHairStrandsRootIndexFormat::Type)); 
 
 	// Translate requested curve count into chunk/offset/size to be read
@@ -601,6 +603,7 @@ void FHairStrandsBulkData::GetResources(FHairStrandsBulkCommon::FQuery& Out)
 
 	const uint32 PointAttributeSize = GetPointAttributeSizeInBytes(PointCount);
 	const uint32 CurveAttributeSize = GetCurveAttributeSizeInBytes(CurveCount);
+	const uint32 PointToCurveIndexSize = GetPointToCurveSizeInBytes(PointCount);
 
 	if (!!(Header.Flags & DataFlags_HasData))
 	{
@@ -610,7 +613,7 @@ void FHairStrandsBulkData::GetResources(FHairStrandsBulkCommon::FQuery& Out)
 		{
 			Out.Add(Data.PointAttributes, 	TEXT("_PointAttributes"), 	Data.PointAttributes.LoadedSize,PointAttributeSize);
 		}
-		Out.Add(Data.PointToCurve, 			TEXT("_PointToCurve"), 		Data.PointToCurve.LoadedSize,	PointCount * Header.Strides.PointToCurveStride);
+		Out.Add(Data.PointToCurve, 			TEXT("_PointToCurve"), 		Data.PointToCurve.LoadedSize,	PointToCurveIndexSize);
 		Out.Add(Data.Curves, 				TEXT("_Curves"), 			Data.Curves.LoadedSize, 		CurveCount * Header.Strides.CurveStride);
 	}
 }
diff --git a/Engine/Plugins/Runtime/HairStrands/Source/HairStrandsCore/Private/HairStrandsInterpolation.cpp b/Engine/Plugins/Runtime/HairStrands/Source/HairStrandsCore/Private/HairStrandsInterpolation.cpp
--- a/Engine/Plugins/Runtime/HairStrands/Source/HairStrandsCore/Private/HairStrandsInterpolation.cpp
+++ b/Engine/Plugins/Runtime/HairStrands/Source/HairStrandsCore/Private/HairStrandsInterpolation.cpp
@@ -486,7 +486,6 @@ class FHairInterpolationCS : public FGlobalShader
 		SHADER_PARAMETER_RDG_BUFFER_SRV(Buffer<float4>, RenRootDeformedPositionBuffer)
 
 		SHADER_PARAMETER_RDG_BUFFER_SRV(Buffer<uint>, RenRootBarycentricBuffer)
-		SHADER_PARAMETER_RDG_BUFFER_SRV(Buffer<uint>, RenPointToCurveBuffer)
 		SHADER_PARAMETER_RDG_BUFFER_SRV(Buffer<uint>, RenRootToUniqueTriangleIndexBuffer)
 
 		SHADER_PARAMETER_RDG_BUFFER_SRV(Buffer<float4>, SimRootRestPositionBuffer)
diff --git a/Engine/Plugins/Runtime/HairStrands/Source/HairStrandsCore/Public/HairStrandsDatas.h b/Engine/Plugins/Runtime/HairStrands/Source/HairStrandsCore/Public/HairStrandsDatas.h
--- a/Engine/Plugins/Runtime/HairStrands/Source/HairStrandsCore/Public/HairStrandsDatas.h
+++ b/Engine/Plugins/Runtime/HairStrands/Source/HairStrandsCore/Public/HairStrandsDatas.h
@@ -98,17 +98,7 @@ struct FHairStrandsAttributeFormat
 	static const EPixelFormat Format = PF_R32_UINT;
 };
 
-struct FHairStrandsPointToCurveFormat16
-{
-	typedef uint16 Type;
-	typedef uint16 BulkType;
-	static const uint32 ComponentCount = 1;
-	static const uint32 SizeInByte = sizeof(Type);
-	static const EVertexElementType VertexElementType = VET_MAX;
-	static const EPixelFormat Format = PF_R16_UINT;
-};
-
-struct FHairStrandsPointToCurveFormat32
+struct FHairStrandsPointToCurveFormat
 {
 	typedef uint32 Type;
 	typedef uint32 BulkType;
@@ -566,8 +556,7 @@ struct HAIRSTRANDSCORE_API FHairStrandsBulkData : FHairStrandsBulkCommon
 	enum EDataFlags
 	{
 		DataFlags_HasData = 1,				// Contains valid data. Otherwise: Position, Attributes, ... are all empty
-		DataFlags_Has16bitsCurveIndex = 2,	// Use 16bits index for vertex to curve mapping
-		DataFlags_HasPointAttribute = 4,	// Contains point attribute data.
+		DataFlags_HasPointAttribute = 2,	// Contains point attribute data.
 	};
 
 	virtual void SerializeHeader(FArchive& Ar, UObject* Owner) override;
@@ -589,6 +578,7 @@ struct HAIRSTRANDSCORE_API FHairStrandsBulkData : FHairStrandsBulkCommon
 
 	uint32 GetCurveAttributeSizeInBytes(uint32 InCurveCount=HAIR_MAX_NUM_CURVE_PER_GROUP) const	{ return InCurveCount > 0 ? FMath::DivideAndRoundUp(FMath::Min(Header.CurveCount, InCurveCount), Header.Strides.CurveAttributeChunkElementCount) * Header.Strides.CurveAttributeChunkStride : 0; }
 	uint32 GetPointAttributeSizeInBytes(uint32 InPointCount=HAIR_MAX_NUM_POINT_PER_GROUP) const	{ return InPointCount > 0 ? FMath::DivideAndRoundUp(FMath::Min(Header.PointCount, InPointCount), Header.Strides.PointAttributeChunkElementCount) * Header.Strides.PointAttributeChunkStride : 0; }
+	uint32 GetPointToCurveSizeInBytes(uint32 InPointCount=HAIR_MAX_NUM_POINT_PER_GROUP) const	{ return InPointCount > 0 ? FMath::DivideAndRoundUp(FMath::Min(Header.PointCount, InPointCount), Header.Strides.PointToCurveChunkElementCount)   * Header.Strides.PointToCurveChunkStride   : 0; }
 
 	struct FHeader
 	{
@@ -616,11 +606,12 @@ struct HAIRSTRANDSCORE_API FHairStrandsBulkData : FHairStrandsBulkCommon
 		{
 			uint32 PositionStride = 0;
 			uint32 CurveStride = 0;
-			uint32 PointToCurveStride = 0;
+			uint32 PointToCurveChunkStride = 0;
 			uint32 CurveAttributeChunkStride = 0;
 			uint32 PointAttributeChunkStride = 0;
 
 			// Number of element per chunk block
+			uint32 PointToCurveChunkElementCount = 0;
 			uint32 CurveAttributeChunkElementCount = 0;
 			uint32 PointAttributeChunkElementCount = 0;
 		} Strides;
diff --git a/Engine/Shaders/Private/HairStrands/HairStrandsAttributeCommon.ush b/Engine/Shaders/Private/HairStrands/HairStrandsAttributeCommon.ush
--- a/Engine/Shaders/Private/HairStrands/HairStrandsAttributeCommon.ush
+++ b/Engine/Shaders/Private/HairStrands/HairStrandsAttributeCommon.ush
@@ -66,13 +66,13 @@ uint2 InternalGetHairStrandsAttributes64(uint Index, uint AttributeOffsetInBytes
 
 uint InternalGetHairStrandsAttributesPerCurve(uint HairControlPointId, uint AttributeOffsetInBytes, uint AttributeStrideInBytes, ByteAddressBuffer InAttributeBuffer, Buffer<uint> InPointToCurveBuffer, uint InIndexToChunkDivAsShift, uint InChunkElementCount, uint InChunkStrideInBytes)
 {
-	const uint CurveIndex = InPointToCurveBuffer[HairControlPointId];
+	const uint CurveIndex = UnpackHairPointToCurveIndex(InPointToCurveBuffer, HairControlPointId);
 	return InternalGetHairStrandsAttributes(CurveIndex, AttributeOffsetInBytes, AttributeStrideInBytes, InAttributeBuffer, InIndexToChunkDivAsShift, InChunkElementCount, InChunkStrideInBytes);
 }
 
 uint2 InternalGetHairStrandsAttributesPerCurve64(uint HairControlPointId, uint AttributeOffsetInBytes, uint AttributeStrideInBytes, ByteAddressBuffer InAttributeBuffer, Buffer<uint> InPointToCurveBuffer, uint InIndexToChunkDivAsShift, uint InChunkElementCount, uint InChunkStrideInBytes)
 {
-	const uint CurveIndex = InPointToCurveBuffer[HairControlPointId];
+	const uint CurveIndex = UnpackHairPointToCurveIndex(InPointToCurveBuffer, HairControlPointId);
 	return InternalGetHairStrandsAttributes64(CurveIndex, AttributeOffsetInBytes, AttributeStrideInBytes, InAttributeBuffer, InIndexToChunkDivAsShift, InChunkElementCount, InChunkStrideInBytes);
 }
 
diff --git a/Engine/Shaders/Private/HairStrands/HairStrandsGuideDeform.usf b/Engine/Shaders/Private/HairStrands/HairStrandsGuideDeform.usf
--- a/Engine/Shaders/Private/HairStrands/HairStrandsGuideDeform.usf
+++ b/Engine/Shaders/Private/HairStrands/HairStrandsGuideDeform.usf
@@ -73,7 +73,7 @@ float3 DisplacePosition(uint VertexIndex, float3 Pos)
 #endif
 
 #if PERMUTATION_DEFORMATION == DEFORMATION_OFFSET_GUIDE_DYNAMIC
-	const uint RootIndex = SimPointToCurveBuffer[VertexIndex];
+	const uint RootIndex = UnpackHairPointToCurveIndex(SimPointToCurveBuffer, VertexIndex);
 	const float3 RootBarycentric = UnpackBarycentrics(SimRootBarycentricBuffer[RootIndex]);
 
 	const uint TriangleIndex = SimRootToUniqueTriangleIndexBuffer[RootIndex];
diff --git a/Engine/Shaders/Private/HairStrands/HairStrandsPack.ush b/Engine/Shaders/Private/HairStrands/HairStrandsPack.ush
--- a/Engine/Shaders/Private/HairStrands/HairStrandsPack.ush
+++ b/Engine/Shaders/Private/HairStrands/HairStrandsPack.ush
@@ -303,6 +303,22 @@ FHairCurve UnpackHairCurve(uint In)
 	return Out;
 }
 
+uint UnpackHairPointToCurveIndex(Buffer<uint> InBuffer, uint InPointIndex)
+{
+	// Encoding: base curve index + detla encoding for 8 consecutive control points
+
+	const uint PointIndex8 = InPointIndex >> 3u;
+	const uint LocalIndex8 = InPointIndex & 0x7;
+	const uint Data = InBuffer[PointIndex8];
+
+	// Extract the curve index delta encoding
+	const uint Mask     = BitFieldMaskU32(LocalIndex8+1, 0);
+	const uint Bitfield = BitFieldExtractI32(Data, 8u, 24u);
+	const uint Deltas   = countbits(Mask & Bitfield);
+
+	return (Data & 0xFFFFFF) + Deltas;
+}
+
 ////////////////////////////////////////////////////////////////////////////////////////////////////////////////////////////////
 // UVs
 
