Fix hair simiulation reset point function, doing arithmetic on packed data.

diff --git a/Engine/Plugins/Runtime/HairStrands/Shaders/Private/NiagaraDataInterfaceHairStrandsTemplate.ush b/Engine/Plugins/Runtime/HairStrands/Shaders/Private/NiagaraDataInterfaceHairStrandsTemplate.ush
--- a/Engine/Plugins/Runtime/HairStrands/Shaders/Private/NiagaraDataInterfaceHairStrandsTemplate.ush
+++ b/Engine/Plugins/Runtime/HairStrands/Shaders/Private/NiagaraDataInterfaceHairStrandsTemplate.ush
@@ -564,12 +564,14 @@ void ResetPointPosition_{ParameterName}(out bool ReportStatus)
 		const int PointBegin = Curve.PointIndex;
 		const int PointCount = Curve.PointCount;
 
-		const uint4 PositionOffset = uint4( f32tof16({ParameterName}_RestPositionOffset - {ParameterName}_DeformedPositionOffset[0].xyz), 0 );
+		const float3 PositionOffsetDelta = {ParameterName}_RestPositionOffset - {ParameterName}_DeformedPositionOffset[0].xyz;
 
 		ReportStatus = true;
 		for (int PointIndex = 0; PointIndex < PointCount; ++PointIndex)
 		{
-			{ParameterName}_DeformedPositionBuffer[PointBegin+PointIndex] = {ParameterName}_RestPositionBuffer[PointBegin+PointIndex] + PositionOffset;
+			const FPackedHairPosition Packed = ReadPackedHairPosition({ParameterName}_RestPositionBuffer, PointBegin+PointIndex);
+			const float3 CPPosition = UnpackHairControlPointPosition(Packed);
+			WritePackedHairControlPointPosition({ParameterName}_DeformedPositionBuffer, PointBegin + PointIndex, Packed, CPPosition + PositionOffsetDelta/*NewPosition*/);
 		}
 	}
 }
