GroomAsset: Update groom_group_cards_id to use FNames rather than integer values.

diff --git a/Engine/Plugins/Runtime/HairStrands/Source/HairStrandsCore/Private/HairAttributes.cpp b/Engine/Plugins/Runtime/HairStrands/Source/HairStrandsCore/Private/HairAttributes.cpp
--- a/Engine/Plugins/Runtime/HairStrands/Source/HairStrandsCore/Private/HairAttributes.cpp
+++ b/Engine/Plugins/Runtime/HairStrands/Source/HairStrandsCore/Private/HairAttributes.cpp
@@ -27,7 +27,7 @@ namespace HairAttribute
 	const FName Strand::CurveType("groom_curve_type");
 	const FName Strand::Knots("groom_knots");
 	const FName Strand::GroupName("groom_group_name");
-	const FName Strand::GroupCardsID("groom_group_cards_id");
+	const FName Strand::GroupCardsName("groom_group_cards_id");
 
 	const FName Groom::Color("groom_color");
 	const FName Groom::Roughness("groom_roughness");
diff --git a/Engine/Plugins/Runtime/HairStrands/Source/HairStrandsCore/Public/GroomAsset.h b/Engine/Plugins/Runtime/HairStrands/Source/HairStrandsCore/Public/GroomAsset.h
--- a/Engine/Plugins/Runtime/HairStrands/Source/HairStrandsCore/Public/GroomAsset.h
+++ b/Engine/Plugins/Runtime/HairStrands/Source/HairStrandsCore/Public/GroomAsset.h
@@ -76,9 +76,6 @@ struct HAIRSTRANDSCORE_API FHairGroupInfo
 	UPROPERTY(VisibleAnywhere, Category = "Info", meta = (DisplayName = "Length of the longest hair strands"))
 	float MaxCurveLength = 0;
 
-	UPROPERTY(VisibleAnywhere, Category = "Info", meta = (DisplayName = "Card Group ID for split hair card generation"))
-	int32 GroupCardsID = 0;
-
 	UPROPERTY()
 	uint32 Flags = 0;
 
diff --git a/Engine/Plugins/Runtime/HairStrands/Source/HairStrandsCore/Public/HairAttributes.h b/Engine/Plugins/Runtime/HairStrands/Source/HairStrandsCore/Public/HairAttributes.h
--- a/Engine/Plugins/Runtime/HairStrands/Source/HairStrandsCore/Public/HairAttributes.h
+++ b/Engine/Plugins/Runtime/HairStrands/Source/HairStrandsCore/Public/HairAttributes.h
@@ -32,7 +32,7 @@ namespace HairAttribute
 		extern HAIRSTRANDSCORE_API const FName CurveType;		// FName (EGroomCurveType)
 		extern HAIRSTRANDSCORE_API const FName Knots;			// float[]
 		extern HAIRSTRANDSCORE_API const FName GroupName;		// FName
-		extern HAIRSTRANDSCORE_API const FName GroupCardsID;	// int
+		extern HAIRSTRANDSCORE_API const FName GroupCardsName;	// FName
 	}
 
 	namespace Groom
