Make Groom PSO requests high priority and apply proxy creation skip method if PSOs are still compiling during proxy creation (fixes possible hitches coming from TooLate PSO precache compilation on Groom assets)

diff --git a/Engine/Plugins/Runtime/HairStrands/Source/HairStrandsCore/Private/GroomComponent.cpp b/Engine/Plugins/Runtime/HairStrands/Source/HairStrandsCore/Private/GroomComponent.cpp
--- a/Engine/Plugins/Runtime/HairStrands/Source/HairStrandsCore/Private/GroomComponent.cpp
+++ b/Engine/Plugins/Runtime/HairStrands/Source/HairStrandsCore/Private/GroomComponent.cpp
@@ -1871,6 +1871,12 @@ FPrimitiveSceneProxy* UGroomComponent::CreateSceneProxy()
 	if (!GroomAsset || GroomAsset->GetNumHairGroups() == 0 || HairGroupInstances.Num() == 0)
 		return nullptr;
 
+	if (CheckPSOPrecachingAndBoostPriority() && GetPSOPrecacheProxyCreationStrategy() == EPSOPrecacheProxyCreationStrategy::DelayUntilPSOPrecached)
+	{
+		UE_LOG(LogHairStrands, Verbose, TEXT("Skipping CreateSceneProxy for UGroomComponent %s (UGroomComponent PSOs are still compiling)"), *GetFullName());
+		return nullptr;
+	}
+
 	bool bIsValid = false;
 	for (const TRefCountPtr<FHairGroupInstance>& Instance : HairGroupInstances)
 	{
@@ -3145,6 +3151,7 @@ void UGroomComponent::CollectPSOPrecacheData(const FPSOPrecacheParams& BasePreca
 			ComponentParams.MaterialInterface = MaterialInterface;
 			ComponentParams.VertexFactoryDataList = VFsPerMaterial.VertexFactoryDataList;
 			ComponentParams.PSOPrecacheParams = PrecachePSOParams;
+			ComponentParams.Priority = EPSOPrecachePriority::High;
 		}
 	}
 
