Fix validation error on Vulkan.

diff --git a/Engine/Plugins/Runtime/HairStrands/Shaders/Private/NiagaraDataInterfaceHairStrandsTemplate.ush b/Engine/Plugins/Runtime/HairStrands/Shaders/Private/NiagaraDataInterfaceHairStrandsTemplate.ush
--- a/Engine/Plugins/Runtime/HairStrands/Shaders/Private/NiagaraDataInterfaceHairStrandsTemplate.ush
+++ b/Engine/Plugins/Runtime/HairStrands/Shaders/Private/NiagaraDataInterfaceHairStrandsTemplate.ush
@@ -564,7 +564,7 @@ void ResetPointPosition_{ParameterName}(out bool ReportStatus)
 		const int PointBegin = Curve.PointIndex;
 		const int PointCount = Curve.PointCount;
 
-		const float3 PositionOffsetDelta = {ParameterName}_RestPositionOffset - {ParameterName}_DeformedPositionOffset[0].xyz;
+		const float3 PositionOffsetDelta = {ParameterName}_RestPositionOffset.xyz - {ParameterName}_DeformedPositionOffset[0].xyz;
 
 		ReportStatus = true;
 		for (int PointIndex = 0; PointIndex < PointCount; ++PointIndex)
diff --git a/Engine/Plugins/Runtime/HairStrands/Source/HairStrandsCore/Private/HairStrandsMeshProjection.cpp b/Engine/Plugins/Runtime/HairStrands/Source/HairStrandsCore/Private/HairStrandsMeshProjection.cpp
--- a/Engine/Plugins/Runtime/HairStrands/Source/HairStrandsCore/Private/HairStrandsMeshProjection.cpp
+++ b/Engine/Plugins/Runtime/HairStrands/Source/HairStrandsCore/Private/HairStrandsMeshProjection.cpp
@@ -164,8 +164,8 @@ private:
 		SHADER_PARAMETER_SRV(Buffer, MeshUVsBuffer)
 
 		SHADER_PARAMETER_RDG_BUFFER_SRV(Buffer, UniqueTriangleIndices)
-		SHADER_PARAMETER_RDG_BUFFER_UAV(RWStructuredBuffer, OutUniqueTrianglePrevPosition)
-		SHADER_PARAMETER_RDG_BUFFER_UAV(RWStructuredBuffer, OutUniqueTriangleCurrPosition)
+		SHADER_PARAMETER_RDG_BUFFER_UAV(RWBuffer, OutUniqueTrianglePrevPosition)
+		SHADER_PARAMETER_RDG_BUFFER_UAV(RWBuffer, OutUniqueTriangleCurrPosition)
 	END_SHADER_PARAMETER_STRUCT()
 
 public:
diff --git a/Engine/Plugins/Runtime/HairStrands/Source/HairStrandsCore/Private/Niagara/NiagaraDataInterfaceHairStrands.cpp b/Engine/Plugins/Runtime/HairStrands/Source/HairStrandsCore/Private/Niagara/NiagaraDataInterfaceHairStrands.cpp
--- a/Engine/Plugins/Runtime/HairStrands/Source/HairStrandsCore/Private/Niagara/NiagaraDataInterfaceHairStrands.cpp
+++ b/Engine/Plugins/Runtime/HairStrands/Source/HairStrandsCore/Private/Niagara/NiagaraDataInterfaceHairStrands.cpp
@@ -22,6 +22,7 @@
 
 #include "GroomComponent.h"
 #include "GroomAsset.h"
+#include "SystemTextures.h"
 
 #include UE_INLINE_GENERATED_CPP_BY_NAME(NiagaraDataInterfaceHairStrands)
 
@@ -50,7 +51,7 @@ BEGIN_SHADER_PARAMETER_STRUCT(FShaderParameters,)
 	SHADER_PARAMETER(FVector3f,							RestRootOffset)
 	SHADER_PARAMETER(FVector3f,							DeformedRootOffset)
 	SHADER_PARAMETER(FVector3f,							RestPositionOffset)
-	SHADER_PARAMETER_RDG_BUFFER_SRV(Buffer<float4>,		DeformedPositionOffset)
+	SHADER_PARAMETER_RDG_BUFFER_SRV(StructuredBuffer<float4>, DeformedPositionOffset)
 	SHADER_PARAMETER_RDG_BUFFER_UAV(RWBuffer<uint>,		BoundingBoxBuffer)
 	SHADER_PARAMETER(uint32,							ResetSimulation)
 	SHADER_PARAMETER(uint32,							RestUpdate)
@@ -323,7 +324,7 @@ void FNDIHairStrandsBuffer::InitRHI(FRHICommandListBase&)
 		{
 			const uint32 PositionsCount = SourceDatas->GetNumPoints();
 
-			FRDGBufferDesc BufferDesc = FRDGBufferDesc::CreateBufferDesc(FHairStrandsPositionFormat::SizeInByte, PositionsCount);
+			FRDGBufferDesc BufferDesc = FRDGBufferDesc::CreateByteAddressDesc(FHairStrandsPositionFormat::SizeInByte * PositionsCount);
 			FRDGBufferRef RDGBuffer = GraphBuilder.CreateBuffer(BufferDesc, TEXT("DeformedPositionBuffer"));
 			DeformedPositionBuffer = GraphBuilder.ConvertToExternalBuffer(RDGBuffer);
 		}
@@ -827,7 +828,7 @@ class FInterpolateGroomGuidesCS : public FGlobalShader
 		SHADER_PARAMETER_RDG_BUFFER_UAV(RWByteAddressBuffer,DeformedPositionBuffer)
 		SHADER_PARAMETER_RDG_BUFFER_SRV(ByteAddressBuffer,	RestPositionBuffer)
 		SHADER_PARAMETER_RDG_BUFFER_SRV(Buffer<uint>,		CurvesOffsetsBuffer)
-		SHADER_PARAMETER_RDG_BUFFER_SRV(Buffer<float4>,	DeformedPositionOffset)
+		SHADER_PARAMETER_RDG_BUFFER_SRV(StructuredBuffer<float4>, DeformedPositionOffset)
 		SHADER_PARAMETER(FVector3f,				RestPositionOffset)
 		SHADER_PARAMETER(FMatrix44f,			WorldToLocal)
 
@@ -3368,6 +3369,8 @@ void UNiagaraDataInterfaceHairStrands::SetShaderParameters(const FNiagaraDataInt
 	{
 		check(ProxyData);
 
+		FRDGBufferSRVRef DummyStructuredBufferSRV = GraphBuilder.CreateSRV(GSystemTextures.GetDefaultStructuredBuffer(GraphBuilder, 16u, FUintVector4(0,0,0,0)));
+
 		FNDIHairStrandsBuffer* HairStrandsBuffer = ProxyData->HairStrandsBuffer;
 
 		FRDGBufferUAVRef DeformedPositionBufferUAV = nullptr;
@@ -3392,7 +3395,7 @@ void UNiagaraDataInterfaceHairStrands::SetShaderParameters(const FNiagaraDataInt
 		else
 		{
 			DeformedPositionBufferUAV = GraphBuilder.CreateUAV(GraphBuilder.RegisterExternalBuffer(HairStrandsBuffer->DeformedPositionBuffer));
-			DeformedPositionOffsetSRV = Context.GetComputeDispatchInterface().GetEmptyBufferSRV(GraphBuilder, FHairStrandsPositionOffsetFormat::Format);
+			DeformedPositionOffsetSRV = DummyStructuredBufferSRV;
 		}
 		const int32 MeshLODIndex = bIsRootValid ? HairStrandsBuffer->SourceDeformedRootResources->MeshLODIndex : -1;
 
@@ -3414,8 +3417,8 @@ void UNiagaraDataInterfaceHairStrands::SetShaderParameters(const FNiagaraDataInt
 		const bool bHasSamples = (RestMeshProjection && RestMeshProjection->SampleCount > 0);
 		const int32 SampleCountValue = bHasSamples ? RestMeshProjection->SampleCount : 0;
 
-		FRDGBufferSRVRef RestSamplePositionsBufferSRV = (bHasSamples && RestMeshProjection) ? RegisterAsSRV(GraphBuilder,RestMeshProjection->RestSamplePositionsBuffer) : Context.GetComputeDispatchInterface().GetEmptyBufferSRV(GraphBuilder, FHairStrandsMeshTrianglePositionFormat::Format);
-		FRDGBufferSRVRef MeshSampleWeightsBufferSRV = (bHasSamples && DeformedMeshProjection) ? RegisterAsSRV(GraphBuilder,DeformedMeshProjection->GetMeshSampleWeightsBuffer((FHairStrandsDeformedRootResource::FLOD::Current))) : Context.GetComputeDispatchInterface().GetEmptyBufferSRV(GraphBuilder, FHairStrandsMeshTrianglePositionFormat::Format);
+		FRDGBufferSRVRef RestSamplePositionsBufferSRV = (bHasSamples && RestMeshProjection) ? RegisterAsSRV(GraphBuilder,RestMeshProjection->RestSamplePositionsBuffer) : DummyStructuredBufferSRV;
+		FRDGBufferSRVRef MeshSampleWeightsBufferSRV = (bHasSamples && DeformedMeshProjection) ? RegisterAsSRV(GraphBuilder,DeformedMeshProjection->GetMeshSampleWeightsBuffer((FHairStrandsDeformedRootResource::FLOD::Current))) : DummyStructuredBufferSRV;
 		
 		// Simulation setup (we update the rest configuration based on the deformed positions 
 		// if in restupdate mode or if we are resetting the sim and using RBF transfer since the rest positions are not matrching the physics asset)
@@ -3516,17 +3519,20 @@ void UNiagaraDataInterfaceHairStrands::SetShaderParameters(const FNiagaraDataInt
 		ShaderParameters->SampleCount = 0;
 
 		// Set Shader UAV
-		ShaderParameters->DeformedPositionBuffer = Context.GetComputeDispatchInterface().GetEmptyBufferUAV(GraphBuilder, PF_R32_UINT);
+		ShaderParameters->DeformedPositionBuffer = GraphBuilder.CreateUAV(GraphBuilder.CreateBuffer(FRDGBufferDesc::CreateByteAddressDesc(16u), TEXT("Niagara.Hair.DummyUAV")), ERDGUnorderedAccessViewFlags::SkipBarrier);
 		ShaderParameters->BoundingBoxBuffer = Context.GetComputeDispatchInterface().GetEmptyBufferUAV(GraphBuilder, PF_R32_UINT);
 
 		// Set Shader SRV
+		FRDGBufferSRVRef DummyStructuredBuffer = GraphBuilder.CreateSRV(GSystemTextures.GetDefaultStructuredBuffer(GraphBuilder, 16u, FUintVector4(0,0,0,0)));
+		FRDGBufferSRVRef DummyByteAddressBuffer = GraphBuilder.CreateSRV(GSystemTextures.GetDefaultByteAddressBuffer(GraphBuilder, 16u));
+
 		ShaderParameters->CurvesOffsetsBuffer = Context.GetComputeDispatchInterface().GetEmptyBufferSRV(GraphBuilder, FHairStrandsCurveFormat::Format);
-		ShaderParameters->RestPositionBuffer = Context.GetComputeDispatchInterface().GetEmptyBufferSRV(GraphBuilder, FHairStrandsPositionFormat::Format);
-		ShaderParameters->DeformedPositionOffset = Context.GetComputeDispatchInterface().GetEmptyBufferSRV(GraphBuilder, FHairStrandsPositionOffsetFormat::Format);
+		ShaderParameters->RestPositionBuffer = DummyByteAddressBuffer;
+		ShaderParameters->DeformedPositionOffset = DummyStructuredBuffer;
 		ShaderParameters->RestTrianglePositionBuffer = Context.GetComputeDispatchInterface().GetEmptyBufferSRV(GraphBuilder, FHairStrandsMeshTrianglePositionFormat::Format);
 		ShaderParameters->DeformedTrianglePositionBuffer = Context.GetComputeDispatchInterface().GetEmptyBufferSRV(GraphBuilder, FHairStrandsMeshTrianglePositionFormat::Format);
-		ShaderParameters->RestSamplePositionsBuffer = Context.GetComputeDispatchInterface().GetEmptyBufferSRV(GraphBuilder, FHairStrandsMeshTrianglePositionFormat::Format);
-		ShaderParameters->MeshSampleWeightsBuffer = Context.GetComputeDispatchInterface().GetEmptyBufferSRV(GraphBuilder, FHairStrandsMeshTrianglePositionFormat::Format);
+		ShaderParameters->RestSamplePositionsBuffer = DummyStructuredBuffer;
+		ShaderParameters->MeshSampleWeightsBuffer = DummyStructuredBuffer;
 		ShaderParameters->RootBarycentricCoordinatesBuffer = Context.GetComputeDispatchInterface().GetEmptyBufferSRV(GraphBuilder, FHairStrandsRootBarycentricFormat::Format);
 		ShaderParameters->RootToUniqueTriangleIndexBuffer = Context.GetComputeDispatchInterface().GetEmptyBufferSRV(GraphBuilder, FHairStrandsRootToUniqueTriangleIndexFormat::Format);
 		ShaderParameters->ParamsScaleBuffer = Context.GetComputeDispatchInterface().GetEmptyBufferSRV(GraphBuilder, PF_R32_FLOAT);
diff --git a/Engine/Plugins/Runtime/HairStrands/Source/HairStrandsCore/Public/HairStrandsDatas.h b/Engine/Plugins/Runtime/HairStrands/Source/HairStrandsCore/Public/HairStrandsDatas.h
--- a/Engine/Plugins/Runtime/HairStrands/Source/HairStrandsCore/Public/HairStrandsDatas.h
+++ b/Engine/Plugins/Runtime/HairStrands/Source/HairStrandsCore/Public/HairStrandsDatas.h
@@ -77,7 +77,7 @@ struct FHairStrandsPositionFormat
 	static const uint32 ComponentCount = 1;
 	static const uint32 SizeInByte = sizeof(Type);
 	static const EVertexElementType VertexElementType = VET_UShort4;
-	static const EPixelFormat Format = PF_R32G32_UINT;
+	static const EPixelFormat Format = PF_Unknown;
 };
 
 struct FHairStrandsPositionOffsetFormat
