HairCardsVF uses PrimitiveID & GPUScene on mobile platform.

diff --git a/Engine/Plugins/Runtime/HairStrands/Source/HairStrandsCore/Private/HairCardsVertexFactory.cpp b/Engine/Plugins/Runtime/HairStrands/Source/HairStrandsCore/Private/HairCardsVertexFactory.cpp
--- a/Engine/Plugins/Runtime/HairStrands/Source/HairStrandsCore/Private/HairCardsVertexFactory.cpp
+++ b/Engine/Plugins/Runtime/HairStrands/Source/HairStrandsCore/Private/HairCardsVertexFactory.cpp
@@ -19,6 +19,7 @@
 #include "GlobalRenderResources.h"
 #include "DataDrivenShaderPlatformInfo.h"
 
+#define HAIR_CARDS_VF_PRIMITIVEID_STREAM_INDEX 13
 static int32 GHairCardCoverageBias = 0;
 static FAutoConsoleVariableRef CVarHairCardCoverageBias(TEXT("r.HairStrands.Cards.CoverageBias"), GHairCardCoverageBias, TEXT("Apply a texture LOD bias to coverage texture"));
 
@@ -289,12 +290,12 @@ void FHairCardsVertexFactory::ValidateCompiledResult(const FVertexFactoryType* T
 void FHairCardsVertexFactory::GetPSOPrecacheVertexFetchElements(EVertexInputStreamType VertexInputStreamType, FVertexDeclarationElementList& Elements)
 {
 	// Manual vertex fetch is available for this factory so only primitive ID stream is used
-	Elements.Add(FVertexElement(0, 0, VET_UInt, 13, 0, true));
+	Elements.Add(FVertexElement(0, 0, VET_UInt, HAIR_CARDS_VF_PRIMITIVEID_STREAM_INDEX, 0, true));
 }
 
 EPrimitiveIdMode FHairCardsVertexFactory::GetPrimitiveIdMode(ERHIFeatureLevel::Type In) const
 {
-	return In > ERHIFeatureLevel::ES3_1 ? PrimID_DynamicPrimitiveShaderData : PrimID_ForceZero;
+	return PrimID_DynamicPrimitiveShaderData;
 }
 
 void FHairCardsVertexFactory::SetData(const FDataType& InData)
@@ -339,12 +340,7 @@ void FHairCardsVertexFactory::InitResources(FRHICommandListBase& RHICmdList)
 	FVertexDeclarationElementList Elements;
 	SetPrimitiveIdStreamIndex(CurrentFeatureLevel, EVertexInputStreamType::Default, -1);
 
-	// Sanity check - When using mobile feature level, ensure we don't use primitive data as we set the PrimitiveIdStream to 0xFF for mobile
-	if (CurrentFeatureLevel <= ERHIFeatureLevel::ES3_1)
-	{
-		check(GetPrimitiveIdMode(CurrentFeatureLevel) == PrimID_ForceZero);
-	}
-	AddPrimitiveIdStreamElement(EVertexInputStreamType::Default, Elements, 13 /*AttributeIndex*/, 0xff /*AttributeIndex_Mobile*/);
+	AddPrimitiveIdStreamElement(EVertexInputStreamType::Default, Elements, HAIR_CARDS_VF_PRIMITIVEID_STREAM_INDEX /*AttributeIndex*/, HAIR_CARDS_VF_PRIMITIVEID_STREAM_INDEX /*AttributeIndex_Mobile*/);
 
 	// Note this is a local version of the VF's bSupportsManualVertexFetch, which take into account the feature level
 	const bool bManualFetch = SupportsManualVertexFetch(CurrentFeatureLevel);
