Split hair binding root bulk/resources into seprate data/resources per LOD to fix crash when concurrent mesh LOD are streamed-in.

diff --git a/Engine/Plugins/Runtime/HairStrands/Source/HairStrandsCore/Private/GroomAsset.cpp b/Engine/Plugins/Runtime/HairStrands/Source/HairStrandsCore/Private/GroomAsset.cpp
--- a/Engine/Plugins/Runtime/HairStrands/Source/HairStrandsCore/Private/GroomAsset.cpp
+++ b/Engine/Plugins/Runtime/HairStrands/Source/HairStrandsCore/Private/GroomAsset.cpp
@@ -276,8 +276,8 @@ void DumpLoadedGroomAssets(IConsoleVariable* InCVarPakTesterEnabled)
 FGroomBindingAssetMemoryStats FGroomBindingAssetMemoryStats::Get(const UGroomBindingAsset::FHairGroupPlatformData& InCPU, const UGroomBindingAsset::FHairGroupResource& InGPU)
 {
 	FGroomBindingAssetMemoryStats Out;
-	Out.CPU.Guides  = InCPU.SimRootBulkData.GetDataSize();
-	Out.CPU.Strands = InCPU.RenRootBulkData.GetDataSize();
+	for (const FHairStrandsRootBulkData& In : InCPU.SimRootBulkDatas) { Out.CPU.Guides += In.GetDataSize(); }
+	for (const FHairStrandsRootBulkData& In : InCPU.RenRootBulkDatas) { Out.CPU.Strands+= In.GetDataSize(); }
 
 	if (const FHairStrandsRestRootResource* R = InGPU.SimRootResources)
 	{
@@ -288,10 +288,10 @@ FGroomBindingAssetMemoryStats FGroomBindingAssetMemoryStats::Get(const UGroomBin
 		Out.GPU.Strands += R->GetResourcesSize();
 	}
 
-	const uint32 CardCount = InCPU.CardsRootBulkData.Num();
+	const uint32 CardCount = InCPU.CardsRootBulkDatas.Num();
 	for (uint32 CardIt = 0; CardIt < CardCount; ++CardIt)
 	{
-		Out.CPU.Cards += InCPU.CardsRootBulkData[CardIt].GetDataSize();
+		for (const FHairStrandsRootBulkData& In : InCPU.CardsRootBulkDatas[CardIt]) { Out.CPU.Cards += In.GetDataSize(); }
 		if (const FHairStrandsRestRootResource* R = InGPU.CardsRootResources[CardIt])
 		{
 			Out.GPU.Cards += R->GetResourcesSize();
@@ -354,7 +354,7 @@ void DumpLoadedGroomBindingAssets(IConsoleVariable* InCVarPakTesterEnabled)
 				const FGroomBindingAssetMemoryStats Group = FGroomBindingAssetMemoryStats::Get(AssetIt->GetHairGroupsPlatformData()[GroupIt], AssetIt->GetHairGroupResources()[GroupIt]);
 				Total.Accumulate(Group);
 
-				const uint32 SkelLODCount = AssetIt->GetHairGroupsPlatformData()[GroupIt].RenRootBulkData.GetLODCount();
+				const uint32 SkelLODCount = AssetIt->GetHairGroupsPlatformData()[GroupIt].RenRootBulkDatas.Num();
 				if (bDetails)
 				{
 //					UE_LOG(LogHairStrands, Log, TEXT("--  No.  - LOD -    CPU Total (     Guides|    Strands|      Cards) -    GPU Total (     Guides|    Strands|      Cards) - Asset Name "));
@@ -1780,9 +1780,9 @@ static void InternalSerializeStrand(FArchive& Ar, UObject* Owner, FHairGroupPlat
 	const bool bFillInBulkDataForCooking = Ar.IsCooking() && Ar.IsSaving();
 	if (bFillInBulkDataForCooking)
 	{
-		{ FHairStreamingRequest R; R.Request(HAIR_MAX_NUM_CURVE_PER_GROUP, HAIR_MAX_NUM_POINT_PER_GROUP, -1/*LODIndex*/, StrandData.BulkData,               true /*bWait*/, true /*bFillBulkdata*/, true /*bWarmCache*/, Owner->GetFName()); }
-		{ FHairStreamingRequest R; R.Request(HAIR_MAX_NUM_CURVE_PER_GROUP, HAIR_MAX_NUM_POINT_PER_GROUP, -1/*LODIndex*/, StrandData.InterpolationBulkData,  true /*bWait*/, true /*bFillBulkdata*/, true /*bWarmCache*/, Owner->GetFName()); }
-		{ FHairStreamingRequest R; R.Request(HAIR_MAX_NUM_CURVE_PER_GROUP, HAIR_MAX_NUM_POINT_PER_GROUP, -1/*LODIndex*/, StrandData.ClusterBulkData,        true /*bWait*/, true /*bFillBulkdata*/, true /*bWarmCache*/, Owner->GetFName()); }
+		{ FHairStreamingRequest R; R.Request(HAIR_MAX_NUM_CURVE_PER_GROUP, HAIR_MAX_NUM_POINT_PER_GROUP, StrandData.BulkData,               true /*bWait*/, true /*bFillBulkdata*/, true /*bWarmCache*/, Owner->GetFName()); }
+		{ FHairStreamingRequest R; R.Request(HAIR_MAX_NUM_CURVE_PER_GROUP, HAIR_MAX_NUM_POINT_PER_GROUP, StrandData.InterpolationBulkData,  true /*bWait*/, true /*bFillBulkdata*/, true /*bWarmCache*/, Owner->GetFName()); }
+		{ FHairStreamingRequest R; R.Request(HAIR_MAX_NUM_CURVE_PER_GROUP, HAIR_MAX_NUM_POINT_PER_GROUP, StrandData.ClusterBulkData,        true /*bWait*/, true /*bFillBulkdata*/, true /*bWarmCache*/, Owner->GetFName()); }
 	}
 
 	if (!Ar.IsCooking() || !StrandData.bIsCookedOut)
@@ -1803,9 +1803,9 @@ static void InternalSerializeStrand(FArchive& Ar, UObject* Owner, FHairGroupPlat
 		if (bPreWarmCache)
 		{
 			bool bHasDataInCache = true;
-			{ FHairStreamingRequest R; bHasDataInCache &= R.WarmCache(HAIR_MAX_NUM_CURVE_PER_GROUP, HAIR_MAX_NUM_POINT_PER_GROUP, -1/*LODIndex*/, StrandData.BulkData); }
-			{ FHairStreamingRequest R; bHasDataInCache &= R.WarmCache(HAIR_MAX_NUM_CURVE_PER_GROUP, HAIR_MAX_NUM_POINT_PER_GROUP, -1/*LODIndex*/, StrandData.InterpolationBulkData); }
-			{ FHairStreamingRequest R; bHasDataInCache &= R.WarmCache(HAIR_MAX_NUM_CURVE_PER_GROUP, HAIR_MAX_NUM_POINT_PER_GROUP, -1/*LODIndex*/, StrandData.ClusterBulkData); }
+			{ FHairStreamingRequest R; bHasDataInCache &= R.WarmCache(HAIR_MAX_NUM_CURVE_PER_GROUP, HAIR_MAX_NUM_POINT_PER_GROUP, StrandData.BulkData); }
+			{ FHairStreamingRequest R; bHasDataInCache &= R.WarmCache(HAIR_MAX_NUM_CURVE_PER_GROUP, HAIR_MAX_NUM_POINT_PER_GROUP, StrandData.InterpolationBulkData); }
+			{ FHairStreamingRequest R; bHasDataInCache &= R.WarmCache(HAIR_MAX_NUM_CURVE_PER_GROUP, HAIR_MAX_NUM_POINT_PER_GROUP, StrandData.ClusterBulkData); }
 
 			if (bOutHasDataInCache) { *bOutHasDataInCache = bHasDataInCache; }
 		}
diff --git a/Engine/Plugins/Runtime/HairStrands/Source/HairStrandsCore/Private/GroomBindingAsset.cpp b/Engine/Plugins/Runtime/HairStrands/Source/HairStrandsCore/Private/GroomBindingAsset.cpp
--- a/Engine/Plugins/Runtime/HairStrands/Source/HairStrandsCore/Private/GroomBindingAsset.cpp
+++ b/Engine/Plugins/Runtime/HairStrands/Source/HairStrandsCore/Private/GroomBindingAsset.cpp
@@ -35,47 +35,62 @@ uint32 GetAssetNameHash(const FString& In);
 
 ////////////////////////////////////////////////////////////////////////////////////////////////////////////
 
-static void InternalSerializeGuides(FArchive& Ar, UObject* Owner, FHairStrandsRootBulkData& Data)
+static void InternalSerializeGuides(FArchive& Ar, UObject* Owner, TArray<FHairStrandsRootBulkData>& Datas)
 {
-	Data.SerializeHeader(Ar, Owner);
-	for (int32 LODIndex = 0, LODCount = Data.GetLODCount(); LODIndex < LODCount; ++LODIndex)
+	uint32 MeshLODCount = Datas.Num();
+	Ar << MeshLODCount;
+	if (Ar.IsLoading())
 	{
-		Data.SerializeData(Ar, Owner, LODIndex);
+		Datas.SetNum(MeshLODCount);
+	}
+	for (FHairStrandsRootBulkData& Data : Datas)
+	{
+		Data.SerializeHeader(Ar, Owner);
+		Data.SerializeData(Ar, Owner);
 	}
 }
 
-static void InternalSerializeStrands(FArchive& Ar, UObject* Owner, FHairStrandsRootBulkData& Data, uint32 Flags, bool bHeader, bool bData)
+static void InternalSerializeStrands(FArchive& Ar, UObject* Owner, TArray<FHairStrandsRootBulkData>& Datas, uint32 Flags, bool bHeader, bool bData)
 {
 	Ar.UsingCustomVersion(FAnimObjectVersion::GUID);
 
 	const bool bStripped = (Flags & UGroomAsset::CDSF_StrandsStripped);
 	if (!bStripped)
 	{
-		if (bHeader){ Data.SerializeHeader(Ar, Owner); }
-		if (bData)
+		uint32 MeshLODCount = Datas.Num();
+		Ar << MeshLODCount;
+		if (Ar.IsLoading())
 		{
-			for (int32 LODIndex = 0, LODCount = Data.GetLODCount(); LODIndex < LODCount; ++LODIndex)
-			{
-				Data.SerializeData(Ar, Owner, LODIndex);
-			}
+			Datas.SetNum(MeshLODCount);
+		}
+		for (FHairStrandsRootBulkData& Data : Datas)
+		{
+			if (bHeader){ Data.SerializeHeader(Ar, Owner); }
+			if (bData)	{ Data.SerializeData(Ar, Owner); }
 		}
 	}
 }
 
-static void InternalSerializeCards(FArchive& Ar, UObject* Owner, TArray<FHairStrandsRootBulkData>& Datas)
+static void InternalSerializeCards(FArchive& Ar, UObject* Owner, TArray<TArray<FHairStrandsRootBulkData>>& Datass)
 {
-	uint32 CardLODCount = Datas.Num();
+	uint32 CardLODCount = Datass.Num();
 	Ar << CardLODCount;
 	if (Ar.IsLoading())
 	{
-		Datas.SetNum(CardLODCount);
+		Datass.SetNum(CardLODCount);
 	}
-	for (FHairStrandsRootBulkData& Data : Datas)
+	for (TArray<FHairStrandsRootBulkData>& Datas : Datass)
 	{
-		Data.SerializeHeader(Ar, Owner);
-		for (int32 LODIndex = 0, LODCount = Data.GetLODCount(); LODIndex < LODCount; ++LODIndex)
+		uint32 MeshLODCount = Datas.Num();
+		Ar << MeshLODCount;
+		if (Ar.IsLoading())
 		{
-			Data.SerializeData(Ar, Owner, LODIndex);
+			Datas.SetNum(MeshLODCount);
+		}
+		for (FHairStrandsRootBulkData& Data : Datas)
+		{	
+			Data.SerializeHeader(Ar, Owner);
+			Data.SerializeData(Ar, Owner);
 		}
 	}
 }
@@ -85,13 +100,13 @@ static void InternalSerializePlatformData(FArchive& Ar, UObject* Owner, UGroomBi
 	Ar.UsingCustomVersion(FAnimObjectVersion::GUID);
 
 	// Guides
-	InternalSerializeGuides(Ar, Owner, GroupData.SimRootBulkData);
+	InternalSerializeGuides(Ar, Owner, GroupData.SimRootBulkDatas);
 
 	// Strands
-	InternalSerializeStrands(Ar, Owner, GroupData.RenRootBulkData, Flags, bHeader, bData);
+	InternalSerializeStrands(Ar, Owner, GroupData.RenRootBulkDatas, Flags, bHeader, bData);
 
 	// Cards
-	InternalSerializeCards(Ar, Owner, GroupData.CardsRootBulkData);
+	InternalSerializeCards(Ar, Owner, GroupData.CardsRootBulkDatas);
 }
 
 static void InternalSerializePlatformDatas(FArchive& Ar, UObject* Owner, TArray<UGroomBindingAsset::FHairGroupPlatformData>& GroupDatas, uint32 Flags)
@@ -167,32 +182,32 @@ void UGroomBindingAsset::InitResource()
 
 		// Guides
 		Resource.SimRootResources = nullptr;
-		if (BulkData.SimRootBulkData.IsValid())
+		if (BulkData.SimRootBulkDatas.Num() > 0)
 		{
-			Resource.SimRootResources = new FHairStrandsRestRootResource(BulkData.SimRootBulkData, EHairStrandsResourcesType::Guides, ResourceName, OwnerName);
-			BeginInitResource(Resource.SimRootResources);
+			Resource.SimRootResources = new FHairStrandsRestRootResource(BulkData.SimRootBulkDatas, EHairStrandsResourcesType::Guides, ResourceName, OwnerName);
+			Resource.SimRootResources->BeginInitResource();
 		}
 
 		// Strands
 		Resource.RenRootResources = nullptr;
-		if (IsHairStrandsEnabled(EHairStrandsShaderType::Strands) && BulkData.RenRootBulkData.IsValid())
+		if (IsHairStrandsEnabled(EHairStrandsShaderType::Strands) && BulkData.RenRootBulkDatas.Num() > 0)
 		{
-			Resource.RenRootResources = new FHairStrandsRestRootResource(BulkData.RenRootBulkData, EHairStrandsResourcesType::Strands, ResourceName, OwnerName);
-			BeginInitResource(Resource.RenRootResources);
+			Resource.RenRootResources = new FHairStrandsRestRootResource(BulkData.RenRootBulkDatas, EHairStrandsResourcesType::Strands, ResourceName, OwnerName);
+			Resource.RenRootResources->BeginInitResource();
 		}
 
 		// Cards
 		if (IsHairStrandsEnabled(EHairStrandsShaderType::Cards))
 		{
-			const uint32 CardsLODCount = BulkData.CardsRootBulkData.Num();
+			const uint32 CardsLODCount = BulkData.CardsRootBulkDatas.Num();
 			Resource.CardsRootResources.SetNum(CardsLODCount);
 			for (uint32 CardsLODIt=0; CardsLODIt<CardsLODCount; ++CardsLODIt)
 			{
 				Resource.CardsRootResources[CardsLODIt] = nullptr;
-				if (BulkData.CardsRootBulkData[CardsLODIt].IsValid())
+				if (BulkData.CardsRootBulkDatas[CardsLODIt].Num() > 0)
 				{
-					Resource.CardsRootResources[CardsLODIt] = new FHairStrandsRestRootResource(BulkData.CardsRootBulkData[CardsLODIt], EHairStrandsResourcesType::Cards, FHairResourceName(GetFName(), GroupIndex, CardsLODIt), GetAssetPathName(CardsLODIt));
-					BeginInitResource(Resource.CardsRootResources[CardsLODIt]);
+					Resource.CardsRootResources[CardsLODIt] = new FHairStrandsRestRootResource(BulkData.CardsRootBulkDatas[CardsLODIt], EHairStrandsResourcesType::Cards, FHairResourceName(GetFName(), GroupIndex, CardsLODIt), GetAssetPathName(CardsLODIt));
+					Resource.CardsRootResources[CardsLODIt]->BeginInitResource();
 				}
 			}
 		}
@@ -205,19 +220,19 @@ void UGroomBindingAsset::UpdateResource()
 	{
 		if (Resource.SimRootResources)
 		{
-			BeginUpdateResourceRHI(Resource.SimRootResources);
+			Resource.SimRootResources->BeginUpdateResourceRHI();
 		}
 
 		if (Resource.RenRootResources)
 		{
-			BeginUpdateResourceRHI(Resource.RenRootResources);
+			Resource.RenRootResources->BeginUpdateResourceRHI();
 		}
 
 		for (FHairStrandsRestRootResource* CardsRootResource : Resource.CardsRootResources)
 		{
 			if (CardsRootResource)
 			{
-				BeginUpdateResourceRHI(CardsRootResource);
+				CardsRootResource->BeginUpdateResourceRHI();
 			}
 		}
 	}
@@ -296,12 +311,12 @@ void UGroomBindingAsset::Reset()
 	ReleaseResource();
 	for (UGroomBindingAsset::FHairGroupPlatformData& Data : GetHairGroupsPlatformData())
 	{
-		Data.SimRootBulkData.Reset();
-		Data.RenRootBulkData.Reset();
+		Data.SimRootBulkDatas.Empty();
+		Data.RenRootBulkDatas.Empty();
 
-		for (FHairStrandsRootBulkData& CardsRootBulkData : Data.CardsRootBulkData)
+		for (TArray<FHairStrandsRootBulkData>& CardsRootBulkData : Data.CardsRootBulkDatas)
 		{
-			CardsRootBulkData.Reset();
+			CardsRootBulkData.Empty();
 		}
 	}
 
@@ -735,11 +750,11 @@ void UpdateGroomBindingAssetInfos(UGroomBindingAsset* In)
 		{
 			FGoomBindingGroupInfo& Info = In->GetGroupInfos()[GroupIt];
 			const UGroomBindingAsset::FHairGroupPlatformData& BulkData = In->GetHairGroupsPlatformData()[GroupIt];
-			Info.SimRootCount = BulkData.SimRootBulkData.GetRootCount();
-			Info.SimLODCount  = BulkData.SimRootBulkData.GetLODCount();
+			Info.SimRootCount = BulkData.SimRootBulkDatas.Num() > 0 ? BulkData.SimRootBulkDatas[0].GetRootCount() : 0u;
+			Info.SimLODCount  = BulkData.SimRootBulkDatas.Num();
 
-			Info.RenRootCount = BulkData.RenRootBulkData.GetRootCount();
-			Info.RenLODCount  = BulkData.RenRootBulkData.GetLODCount();
+			Info.RenRootCount = BulkData.RenRootBulkDatas.Num() > 0 ? BulkData.RenRootBulkDatas[0].GetRootCount() : 0u;
+			Info.RenLODCount  = BulkData.RenRootBulkDatas.Num();
 		}
 	}
 }
@@ -916,7 +931,10 @@ static void CacheDerivedDatas(UGroomBindingAsset* In, const uint32 InGroupIndex,
 		// Populate key/name for streaming data request
 		auto FillDrivedDataKey = [&DerivedDataKey, &Name](UGroomBindingAsset::FHairGroupPlatformData& In)
 		{
-			In.RenRootBulkData.DerivedDataKey = DerivedDataKey + FString(TEXT("_Data"));
+			for (uint32 MeshLODIndex = 0, MeshLODCount=In.RenRootBulkDatas.Num(); MeshLODIndex < MeshLODCount; ++MeshLODIndex)
+			{
+				In.RenRootBulkDatas[MeshLODIndex].DerivedDataKey = DerivedDataKey + FString::Printf(TEXT("_RenRootData_MeshLOD%d"), MeshLODIndex);
+			}
 		};
 
 		bool bHasDataInCache = false;
@@ -933,9 +951,9 @@ static void CacheDerivedDatas(UGroomBindingAsset* In, const uint32 InGroupIndex,
 
 			// Verify that all strands data are correctly cached into the DDC
 			{
-				for (int32 LODIndex = 0, LODCount = OutPlatformData.RenRootBulkData.GetLODCount(); LODIndex < LODCount; ++LODIndex)
+				for (int32 MeshLODIndex = 0, MeshLODCount = OutPlatformData.RenRootBulkDatas.Num(); MeshLODIndex < MeshLODCount; ++MeshLODIndex)
 				{
-					FHairStreamingRequest R; bHasDataInCache &= R.WarmCache(HAIR_MAX_NUM_CURVE_PER_GROUP, HAIR_MAX_NUM_POINT_PER_GROUP, LODIndex, OutPlatformData.RenRootBulkData);
+					FHairStreamingRequest R; bHasDataInCache &= R.WarmCache(HAIR_MAX_NUM_CURVE_PER_GROUP, HAIR_MAX_NUM_POINT_PER_GROUP, OutPlatformData.RenRootBulkDatas[MeshLODIndex]);
 				}
 			}
 
@@ -975,10 +993,10 @@ static void CacheDerivedDatas(UGroomBindingAsset* In, const uint32 InGroupIndex,
 				}
 	
 				// Data
-				for (uint32 LODIt=0, LODCount = OutPlatformData.RenRootBulkData.GetLODCount(); LODIt<LODCount;++LODIt)
+				for (uint32 MeshLODIndex=0, MeshLODCount = OutPlatformData.RenRootBulkDatas.Num(); MeshLODIndex<MeshLODCount;++MeshLODIndex)
 				{
 					TArray<FCachePutValueRequest> Out;
-					OutPlatformData.RenRootBulkData.Write_DDC(In, Out, LODIt);
+					OutPlatformData.RenRootBulkDatas[MeshLODIndex].Write_DDC(In, Out);
 
 					FRequestOwner AsyncOwner(EPriority::Normal);
 					GetCache().PutValue(Out, AsyncOwner);
@@ -1053,10 +1071,10 @@ void UGroomBindingAsset::BeginCacheForCookedPlatformData(const ITargetPlatform*
 		// from the 'start' (i.e., without offset)
 		for (uint32 GroupIndex = 0; GroupIndex < GroupCount; ++GroupIndex)
 		{
-			FHairStrandsRootBulkData& RenRootBulkData = TargetPlatformData->GroupPlatformDatas[GroupIndex].RenRootBulkData;
-			for (int32 LODIndex = 0, LODCount = RenRootBulkData.GetLODCount(); LODIndex < LODCount; ++LODIndex)
+			TArray<FHairStrandsRootBulkData>& RenRootBulkDatas = TargetPlatformData->GroupPlatformDatas[GroupIndex].RenRootBulkDatas;
+			for (int32 MeshLODIndex = 0, MeshLODCount = RenRootBulkDatas.Num(); MeshLODIndex < MeshLODCount; ++MeshLODIndex)
 			{
-				FHairStreamingRequest R; R.Request(HAIR_MAX_NUM_CURVE_PER_GROUP, HAIR_MAX_NUM_POINT_PER_GROUP, LODIndex, RenRootBulkData, true /*bWait*/, true /*bFillBulkdata*/, true /*bWarmCache*/, GetFName());
+				FHairStreamingRequest R; R.Request(HAIR_MAX_NUM_CURVE_PER_GROUP, HAIR_MAX_NUM_POINT_PER_GROUP, RenRootBulkDatas[MeshLODIndex], true /*bWait*/, true /*bFillBulkdata*/, true /*bWarmCache*/, GetFName());
 			}
 		}
 
diff --git a/Engine/Plugins/Runtime/HairStrands/Source/HairStrandsCore/Private/GroomBindingBuilder.cpp b/Engine/Plugins/Runtime/HairStrands/Source/HairStrandsCore/Private/GroomBindingBuilder.cpp
--- a/Engine/Plugins/Runtime/HairStrands/Source/HairStrandsCore/Private/GroomBindingBuilder.cpp
+++ b/Engine/Plugins/Runtime/HairStrands/Source/HairStrandsCore/Private/GroomBindingBuilder.cpp
@@ -60,7 +60,7 @@ static FAutoConsoleVariableRef CVarHairStrandsBindingBuilderWarningEnable(TEXT("
 FString FGroomBindingBuilder::GetVersion()
 {
 	// Important to update the version when groom building changes
-	return TEXT("3p_15");
+	return TEXT("4a");
 }
 
 ///////////////////////////////////////////////////////////////////////////////////////////////////
@@ -111,9 +111,9 @@ float FHairStrandsRootUtils::PackUVsToFloat(const FVector2f& UV)
 /** Binding data */
 struct FHairRootGroupData
 {
-	FHairStrandsRootData			SimRootData;
-	FHairStrandsRootData			RenRootData;
-	TArray<FHairStrandsRootData>	CardsRootData;
+	TArray<FHairStrandsRootData>		SimRootDatas;
+	TArray<FHairStrandsRootData>		RenRootDatas;
+	TArray<TArray<FHairStrandsRootData>>CardsRootDatas;
 };
 
 namespace GroomBinding_Mesh
@@ -693,7 +693,7 @@ namespace GroomBinding_RBFWeighting
 		const FWeightsBuilder& InterpolationWeights,
 		const FPointsSampler& PointsSampler,
 		const GroomBinding_Mesh::IMeshLODData& MeshData,
-		FHairStrandsRootData::FMeshProjectionLOD& OutRootLODData)
+		FHairStrandsRootData& OutRootLODData)
 	{
 		OutRootLODData.MeshSampleIndicesBuffer.SetNum(PointsSampler.SampleIndices.Num());
 		OutRootLODData.MeshInterpolationWeightsBuffer.SetNum(InterpolationWeights.InverseEntries.Num());
@@ -713,7 +713,7 @@ namespace GroomBinding_RBFWeighting
 	void FillLocalValidPoints(
 		const GroomBinding_Mesh::IMeshLODData& MeshLODData, 
 		const int32 TargetSection,
-		const FHairStrandsRootData::FMeshProjectionLOD& ProjectionLOD, 
+		const FHairStrandsRootData& ProjectionLOD, 
 		TArray<bool>& OutValidPoints)
 	{
 		const TArray<uint32>& TriangleIndices = MeshLODData.GetIndexBuffer();
@@ -762,13 +762,12 @@ namespace GroomBinding_RBFWeighting
 		}
 	}
 
-	void ResetSampleData(FHairStrandsRootData& Out, uint32 LODIndex)
+	void ResetSampleData(FHairStrandsRootData& Out)
 	{
-		FHairStrandsRootData::FMeshProjectionLOD& LOD = Out.MeshProjectionLODs[LODIndex];
-		LOD.SampleCount = 0;	
-		LOD.MeshInterpolationWeightsBuffer.SetNum(0);
-		LOD.MeshSampleIndicesBuffer.SetNum(0);
-		LOD.RestSamplePositionsBuffer.SetNum(0);
+		Out.SampleCount = 0;
+		Out.MeshInterpolationWeightsBuffer.SetNum(0);
+		Out.MeshSampleIndicesBuffer.SetNum(0);
+		Out.RestSamplePositionsBuffer.SetNum(0);
 	}
 
 	void ComputeInterpolationWeights(
@@ -782,16 +781,16 @@ namespace GroomBinding_RBFWeighting
 		const uint32 MeshLODCount= MeshData->GetNumLODs();
 		const uint32 MaxSamples  = NumInterpolationPoints;
 
-		for (uint32 LODIndex = 0; LODIndex < MeshLODCount; ++LODIndex)
+		for (uint32 MeshLODIndex = 0; MeshLODIndex < MeshLODCount; ++MeshLODIndex)
 		{
-			const GroomBinding_Mesh::IMeshLODData& MeshLODData = MeshData->GetMeshLODData(LODIndex);
+			const GroomBinding_Mesh::IMeshLODData& MeshLODData = MeshData->GetMeshLODData(MeshLODIndex);
 
 			int32 TargetSection = -1;
 			bool GlobalSamples = false;
 			const FVector3f* PositionsPointer = nullptr;
 			if (TransferedPositions.Num() == MeshLODCount)
 			{
-				PositionsPointer = TransferedPositions[LODIndex].GetData();
+				PositionsPointer = TransferedPositions[MeshLODIndex].GetData();
 				GlobalSamples = true;
 				TargetSection = MatchingSection;
 			}
@@ -804,7 +803,7 @@ namespace GroomBinding_RBFWeighting
 			{
 				TArray<bool> ValidPoints;
 				{
-					FillLocalValidPoints(MeshLODData, TargetSection, Out.SimRootData.MeshProjectionLODs[LODIndex], ValidPoints);
+					FillLocalValidPoints(MeshLODData, TargetSection, Out.SimRootDatas[MeshLODIndex], ValidPoints);
 
 					FPointsSampler PointsSampler(ValidPoints, PositionsPointer, MaxSamples);
 					const uint32 SampleCount = PointsSampler.SamplePositions.Num();
@@ -812,13 +811,13 @@ namespace GroomBinding_RBFWeighting
 					FWeightsBuilder InterpolationWeights(SampleCount, SampleCount, PointsSampler.SamplePositions.GetData(), PointsSampler.SamplePositions.GetData());
 
 					// Guides
-					UpdateInterpolationWeights(InterpolationWeights, PointsSampler, MeshLODData, Out.SimRootData.MeshProjectionLODs[LODIndex]);
+					UpdateInterpolationWeights(InterpolationWeights, PointsSampler, MeshLODData, Out.SimRootDatas[MeshLODIndex]);
 
 					// Strands
 					// No sample data, only used/available for guides
 					if (bNeedStrandsRoot)
 					{
-						ResetSampleData(Out.RenRootData, LODIndex);
+						ResetSampleData(Out.RenRootDatas[MeshLODIndex]);
 					}
 				}
 			}
@@ -834,13 +833,13 @@ namespace GroomBinding_RBFWeighting
 				FWeightsBuilder InterpolationWeights(SampleCount, SampleCount, PointsSampler.SamplePositions.GetData(), PointsSampler.SamplePositions.GetData());
 
 				// Guides
-				UpdateInterpolationWeights(InterpolationWeights, PointsSampler, MeshLODData, Out.SimRootData.MeshProjectionLODs[LODIndex]);
+				UpdateInterpolationWeights(InterpolationWeights, PointsSampler, MeshLODData, Out.SimRootDatas[MeshLODIndex]);
 
 				// Strands 
 				// No sample data, only used/available for guides
 				if (bNeedStrandsRoot)
 				{
-					ResetSampleData(Out.RenRootData, LODIndex);
+					ResetSampleData(Out.RenRootDatas[MeshLODIndex]);
 				}
 			}
 		}
@@ -1150,14 +1149,14 @@ namespace GroomBinding_RootProjection
 		const FHairStrandsDatas& InStrandsData,
 		const GroomBinding_Mesh::IMeshData* InMeshData,
 		const TArray<TArray<FVector3f>>& InTransferredPositions,
-		FHairStrandsRootData& OutRootData)
+		TArray<FHairStrandsRootData>& OutRootData)
 	{
 		// 2. Project root for each mesh LOD
 		const uint32 CurveCount = InStrandsData.GetNumCurves();
 		const uint32 ChannelIndex = 0;
 		const float VoxelWorldSize = 2; //cm
 		const uint32 MeshLODCount = InMeshData->GetNumLODs();
-		check(MeshLODCount == OutRootData.MeshProjectionLODs.Num());
+		check(MeshLODCount == OutRootData.Num());
 
 		const bool bHasTransferredPosition = InTransferredPositions.Num() > 0;
 		if (bHasTransferredPosition)
@@ -1165,12 +1164,12 @@ namespace GroomBinding_RootProjection
 			check(InTransferredPositions.Num() == MeshLODCount);
 		}
 
-		for (uint32 LODIt = 0; LODIt < MeshLODCount; ++LODIt)
+		for (uint32 MeshLODIt = 0; MeshLODIt < MeshLODCount; ++MeshLODIt)
 		{
-			check(LODIt == OutRootData.MeshProjectionLODs[LODIt].LODIndex);
+			check(MeshLODIt == OutRootData[MeshLODIt].LODIndex);
 
 			// 2.1. Build a grid around the hair AABB
-			const GroomBinding_Mesh::IMeshLODData& MeshLODData = InMeshData->GetMeshLODData(LODIt);
+			const GroomBinding_Mesh::IMeshLODData& MeshLODData = InMeshData->GetMeshLODData(MeshLODIt);
 			const TArray<uint32>& IndexBuffer = MeshLODData.GetIndexBuffer();
 
 			const uint32 MaxSectionCount = GetHairStrandsMaxSectionCount();
@@ -1212,9 +1211,9 @@ namespace GroomBinding_RootProjection
 
 					if (bHasTransferredPosition)
 					{
-						T.P0 = InTransferredPositions[LODIt][T.I0];
-						T.P1 = InTransferredPositions[LODIt][T.I1];
-						T.P2 = InTransferredPositions[LODIt][T.I2];
+						T.P0 = InTransferredPositions[MeshLODIt][T.I0];
+						T.P1 = InTransferredPositions[MeshLODIt][T.I1];
+						T.P2 = InTransferredPositions[MeshLODIt][T.I2];
 					}
 					else
 					{
@@ -1289,9 +1288,9 @@ namespace GroomBinding_RootProjection
 
 					if (bHasTransferredPosition)
 					{
-						T.P0 = InTransferredPositions[LODIt][T.I0];
-						T.P1 = InTransferredPositions[LODIt][T.I1];
-						T.P2 = InTransferredPositions[LODIt][T.I2];
+						T.P0 = InTransferredPositions[MeshLODIt][T.I0];
+						T.P1 = InTransferredPositions[MeshLODIt][T.I1];
+						T.P2 = InTransferredPositions[MeshLODIt][T.I2];
 					}
 					else
 					{
@@ -1314,8 +1313,8 @@ namespace GroomBinding_RootProjection
 				return false;
 			}
 
-			OutRootData.MeshProjectionLODs[LODIt].RootBarycentricBuffer.SetNum(CurveCount);
-			OutRootData.MeshProjectionLODs[LODIt].RootToUniqueTriangleIndexBuffer.SetNum(CurveCount);
+			OutRootData[MeshLODIt].RootBarycentricBuffer.SetNum(CurveCount);
+			OutRootData[MeshLODIt].RootToUniqueTriangleIndexBuffer.SetNum(CurveCount);
 
 			// 2.3. Compute the closest triangle for each root
 			//InMeshRenderData->LODRenderData[LODIt].GetNumVertices();
@@ -1330,7 +1329,7 @@ namespace GroomBinding_RootProjection
 			TAtomic<uint32> bIsValid(1);
 			ParallelFor(CurveCount,
 				[
-					LODIt,
+					MeshLODIt,
 					&InStrandsData,
 					&Grid,
 					&RootTriangleIndexBuffer,
@@ -1377,7 +1376,7 @@ namespace GroomBinding_RootProjection
 				// Record closest triangle and the root's barycentrics
 				const uint32 EncodedBarycentrics = FHairStrandsRootUtils::PackBarycentrics(FVector2f(ClosestBarycentrics));	// LWC_TODO: Precision loss
 				const uint32 EncodedTriangleIndex = FHairStrandsRootUtils::PackTriangleIndex(ClosestTriangle.TriangleIndex, ClosestTriangle.SectionIndex);
-				OutRootData.MeshProjectionLODs[LODIt].RootBarycentricBuffer[CurveIndex] = EncodedBarycentrics;
+				OutRootData[MeshLODIt].RootBarycentricBuffer[CurveIndex] = EncodedBarycentrics;
 
 				RootTriangleIndexBuffer[CurveIndex] = EncodedTriangleIndex;
 				RestRootTrianglePositionBuffer[CurveIndex * 3 + 0] = FVector4f((FVector3f)ClosestTriangle.P0, FHairStrandsRootUtils::PackUVsToFloat(FVector2f(ClosestTriangle.UV0)));	// LWC_TODO: Precision loss
@@ -1424,36 +1423,36 @@ namespace GroomBinding_RootProjection
 
 			// Build final unique triangle list and the root-to-unique-triangle mapping
 			const uint32 UniqueTriangleCount = UniqueTriangleToRootList.Num();
-			OutRootData.MeshProjectionLODs[LODIt].UniqueTriangleIndexBuffer.Reserve(UniqueTriangleCount );
-			OutRootData.MeshProjectionLODs[LODIt].RestUniqueTrianglePositionBuffer.Reserve(UniqueTriangleCount * 3);
+			OutRootData[MeshLODIt].UniqueTriangleIndexBuffer.Reserve(UniqueTriangleCount );
+			OutRootData[MeshLODIt].RestUniqueTrianglePositionBuffer.Reserve(UniqueTriangleCount * 3);
 			for (uint32 EncodedTriangleId : UniqueTriangleToRootList)
 			{
 				auto It = UniqueTriangleToRootMap.Find(EncodedTriangleId);
 				check(It);
 
-				OutRootData.MeshProjectionLODs[LODIt].UniqueTriangleIndexBuffer.Add(EncodedTriangleId);
+				OutRootData[MeshLODIt].UniqueTriangleIndexBuffer.Add(EncodedTriangleId);
 
 				const uint32 FirstCurveIndex = (*It)[0];
-				OutRootData.MeshProjectionLODs[LODIt].RestUniqueTrianglePositionBuffer.Add(RestRootTrianglePositionBuffer[FirstCurveIndex * 3 + 0]);
-				OutRootData.MeshProjectionLODs[LODIt].RestUniqueTrianglePositionBuffer.Add(RestRootTrianglePositionBuffer[FirstCurveIndex * 3 + 1]);
-				OutRootData.MeshProjectionLODs[LODIt].RestUniqueTrianglePositionBuffer.Add(RestRootTrianglePositionBuffer[FirstCurveIndex * 3 + 2]);
+				OutRootData[MeshLODIt].RestUniqueTrianglePositionBuffer.Add(RestRootTrianglePositionBuffer[FirstCurveIndex * 3 + 0]);
+				OutRootData[MeshLODIt].RestUniqueTrianglePositionBuffer.Add(RestRootTrianglePositionBuffer[FirstCurveIndex * 3 + 1]);
+				OutRootData[MeshLODIt].RestUniqueTrianglePositionBuffer.Add(RestRootTrianglePositionBuffer[FirstCurveIndex * 3 + 2]);
 
 				// Write for each root, the index of the triangle
-				const uint32 UniqueTriangleIndex = OutRootData.MeshProjectionLODs[LODIt].UniqueTriangleIndexBuffer.Num()-1;
+				const uint32 UniqueTriangleIndex = OutRootData[MeshLODIt].UniqueTriangleIndexBuffer.Num()-1;
 				for (uint32 CurveIndex : *It)
 				{
-					OutRootData.MeshProjectionLODs[LODIt].RootToUniqueTriangleIndexBuffer[CurveIndex] = UniqueTriangleIndex;
+					OutRootData[MeshLODIt].RootToUniqueTriangleIndexBuffer[CurveIndex] = UniqueTriangleIndex;
 				}
 			}
 
 			// Sanity check
-			check(OutRootData.MeshProjectionLODs[LODIt].RootToUniqueTriangleIndexBuffer.Num() == CurveCount);
-			check(OutRootData.MeshProjectionLODs[LODIt].RestUniqueTrianglePositionBuffer.Num() == UniqueTriangleCount * 3);
-			check(OutRootData.MeshProjectionLODs[LODIt].UniqueTriangleIndexBuffer.Num() == UniqueTriangleCount);
+			check(OutRootData[MeshLODIt].RootToUniqueTriangleIndexBuffer.Num() == CurveCount);
+			check(OutRootData[MeshLODIt].RestUniqueTrianglePositionBuffer.Num() == UniqueTriangleCount * 3);
+			check(OutRootData[MeshLODIt].UniqueTriangleIndexBuffer.Num() == UniqueTriangleCount);
 
 			// Update the root mesh projection data with unique valid mesh section IDs, based on the projection data
-			OutRootData.MeshProjectionLODs[LODIt].UniqueSectionIds = UniqueSectionId;
-			OutRootData.MeshProjectionLODs[LODIt].MeshSectionCount = SectionCount;
+			OutRootData[MeshLODIt].UniqueSectionIds = UniqueSectionId;
+			OutRootData[MeshLODIt].MeshSectionCount = SectionCount;
 		}
 
 		return true;
@@ -1920,22 +1919,21 @@ namespace GroomBinding_Transfer
 }
 // namespace GroomBinding_Transfer
 
-static void InitHairStrandsRootData(FHairStrandsRootData& Out, const FHairStrandsDatas* HairStrandsDatas, uint32 LODCount, uint32 NumSamples)
+static void InitHairStrandsRootData(TArray<FHairStrandsRootData>& Out, const FHairStrandsDatas* HairStrandsDatas, uint32 LODCount, uint32 NumSamples)
 {
 	check(HairStrandsDatas);
 
-	Out.RootCount = HairStrandsDatas->GetNumCurves();
-	Out.PointCount = HairStrandsDatas->GetNumPoints();
-
-	Out.MeshProjectionLODs.SetNum(LODCount);
+	Out.SetNum(LODCount);
 	uint32 LODIndex = 0;
-	for (FHairStrandsRootData::FMeshProjectionLOD& MeshProjectionLOD : Out.MeshProjectionLODs)
+	for (FHairStrandsRootData& OutLOD : Out)
 	{
-		MeshProjectionLOD.SampleCount = NumSamples;
-		MeshProjectionLOD.LODIndex = LODIndex++;
-		MeshProjectionLOD.MeshInterpolationWeightsBuffer.Empty();
-		MeshProjectionLOD.MeshSampleIndicesBuffer.Empty();
-		MeshProjectionLOD.RestSamplePositionsBuffer.Empty();
+		OutLOD.RootCount = HairStrandsDatas->GetNumCurves();
+		OutLOD.PointCount = HairStrandsDatas->GetNumPoints();
+		OutLOD.SampleCount = NumSamples;
+		OutLOD.MeshInterpolationWeightsBuffer.Empty();
+		OutLOD.MeshSampleIndicesBuffer.Empty();
+		OutLOD.RestSamplePositionsBuffer.Empty();
+		OutLOD.LODIndex = LODIndex++;
 	}
 }
 
@@ -1988,8 +1986,6 @@ static void BuildRootBulkData(
 	FHairStrandsRootBulkData& Out,
 	const FHairStrandsRootData& In)
 {
-	const uint32 MeshLODCount = In.MeshProjectionLODs.Num();
-
 	// Header
 	Out.Header.RootCount  = In.RootCount;
 	Out.Header.PointCount = In.PointCount;
@@ -2003,63 +1999,59 @@ static void BuildRootBulkData(
 	Out.Header.Strides.MeshSampleIndicesAndSectionsBufferStride = FHairStrandsRBFSampleIndexFormat::SizeInByte;
 	Out.Header.Strides.RestSamplePositionsBufferStride 			= FHairStrandsMeshTrianglePositionFormat::SizeInByte;
 
-	Out.Header.LODs.SetNum(MeshLODCount);
-	for (uint32 MeshLODIt = 0; MeshLODIt < MeshLODCount; ++MeshLODIt)
 	{
 		const bool bHasValidSamples =
-			In.MeshProjectionLODs[MeshLODIt].MeshInterpolationWeightsBuffer.Num() > 0 &&
-			In.MeshProjectionLODs[MeshLODIt].MeshSampleIndicesBuffer.Num() > 0 &&
-			In.MeshProjectionLODs[MeshLODIt].RestSamplePositionsBuffer.Num() > 0;
-
-		Out.Header.LODs[MeshLODIt].LODIndex 			= In.MeshProjectionLODs[MeshLODIt].LODIndex;
-		Out.Header.LODs[MeshLODIt].SampleCount 			= bHasValidSamples ? In.MeshProjectionLODs[MeshLODIt].SampleCount : 0u;
-		Out.Header.LODs[MeshLODIt].UniqueTriangleCount 	= In.MeshProjectionLODs[MeshLODIt].UniqueTriangleIndexBuffer.Num();
-		Out.Header.LODs[MeshLODIt].UniqueSectionIndices = In.MeshProjectionLODs[MeshLODIt].UniqueSectionIds;
-		Out.Header.LODs[MeshLODIt].MeshSectionCount		= In.MeshProjectionLODs[MeshLODIt].MeshSectionCount;
+			In.MeshInterpolationWeightsBuffer.Num() > 0 &&
+			In.MeshSampleIndicesBuffer.Num() > 0 &&
+			In.RestSamplePositionsBuffer.Num() > 0;
+
+		Out.Header.LODIndex 			= In.LODIndex;
+		Out.Header.SampleCount 			= bHasValidSamples ? In.SampleCount : 0u;
+		Out.Header.UniqueTriangleCount 	= In.UniqueTriangleIndexBuffer.Num();
+		Out.Header.UniqueSectionIndices	= In.UniqueSectionIds;
+		Out.Header.MeshSectionCount		= In.MeshSectionCount;
 	}
 
 	// Data
-	Out.Data.LODs.SetNum(MeshLODCount);
-	for (uint32 MeshLODIt = 0; MeshLODIt < MeshLODCount; ++MeshLODIt)
 	{
 		const bool bHasValidSamples =
-			In.MeshProjectionLODs[MeshLODIt].MeshInterpolationWeightsBuffer.Num() > 0 &&
-			In.MeshProjectionLODs[MeshLODIt].MeshSampleIndicesBuffer.Num() > 0 &&
-			In.MeshProjectionLODs[MeshLODIt].RestSamplePositionsBuffer.Num() > 0;
+			In.MeshInterpolationWeightsBuffer.Num() > 0 &&
+			In.MeshSampleIndicesBuffer.Num() > 0 &&
+			In.RestSamplePositionsBuffer.Num() > 0;
 
-		CopyToBulkData<FHairStrandsUniqueTriangleIndexFormat>(Out.Data.LODs[MeshLODIt].UniqueTriangleIndexBuffer, In.MeshProjectionLODs[MeshLODIt].UniqueTriangleIndexBuffer);
-		CopyToBulkData<FHairStrandsRootBarycentricFormat>(Out.Data.LODs[MeshLODIt].RootBarycentricBuffer, In.MeshProjectionLODs[MeshLODIt].RootBarycentricBuffer);
-		CopyToBulkData<FHairStrandsRootToUniqueTriangleIndexFormat>(Out.Data.LODs[MeshLODIt].RootToUniqueTriangleIndexBuffer, In.MeshProjectionLODs[MeshLODIt].RootToUniqueTriangleIndexBuffer);
-		CopyToBulkData<FHairStrandsMeshTrianglePositionFormat>(Out.Data.LODs[MeshLODIt].RestUniqueTrianglePositionBuffer, In.MeshProjectionLODs[MeshLODIt].RestUniqueTrianglePositionBuffer);
+		CopyToBulkData<FHairStrandsUniqueTriangleIndexFormat>(Out.Data.UniqueTriangleIndexBuffer, In.UniqueTriangleIndexBuffer);
+		CopyToBulkData<FHairStrandsRootBarycentricFormat>(Out.Data.RootBarycentricBuffer, In.RootBarycentricBuffer);
+		CopyToBulkData<FHairStrandsRootToUniqueTriangleIndexFormat>(Out.Data.RootToUniqueTriangleIndexBuffer, In.RootToUniqueTriangleIndexBuffer);
+		CopyToBulkData<FHairStrandsMeshTrianglePositionFormat>(Out.Data.RestUniqueTrianglePositionBuffer, In.RestUniqueTrianglePositionBuffer);
 
 		if (bHasValidSamples)
 		{
-			check(In.MeshProjectionLODs[MeshLODIt].MeshSampleIndicesBuffer.Num() == In.MeshProjectionLODs[MeshLODIt].MeshSampleSectionsBuffer.Num());
+			check(In.MeshSampleIndicesBuffer.Num() == In.MeshSampleSectionsBuffer.Num());
 
-			const uint32 SampleCount = In.MeshProjectionLODs[MeshLODIt].MeshSampleIndicesBuffer.Num();
+			const uint32 SampleCount = In.MeshSampleIndicesBuffer.Num();
 			TArray<uint32> MeshSampleIndicesAndSectionBuffer;
 			MeshSampleIndicesAndSectionBuffer.SetNum(SampleCount);
 			for (uint32 SampleIt=0;SampleIt<SampleCount;++SampleIt)
 			{
-				const uint32 Index = In.MeshProjectionLODs[MeshLODIt].MeshSampleIndicesBuffer[SampleIt];
-				const uint32 SectionIndex = In.MeshProjectionLODs[MeshLODIt].MeshSampleSectionsBuffer[SampleIt];
+				const uint32 Index = In.MeshSampleIndicesBuffer[SampleIt];
+				const uint32 SectionIndex = In.MeshSampleSectionsBuffer[SampleIt];
 				MeshSampleIndicesAndSectionBuffer[SampleIt] = FHairStrandsRootUtils::PackTriangleIndex(Index, SectionIndex);
 
 				// Update the unique section indices with section containing RBF samples
 				// This allows faster update at runtime when not using skin cache
 				// This is done only for guides, which is the only root data containing RBF data
-				Out.Header.LODs[MeshLODIt].UniqueSectionIndices.AddUnique(SectionIndex);
+				Out.Header.UniqueSectionIndices.AddUnique(SectionIndex);
 			}
 
-			CopyToBulkData<FHairStrandsWeightFormat>(Out.Data.LODs[MeshLODIt].MeshInterpolationWeightsBuffer, In.MeshProjectionLODs[MeshLODIt].MeshInterpolationWeightsBuffer);
-			CopyToBulkData<FHairStrandsRBFSampleIndexFormat>(Out.Data.LODs[MeshLODIt].MeshSampleIndicesAndSectionsBuffer, MeshSampleIndicesAndSectionBuffer);
-			CopyToBulkData<FHairStrandsMeshTrianglePositionFormat>(Out.Data.LODs[MeshLODIt].RestSamplePositionsBuffer, In.MeshProjectionLODs[MeshLODIt].RestSamplePositionsBuffer);
+			CopyToBulkData<FHairStrandsWeightFormat>(Out.Data.MeshInterpolationWeightsBuffer, In.MeshInterpolationWeightsBuffer);
+			CopyToBulkData<FHairStrandsRBFSampleIndexFormat>(Out.Data.MeshSampleIndicesAndSectionsBuffer, MeshSampleIndicesAndSectionBuffer);
+			CopyToBulkData<FHairStrandsMeshTrianglePositionFormat>(Out.Data.RestSamplePositionsBuffer, In.RestSamplePositionsBuffer);
 		}
 		else
 		{
-			Out.Data.LODs[MeshLODIt].MeshInterpolationWeightsBuffer.RemoveBulkData();
-			Out.Data.LODs[MeshLODIt].MeshSampleIndicesAndSectionsBuffer.RemoveBulkData();
-			Out.Data.LODs[MeshLODIt].RestSamplePositionsBuffer.RemoveBulkData();
+			Out.Data.MeshInterpolationWeightsBuffer.RemoveBulkData();
+			Out.Data.MeshSampleIndicesAndSectionsBuffer.RemoveBulkData();
+			Out.Data.RestSamplePositionsBuffer.RemoveBulkData();
 		}
 	}
 }
@@ -2073,49 +2065,45 @@ static void BuildRootData(
 
 	Out.RootCount = In.Header.RootCount;
 	Out.PointCount = In.Header.PointCount;
-
-	const uint32 MeshLODCount = In.Header.LODs.Num();
-	Out.MeshProjectionLODs.SetNum(MeshLODCount);
-	for (uint32 MeshLODIt = 0; MeshLODIt < MeshLODCount; ++MeshLODIt)
 	{
-		const bool bHasValidSamples = In.Header.LODs[MeshLODIt].SampleCount > 0;
+		const bool bHasValidSamples = In.Header.SampleCount > 0;
 
-		Out.MeshProjectionLODs[MeshLODIt].LODIndex = In.Header.LODs[MeshLODIt].LODIndex;
-		Out.MeshProjectionLODs[MeshLODIt].SampleCount = bHasValidSamples ? In.Header.LODs[MeshLODIt].SampleCount : 0u;
-		Out.MeshProjectionLODs[MeshLODIt].UniqueSectionIds = In.Header.LODs[MeshLODIt].UniqueSectionIndices;
-		Out.MeshProjectionLODs[MeshLODIt].MeshSectionCount = In.Header.LODs[MeshLODIt].MeshSectionCount;
+		Out.LODIndex = In.Header.LODIndex;
+		Out.SampleCount = bHasValidSamples ? In.Header.SampleCount : 0u;
+		Out.UniqueSectionIds = In.Header.UniqueSectionIndices;
+		Out.MeshSectionCount = In.Header.MeshSectionCount;
 
-		CopyFromBulkData<FHairStrandsUniqueTriangleIndexFormat>(Out.MeshProjectionLODs[MeshLODIt].UniqueTriangleIndexBuffer, In.Data.LODs[MeshLODIt].UniqueTriangleIndexBuffer);
-		CopyFromBulkData<FHairStrandsRootToUniqueTriangleIndexFormat>(Out.MeshProjectionLODs[MeshLODIt].RootToUniqueTriangleIndexBuffer, In.Data.LODs[MeshLODIt].RootToUniqueTriangleIndexBuffer);
-		CopyFromBulkData<FHairStrandsRootBarycentricFormat>(Out.MeshProjectionLODs[MeshLODIt].RootBarycentricBuffer, In.Data.LODs[MeshLODIt].RootBarycentricBuffer);
-		CopyFromBulkData<FHairStrandsMeshTrianglePositionFormat>(Out.MeshProjectionLODs[MeshLODIt].RestUniqueTrianglePositionBuffer, In.Data.LODs[MeshLODIt].RestUniqueTrianglePositionBuffer);
+		CopyFromBulkData<FHairStrandsUniqueTriangleIndexFormat>(Out.UniqueTriangleIndexBuffer, In.Data.UniqueTriangleIndexBuffer);
+		CopyFromBulkData<FHairStrandsRootToUniqueTriangleIndexFormat>(Out.RootToUniqueTriangleIndexBuffer, In.Data.RootToUniqueTriangleIndexBuffer);
+		CopyFromBulkData<FHairStrandsRootBarycentricFormat>(Out.RootBarycentricBuffer, In.Data.RootBarycentricBuffer);
+		CopyFromBulkData<FHairStrandsMeshTrianglePositionFormat>(Out.RestUniqueTrianglePositionBuffer, In.Data.RestUniqueTrianglePositionBuffer);
 
 		if (bHasValidSamples)
 		{
 			TArray<uint32> MeshSampleIndicesAndSectionBuffer;
 
-			CopyFromBulkData<FHairStrandsWeightFormat>(Out.MeshProjectionLODs[MeshLODIt].MeshInterpolationWeightsBuffer, In.Data.LODs[MeshLODIt].MeshInterpolationWeightsBuffer);
-			CopyFromBulkData<FHairStrandsRBFSampleIndexFormat>(MeshSampleIndicesAndSectionBuffer, In.Data.LODs[MeshLODIt].MeshSampleIndicesAndSectionsBuffer);
-			CopyFromBulkData<FHairStrandsMeshTrianglePositionFormat>(Out.MeshProjectionLODs[MeshLODIt].RestSamplePositionsBuffer, In.Data.LODs[MeshLODIt].RestSamplePositionsBuffer);
+			CopyFromBulkData<FHairStrandsWeightFormat>(Out.MeshInterpolationWeightsBuffer, In.Data.MeshInterpolationWeightsBuffer);
+			CopyFromBulkData<FHairStrandsRBFSampleIndexFormat>(MeshSampleIndicesAndSectionBuffer, In.Data.MeshSampleIndicesAndSectionsBuffer);
+			CopyFromBulkData<FHairStrandsMeshTrianglePositionFormat>(Out.RestSamplePositionsBuffer, In.Data.RestSamplePositionsBuffer);
 
 			// Split indices and sections
 			const uint32 SampleCount = MeshSampleIndicesAndSectionBuffer.Num();
-			Out.MeshProjectionLODs[MeshLODIt].MeshSampleIndicesBuffer.SetNum(SampleCount);
-			Out.MeshProjectionLODs[MeshLODIt].MeshSampleSectionsBuffer.SetNum(SampleCount);
+			Out.MeshSampleIndicesBuffer.SetNum(SampleCount);
+			Out.MeshSampleSectionsBuffer.SetNum(SampleCount);
 			for (uint32 SampleIt=0;SampleIt<SampleCount;++SampleIt)
 			{
 				FHairStrandsRootUtils::UnpackTriangleIndex(
 					MeshSampleIndicesAndSectionBuffer[SampleIt], 
-					Out.MeshProjectionLODs[MeshLODIt].MeshSampleIndicesBuffer[SampleIt],
-					Out.MeshProjectionLODs[MeshLODIt].MeshSampleSectionsBuffer[SampleIt]);
+					Out.MeshSampleIndicesBuffer[SampleIt],
+					Out.MeshSampleSectionsBuffer[SampleIt]);
 			}
 		}
 		else
 		{
-			Out.MeshProjectionLODs[MeshLODIt].MeshInterpolationWeightsBuffer.Empty();
-			Out.MeshProjectionLODs[MeshLODIt].MeshSampleIndicesBuffer.Empty();
-			Out.MeshProjectionLODs[MeshLODIt].MeshSampleSectionsBuffer.Empty();
-			Out.MeshProjectionLODs[MeshLODIt].RestSamplePositionsBuffer.Empty();
+			Out.MeshInterpolationWeightsBuffer.Empty();
+			Out.MeshSampleIndicesBuffer.Empty();
+			Out.MeshSampleSectionsBuffer.Empty();
+			Out.RestSamplePositionsBuffer.Empty();
 		}
 	}
 }
@@ -2126,17 +2114,36 @@ static void BuildRootBulkData(
 	const FHairRootGroupData& In)
 {
 	// Guides
-	BuildRootBulkData(Out.SimRootBulkData, In.SimRootData);
+	{
+		const uint32 MeshLODCount = In.SimRootDatas.Num();
+		Out.SimRootBulkDatas.SetNum(MeshLODCount);
+		for (uint32 MeshLODIt = 0; MeshLODIt < MeshLODCount; ++MeshLODIt)
+		{
+			BuildRootBulkData(Out.SimRootBulkDatas[MeshLODIt], In.SimRootDatas[MeshLODIt]);
+		}
+	}
 
 	// Strands
-	BuildRootBulkData(Out.RenRootBulkData, In.RenRootData);
+	{
+		const uint32 MeshLODCount = In.RenRootDatas.Num();
+		Out.RenRootBulkDatas.SetNum(MeshLODCount);
+		for (uint32 MeshLODIt = 0; MeshLODIt < MeshLODCount; ++MeshLODIt)
+		{
+			BuildRootBulkData(Out.RenRootBulkDatas[MeshLODIt], In.RenRootDatas[MeshLODIt]);
+		}
+	}
 
 	// Cards
-	const uint32 LODCount = In.CardsRootData.Num();
-	Out.CardsRootBulkData.SetNum(LODCount);
+	const uint32 LODCount = In.CardsRootDatas.Num();
+	Out.CardsRootBulkDatas.SetNum(LODCount);
 	for (uint32 LODIt = 0; LODIt < LODCount; ++LODIt)
 	{
-		BuildRootBulkData(Out.CardsRootBulkData[LODIt], In.CardsRootData[LODIt]);
+		const uint32 MeshLODCount = In.CardsRootDatas[LODIt].Num();
+		Out.CardsRootBulkDatas[LODIt].SetNum(MeshLODCount);
+		for (uint32 MeshLODIt = 0; MeshLODIt < MeshLODCount; ++MeshLODIt)
+		{
+			BuildRootBulkData(Out.CardsRootBulkDatas[LODIt][MeshLODIt], In.CardsRootDatas[LODIt][MeshLODIt]);
+		}
 	}
 }
 
@@ -2255,17 +2262,17 @@ static bool InternalBuildBinding_CPU(const FGroomBindingBuilder::FInput& In, uin
 		const FHairGroupPlatformData& GroupData = In.GroomAsset->GetHairGroupsPlatformData()[InGroupIndex];
 		{
 			// Guides
-			InitHairStrandsRootData(OutData.SimRootData, &GuidesData, MeshLODCount, In.NumInterpolationPoints);
+			InitHairStrandsRootData(OutData.SimRootDatas, &GuidesData, MeshLODCount, In.NumInterpolationPoints);
 
 			// Strands
 			if (bNeedStrandsRoot)
 			{
-				InitHairStrandsRootData(OutData.RenRootData, &StrandsData, MeshLODCount, In.NumInterpolationPoints);
+				InitHairStrandsRootData(OutData.RenRootDatas, &StrandsData, MeshLODCount, In.NumInterpolationPoints);
 			}
 
 			// Cards
 			const uint32 CardsLODCount = GroupData.Cards.LODs.Num();
-			OutData.CardsRootData.SetNum(GroupData.Cards.LODs.Num());
+			OutData.CardsRootDatas.SetNum(GroupData.Cards.LODs.Num());
 			for (uint32 CardsLODIt = 0; CardsLODIt < CardsLODCount; ++CardsLODIt)
 			{
 				if (GroupData.Cards.IsValid(CardsLODIt))
@@ -2274,7 +2281,7 @@ static bool InternalBuildBinding_CPU(const FGroomBindingBuilder::FInput& In, uin
 					const bool bIsValid = In.GroomAsset->GetHairCardsGuidesDatas(InGroupIndex, CardsLODIt, LODGuidesData);
 					if (bIsValid)
 					{
-						InitHairStrandsRootData(OutData.CardsRootData[CardsLODIt], &LODGuidesData, MeshLODCount, In.NumInterpolationPoints);
+						InitHairStrandsRootData(OutData.CardsRootDatas[CardsLODIt], &LODGuidesData, MeshLODCount, In.NumInterpolationPoints);
 					}
 				}
 			}
@@ -2291,7 +2298,7 @@ static bool InternalBuildBinding_CPU(const FGroomBindingBuilder::FInput& In, uin
 			// Strands
 			WorkItemCount += bNeedStrandsRoot ? 1 : 0;
 			// Cards
-			WorkItemCount += OutData.CardsRootData.Num();
+			WorkItemCount += OutData.CardsRootDatas.Num();
 		}
 	
 		uint32 WorkItemIndex = 0;
@@ -2321,7 +2328,7 @@ static bool InternalBuildBinding_CPU(const FGroomBindingBuilder::FInput& In, uin
 					GuidesData,
 					TargetMeshData.Get(),
 					TransferredPositions,
-					OutData.SimRootData))
+					OutData.SimRootDatas))
 				{
 					UE_LOG(LogHairStrands, Error, TEXT("[Groom] Binding asset could not be built. Some guide roots are not close enough to the target mesh to be projected onto it."));
 					return false; 
@@ -2336,7 +2343,7 @@ static bool InternalBuildBinding_CPU(const FGroomBindingBuilder::FInput& In, uin
 					StrandsData,
 					TargetMeshData.Get(),
 					TransferredPositions,
-					OutData.RenRootData))
+					OutData.RenRootDatas))
 				{
 					UE_LOG(LogHairStrands, Error, TEXT("[Groom] Binding asset could not be built. Some strand roots are not close enough to the target mesh to be projected onto it."));
 					return false;
@@ -2345,7 +2352,7 @@ static bool InternalBuildBinding_CPU(const FGroomBindingBuilder::FInput& In, uin
 			}
 
 			// Cards
-			const uint32 CardsLODCount = OutData.CardsRootData.Num();
+			const uint32 CardsLODCount = OutData.CardsRootDatas.Num();
 			for (uint32 CardsLODIt = 0; CardsLODIt < CardsLODCount; ++CardsLODIt)
 			{
 				if (In.GroomAsset->GetHairGroupsPlatformData()[InGroupIndex].Cards.IsValid(CardsLODIt))
@@ -2358,7 +2365,7 @@ static bool InternalBuildBinding_CPU(const FGroomBindingBuilder::FInput& In, uin
 							LODGuidesData,
 							TargetMeshData.Get(),
 							TransferredPositions,
-							OutData.CardsRootData[CardsLODIt]))
+							OutData.CardsRootDatas[CardsLODIt]))
 						{
 							UE_LOG(LogHairStrands, Error, TEXT("[Groom] Binding asset could not be built. Some cards guide roots are not close enough to the target mesh to be projected onto it."));
 							return false; 
diff --git a/Engine/Plugins/Runtime/HairStrands/Source/HairStrandsCore/Private/GroomDebug.cpp b/Engine/Plugins/Runtime/HairStrands/Source/HairStrandsCore/Private/GroomDebug.cpp
--- a/Engine/Plugins/Runtime/HairStrands/Source/HairStrandsCore/Private/GroomDebug.cpp
+++ b/Engine/Plugins/Runtime/HairStrands/Source/HairStrandsCore/Private/GroomDebug.cpp
@@ -158,13 +158,13 @@ static void AddDebugProjectionHairPass(
 	// Force ShaderPrint on.
 	ShaderPrint::SetEnabled(true);
 
-	if (MeshLODIndex < 0 || MeshLODIndex >= RestRootResources->LODs.Num() || MeshLODIndex >= DeformedRootResources->LODs.Num())
+	if (!RestRootResources->IsValid(MeshLODIndex) || !DeformedRootResources->IsValid(MeshLODIndex))
 	{
 		return;
 	}
 	
 	const EPrimitiveType PrimitiveType = GeometryType == EDebugProjectionHairType::HairFrame ? PT_LineList : GeometryType == EDebugProjectionHairType::HairTriangle ? PT_TriangleList : PT_LineList;
-	const uint32 RootCount = EDebugProjectionHairType::HairSamples == GeometryType ? 3 * RestRootResources->LODs[MeshLODIndex].SampleCount : RestRootResources->GetRootCount();
+	const uint32 RootCount = EDebugProjectionHairType::HairSamples == GeometryType ? 3 * RestRootResources->GetLOD(MeshLODIndex)->SampleCount : RestRootResources->GetRootCount();
 	const uint32 PrimitiveCount = RootCount;
 
 	if (PrimitiveCount == 0)
@@ -181,19 +181,19 @@ static void AddDebugProjectionHairPass(
 	ShaderPrint::RequestSpaceForTriangles(PrimitiveCount * 3);
 
 	if (EDebugProjectionHairType::HairFrame == GeometryType &&
-		!RestRootResources->LODs[MeshLODIndex].RootBarycentricBuffer.Buffer)
+		!RestRootResources->GetLOD(MeshLODIndex)->RootBarycentricBuffer.Buffer)
 	{
 		return;
 	}
 
 	if (EDebugProjectionHairType::HairSamples == GeometryType &&
-		!RestRootResources->LODs[MeshLODIndex].RestSamplePositionsBuffer.Buffer)
+		!RestRootResources->GetLOD(MeshLODIndex)->RestSamplePositionsBuffer.Buffer)
 	{
 			return;
 	}
 
-	const FHairStrandsRestRootResource::FLOD& RestLODDatas = RestRootResources->LODs[MeshLODIndex];
-	const FHairStrandsDeformedRootResource::FLOD& DeformedLODDatas = DeformedRootResources->LODs[MeshLODIndex];
+	const FHairStrandsLODRestRootResource& RestLODDatas = *RestRootResources->GetLOD(MeshLODIndex);
+	const FHairStrandsLODDeformedRootResource& DeformedLODDatas = *DeformedRootResources->GetLOD(MeshLODIndex);
 
 	if (!RestLODDatas.RestUniqueTrianglePositionBuffer.Buffer || !DeformedLODDatas.DeformedUniqueTrianglePositionBuffer[0].Buffer)
 	{
@@ -236,10 +236,10 @@ static void AddDebugProjectionHairPass(
 	Parameters->RootToUniqueTriangleIndexBuffer = RegisterAsSRV(GraphBuilder, RestLODDatas.RootToUniqueTriangleIndexBuffer);
 
 	Parameters->RestPositionBuffer = RegisterAsSRV(GraphBuilder, RestLODDatas.RestUniqueTrianglePositionBuffer);
-	Parameters->DeformedPositionBuffer = RegisterAsSRV(GraphBuilder, DeformedLODDatas.GetDeformedUniqueTrianglePositionBuffer(FHairStrandsDeformedRootResource::FLOD::Current));
+	Parameters->DeformedPositionBuffer = RegisterAsSRV(GraphBuilder, DeformedLODDatas.GetDeformedUniqueTrianglePositionBuffer(FHairStrandsLODDeformedRootResource::Current));
 
 	Parameters->RestSamplePositionsBuffer = RegisterAsSRV(GraphBuilder, RestLODDatas.RestSamplePositionsBuffer);
-	Parameters->DeformedSamplePositionsBuffer = RegisterAsSRV(GraphBuilder, DeformedLODDatas.GetDeformedSamplePositionsBuffer(FHairStrandsDeformedRootResource::FLOD::Current));
+	Parameters->DeformedSamplePositionsBuffer = RegisterAsSRV(GraphBuilder, DeformedLODDatas.GetDeformedSamplePositionsBuffer(FHairStrandsLODDeformedRootResource::Current));
 
 	Parameters->ViewUniformBuffer = ViewUniformBuffer;
 	ShaderPrint::SetParameters(GraphBuilder, *ShaderPrintData, Parameters->ShaderPrintUniformBuffer);
@@ -823,12 +823,12 @@ static void AddHairDebugPrintInstancePass(
 				D.Data0.Z = Instance->Strands.GetData().GetNumCurves(); // Change this later on for having dynamic value
 				D.Data0.W = Instance->Strands.GetData().GetNumPoints(); // Change this later on for having dynamic value
 				const int32 MeshLODIndex = Instance->HairGroupPublicData->MeshLODIndex;
-				if (MeshLODIndex>=0 && Instance->Strands.RestRootResource)
+				if (MeshLODIndex>=0 && Instance->Strands.RestRootResource->IsValid(MeshLODIndex))
 				{
-					D.Data1.X = Instance->Strands.RestRootResource->BulkData.Header.LODs[MeshLODIndex].UniqueSectionIndices.Num();
-					D.Data1.Y = Instance->Strands.RestRootResource->BulkData.Header.LODs[MeshLODIndex].UniqueTriangleCount;
-					D.Data1.Z = Instance->Strands.RestRootResource->BulkData.Header.RootCount;
-					D.Data1.W = Instance->Strands.RestRootResource->BulkData.Header.PointCount;
+					D.Data1.X = Instance->Strands.RestRootResource->GetLOD(MeshLODIndex)->BulkData.Header.UniqueSectionIndices.Num();
+					D.Data1.Y = Instance->Strands.RestRootResource->GetLOD(MeshLODIndex)->BulkData.Header.UniqueTriangleCount;
+					D.Data1.Z = Instance->Strands.RestRootResource->GetLOD(MeshLODIndex)->BulkData.Header.RootCount;
+					D.Data1.W = Instance->Strands.RestRootResource->GetLOD(MeshLODIndex)->BulkData.Header.PointCount;
 				}
 
 				{
@@ -1147,7 +1147,7 @@ static void AddHairDebugPrintMemoryPass(
 					FInfos& D = BindingBuffer.AddDefaulted_GetRef();
 					D.Data0.X = GroupIt;
 					D.Data0.Y = GroupCount;
-					D.Data0.Z = Resource.RenRootResources ? Resource.RenRootResources->MaxAvailableCurveCount : 0;
+					D.Data0.Z = Resource.RenRootResources ? 0u : 0u; // TODO: need a mesh index: we could take also the max amoung all mesh LODs. Resource.RenRootResources->MaxAvailableCurveCount : 0;
 					D.Data0.W = Resource.RenRootResources ? Resource.RenRootResources->GetRootCount() : 0;
 			
 					D.Data1.X = MemoryStats.CPU.Guides;
diff --git a/Engine/Plugins/Runtime/HairStrands/Source/HairStrandsCore/Private/GroomManager.cpp b/Engine/Plugins/Runtime/HairStrands/Source/HairStrandsCore/Private/GroomManager.cpp
--- a/Engine/Plugins/Runtime/HairStrands/Source/HairStrandsCore/Private/GroomManager.cpp
+++ b/Engine/Plugins/Runtime/HairStrands/Source/HairStrandsCore/Private/GroomManager.cpp
@@ -481,7 +481,7 @@ static void GetOrAllocateCachedGeometry(
 	FRDGBuilder& GraphBuilder,
 	FGlobalShaderMap* ShaderMap, 
 	const FSkeletalMeshSceneProxy* Proxy,
-	const FHairStrandsRootBulkData* RootBulkData,
+	const FHairStrandsRestRootResource* RootResource,
 	const bool bOutputTriangleData,
 	FCachedGeometry& Out,
 	FHairGeometryCache& OutHairGeometryCache,
@@ -510,9 +510,9 @@ static void GetOrAllocateCachedGeometry(
 		return;
 	}
 
-	check(RootBulkData);
-	check(RootBulkData->Header.LODs.IsValidIndex(LODIndex));
-	const TArray<uint32>& UniqueSections = RootBulkData->Header.LODs[LODIndex].UniqueSectionIndices;
+	check(RootResource);
+	check(RootResource->IsDataValid(LODIndex));
+	const TArray<uint32>& UniqueSections = RootResource->GetLOD(LODIndex)->BulkData.Header.UniqueSectionIndices;
 
 	// Create deformed position buffer (output)
 	FRDGBufferSRVRef DeformedPositionSRV = nullptr;
@@ -586,7 +586,7 @@ static FCachedGeometry GetCacheGeometryForHair(
 	FGlobalShaderMap* ShaderMap,
 	FSceneInterface* Scene,
 	FHairGroupInstance* Instance, 
-	const FHairStrandsRootBulkData* RootBulkData,
+	const FHairStrandsRestRootResource* RootResource,
 	const bool bOutputTriangleData,
 	const EHairPositionUpdateType PositionUpdateType,
 	FHairGeometryCache& OutHairGeometryCache)
@@ -608,7 +608,7 @@ static FCachedGeometry GetCacheGeometryForHair(
 				else if (IsHairManualSkinCacheEnabled())
 				{
 					uint32 OutTotalSectionCount = 0;
-					GetOrAllocateCachedGeometry(GraphBuilder, ShaderMap, SceneProxy, RootBulkData, bOutputTriangleData, Out, OutHairGeometryCache, OutTotalSectionCount);
+					GetOrAllocateCachedGeometry(GraphBuilder, ShaderMap, SceneProxy, RootResource, bOutputTriangleData, Out, OutHairGeometryCache, OutTotalSectionCount);
 					OutHairGeometryCache.AddDebug(Instance, SceneProxy, Out, PositionUpdateType, FHairGeometryCache::ECacheType::HairCache, OutTotalSectionCount);
 				}
 			}
@@ -660,7 +660,7 @@ static void RunHairBindingSurfaceUpdate(
 
 		// 1. Guides
 		{
-			FHairStrandsRootBulkData* RootBulkData = nullptr;
+			FHairStrandsRestRootResource* RootResource = nullptr;
 			if (Instance->Guides.IsValid())
 			{	
 				const bool bNeedSurfaceUpdate = 
@@ -672,10 +672,10 @@ static void RunHairBindingSurfaceUpdate(
 				{
 					// Extract MeshLODData and MeshLODIndex
 					check(Instance->Guides.RestRootResource);
-					RootBulkData = &Instance->Guides.RestRootResource->BulkData;
+					RootResource = Instance->Guides.RestRootResource;
 				}
 			}
-			TransientResources.SimMeshDatas[Instance->RegisteredIndex] = GetCacheGeometryForHair(GraphBuilder, ShaderMap, Scene, Instance, RootBulkData, RootBulkData != nullptr /*bOutputTriangleData*/, EHairPositionUpdateType::Guides, HairGeometryCache);
+			TransientResources.SimMeshDatas[Instance->RegisteredIndex] = GetCacheGeometryForHair(GraphBuilder, ShaderMap, Scene, Instance, RootResource, RootResource != nullptr /*bOutputTriangleData*/, EHairPositionUpdateType::Guides, HairGeometryCache);
 			MeshLODIndex = TransientResources.SimMeshDatas[Instance->RegisteredIndex].LODIndex;
 		}
 
@@ -686,15 +686,15 @@ static void RunHairBindingSurfaceUpdate(
 			{
 				// Extract MeshLODData and compute MeshLODIndex
 				const bool bNeedOutputTriangleData = Instance->Strands.RestRootResource != nullptr;
-				FHairStrandsRootBulkData* RootBulkData = bNeedOutputTriangleData ? &Instance->Strands.RestRootResource->BulkData : nullptr;
-				TransientResources.RenMeshDatas[Instance->RegisteredIndex] = GetCacheGeometryForHair(GraphBuilder, ShaderMap, Scene, Instance, RootBulkData, RootBulkData != nullptr /*bOutputTriangleData*/, EHairPositionUpdateType::Strands, HairGeometryCache);
+				FHairStrandsRestRootResource* RootResource = bNeedOutputTriangleData ? Instance->Strands.RestRootResource : nullptr;
+				TransientResources.RenMeshDatas[Instance->RegisteredIndex] = GetCacheGeometryForHair(GraphBuilder, ShaderMap, Scene, Instance, RootResource, RootResource != nullptr /*bOutputTriangleData*/, EHairPositionUpdateType::Strands, HairGeometryCache);
 				MeshLODIndex = TransientResources.RenMeshDatas[Instance->RegisteredIndex].LODIndex;
 			}
 			// Cards 
 			// This is only needed for card geometry. Mesh geometry only uses RBF deformation, which are initalized by the guide pass.
 			else if (Instance->GeometryType == EHairGeometryType::Cards && bCardSupported && Instance->Cards.IsValid(HairLODIndex))
 			{
-				FHairStrandsRootBulkData* RootBulkData = nullptr;
+				FHairStrandsRestRootResource* RootResource = nullptr;
 				const bool bNeedSurfaceUpdate = 
 					Instance->BindingType == EHairBindingType::Skinning || 
 					Instance->HairGroupPublicData->IsGlobalInterpolationEnable(HairLODIndex) || 
@@ -703,9 +703,9 @@ static void RunHairBindingSurfaceUpdate(
 				{	
 					// Extract MeshLODData and MeshLODIndex
 					check(Instance->Cards.LODs.IsValidIndex(HairLODIndex));
-					RootBulkData = &Instance->Cards.LODs[HairLODIndex].Guides.RestRootResource->BulkData;
+					RootResource = Instance->Cards.LODs[HairLODIndex].Guides.RestRootResource;
 				}
-				TransientResources.RenMeshDatas[Instance->RegisteredIndex] = GetCacheGeometryForHair(GraphBuilder, ShaderMap, Scene, Instance, RootBulkData, RootBulkData != nullptr /*bOutputTriangleData*/, EHairPositionUpdateType::Cards, HairGeometryCache);
+				TransientResources.RenMeshDatas[Instance->RegisteredIndex] = GetCacheGeometryForHair(GraphBuilder, ShaderMap, Scene, Instance, RootResource, RootResource != nullptr /*bOutputTriangleData*/, EHairPositionUpdateType::Cards, HairGeometryCache);
 				MeshLODIndex = TransientResources.RenMeshDatas[Instance->RegisteredIndex].LODIndex;
 			}
 		}
@@ -876,12 +876,12 @@ static void RunHairStrandsInterpolation_Guide(
 
 			// Add manual transition for the GPU solver as Niagara does not track properly the RDG buffer, and so doesn't issue the correct transitions
 			{
-				GraphBuilder.UseExternalAccessMode(Register(GraphBuilder, InstanceData.Instance->Guides.DeformedRootResource->LODs[InstanceData.MeshLODIndex].GetDeformedUniqueTrianglePositionBuffer(FHairStrandsDeformedRootResource::FLOD::Current), ERDGImportedBufferFlags::CreateSRV).Buffer, ERHIAccess::SRVMask);
+				GraphBuilder.UseExternalAccessMode(Register(GraphBuilder, InstanceData.Instance->Guides.DeformedRootResource->GetLOD(InstanceData.MeshLODIndex)->GetDeformedUniqueTrianglePositionBuffer(FHairStrandsLODDeformedRootResource::Current), ERDGImportedBufferFlags::CreateSRV).Buffer, ERHIAccess::SRVMask);
 
 				if (InstanceData.bGlobalDeformationEnable)
 				{
-					GraphBuilder.UseExternalAccessMode(Register(GraphBuilder, InstanceData.Instance->Guides.DeformedRootResource->LODs[InstanceData.MeshLODIndex].GetDeformedSamplePositionsBuffer(FHairStrandsDeformedRootResource::FLOD::Current), ERDGImportedBufferFlags::CreateSRV).Buffer, ERHIAccess::SRVMask);
-					GraphBuilder.UseExternalAccessMode(Register(GraphBuilder, InstanceData.Instance->Guides.DeformedRootResource->LODs[InstanceData.MeshLODIndex].GetMeshSampleWeightsBuffer(FHairStrandsDeformedRootResource::FLOD::Current), ERDGImportedBufferFlags::CreateSRV).Buffer, ERHIAccess::SRVMask);
+					GraphBuilder.UseExternalAccessMode(Register(GraphBuilder, InstanceData.Instance->Guides.DeformedRootResource->GetLOD(InstanceData.MeshLODIndex)->GetDeformedSamplePositionsBuffer(FHairStrandsLODDeformedRootResource::Current), ERDGImportedBufferFlags::CreateSRV).Buffer, ERHIAccess::SRVMask);
+					GraphBuilder.UseExternalAccessMode(Register(GraphBuilder, InstanceData.Instance->Guides.DeformedRootResource->GetLOD(InstanceData.MeshLODIndex)->GetMeshSampleWeightsBuffer(FHairStrandsLODDeformedRootResource::Current), ERDGImportedBufferFlags::CreateSRV).Buffer, ERHIAccess::SRVMask);
 				}
 			}
 			GraphBuilder.UseExternalAccessMode(Register(GraphBuilder, InstanceData.Instance->Guides.DeformedResource->GetPositionOffsetBuffer(FHairStrandsDeformedResource::EFrameType::Current), ERDGImportedBufferFlags::CreateSRV).Buffer, ERHIAccess::SRVMask);
@@ -1252,7 +1252,7 @@ static void RunHairStrandsInterpolation_Strands(
 					InstanceData.MeshLODData,
 					InstanceData.Instance->Strands.RestRootResource,
 					InstanceData.Instance->Strands.DeformedRootResource);
-				Transitions.Add(RegisterAsSRV(GraphBuilder, InstanceData.Instance->Strands.DeformedRootResource->LODs[InstanceData.MeshLODIndex].GetDeformedUniqueTrianglePositionBuffer(FHairStrandsDeformedRootResource::FLOD::Current)));
+				Transitions.Add(RegisterAsSRV(GraphBuilder, InstanceData.Instance->Strands.DeformedRootResource->GetLOD(InstanceData.MeshLODIndex)->GetDeformedUniqueTrianglePositionBuffer(FHairStrandsLODDeformedRootResource::Current)));
 			}
 		}
 		AddTransitionPass(GraphBuilder, ShaderMap, View->GetShaderPlatform(), Transitions);
@@ -1405,7 +1405,7 @@ static void RunHairStrandsInterpolation_Strands(
 					bValidGuide ? RegisterAsSRV(GraphBuilder, InstanceData.Instance->Guides.DeformedResource->GetBuffer(FHairStrandsDeformedResource::Current)) : nullptr,
 					InstanceData.RDGResources.DeformerPositionSRV,
 					InstanceData.RDGResources.PositionUAV,
-					FHairStrandsDeformedRootResource::FLOD::Current);
+					FHairStrandsLODDeformedRootResource::Current);
 	
 				Transitions.Add(InstanceData.RDGResources.PositionSRV);
 
@@ -1963,7 +1963,7 @@ static void RunHairStrandsInterpolation_Cards(
 					InstanceData.bValidGuide ? RegisterAsSRV(GraphBuilder, InstanceData.Instance->Guides.DeformedResource->GetBuffer(FHairStrandsDeformedResource::Current)) : nullptr,
 					nullptr,
 					Guides_DeformedPositionUAV,
-					FHairStrandsDeformedRootResource::FLOD::Current); // <- this should be optional
+					FHairStrandsLODDeformedRootResource::Current);
 			}
 		}
 	}
@@ -1985,6 +1985,7 @@ static void RunHairStrandsInterpolation_Cards(
 					View->GetFeatureLevel(),
 					ShaderPrintData,
 					InstanceData.Instance,
+					InstanceData.HairLODIndex,
 					InstanceData.MeshLODIndex);
 			}
 		}
@@ -2069,7 +2070,7 @@ static void RunHairStrandsInterpolation_Cards(
 			AddHairMeshesRBFInterpolationPass(
 				GraphBuilder,
 				ShaderMap,
-				InstanceData.Instance->Debug.MeshLODIndex,
+				InstanceData.MeshLODIndex,
 				InstanceData.MeshInstance->RestResource,
 				InstanceData.MeshInstance->DeformedResource,
 				InstanceData.Instance->Guides.RestRootResource,
@@ -2210,14 +2211,14 @@ static void AddCopyHairStrandsPositionPass(
 }
 
 #if RHI_RAYTRACING
-static void AllocateRaytracingResources(FHairGroupInstance* Instance)
+static void AllocateRaytracingResources(FHairGroupInstance* Instance, int32 InMeshLODIndex)
 {
 	if (IsHairRayTracingEnabled() && !Instance->Strands.RenRaytracingResource)
 	{
 		check(Instance->Strands.IsValid());
 
 #if RHI_ENABLE_RESOURCE_INFO
-		FName OwnerName(FString::Printf(TEXT("%s [LOD%d]"), *Instance->Debug.MeshComponentName, Instance->Debug.MeshLODIndex));
+		FName OwnerName(FString::Printf(TEXT("%s [LOD%d]"), *Instance->Debug.MeshComponentName, InMeshLODIndex));
 #else
 		FName OwnerName = NAME_None;
 #endif
@@ -2629,17 +2630,17 @@ static bool SelectValidLOD(
 	{
 		check(Instance->HairGroupPublicData);
 
-		if (Instance->Strands.RestRootResource)			{ Instance->Strands.RestRootResource->Allocate		(GraphBuilder, EHairResourceLoadingType::Async, ResourceStatus, RequestedCurveCount, RequestedPointCount, MeshLODIndex,       false/*bAllowDeallocation*/); }
-		if (Instance->Strands.RestResource)				{ Instance->Strands.RestResource->Allocate			(GraphBuilder, EHairResourceLoadingType::Async, ResourceStatus, RequestedCurveCount, RequestedPointCount, -1/*MeshLODIndex*/, false/*bAllowDeallocation*/); }
-		if (Instance->Strands.ClusterResource)			{ Instance->Strands.ClusterResource->Allocate		(GraphBuilder, EHairResourceLoadingType::Async, ResourceStatus, RequestedCurveCount, RequestedPointCount, -1/*MeshLODIndex*/, false/*bAllowDeallocation*/); }
-		if (Instance->Strands.InterpolationResource)	{ Instance->Strands.InterpolationResource->Allocate	(GraphBuilder, EHairResourceLoadingType::Async, ResourceStatus, RequestedCurveCount, RequestedPointCount, -1/*MeshLODIndex*/, false/*bAllowDeallocation*/); }
+		if (Instance->Strands.RestRootResource)			{ Instance->Strands.RestRootResource->Allocate		(GraphBuilder, EHairResourceLoadingType::Async, ResourceStatus, RequestedCurveCount, RequestedPointCount, false/*bAllowDeallocation*/, MeshLODIndex); }
+		if (Instance->Strands.RestResource)				{ Instance->Strands.RestResource->Allocate			(GraphBuilder, EHairResourceLoadingType::Async, ResourceStatus, RequestedCurveCount, RequestedPointCount, false/*bAllowDeallocation*/); }
+		if (Instance->Strands.ClusterResource)			{ Instance->Strands.ClusterResource->Allocate		(GraphBuilder, EHairResourceLoadingType::Async, ResourceStatus, RequestedCurveCount, RequestedPointCount, false/*bAllowDeallocation*/); }
+		if (Instance->Strands.InterpolationResource)	{ Instance->Strands.InterpolationResource->Allocate	(GraphBuilder, EHairResourceLoadingType::Async, ResourceStatus, RequestedCurveCount, RequestedPointCount, false/*bAllowDeallocation*/); }
 		if (Instance->Strands.CullingResource)			{ Instance->Strands.CullingResource->Allocate		(GraphBuilder, EHairResourceLoadingType::Async, ResourceStatus, RequestedCurveCount, RequestedPointCount); }
 
-		if (Instance->Strands.DeformedRootResource)		{ Instance->Strands.DeformedRootResource->Allocate	(GraphBuilder, EHairResourceLoadingType::Async, ResourceStatus, RequestedCurveCount, RequestedPointCount, MeshLODIndex); }
+		if (Instance->Strands.DeformedRootResource)		{ Instance->Strands.DeformedRootResource->Allocate	(GraphBuilder, EHairResourceLoadingType::Async, ResourceStatus, RequestedCurveCount, RequestedPointCount, false/*bAllowDeallocation*/, MeshLODIndex); }
 		if (Instance->Strands.DeformedResource)			{ Instance->Strands.DeformedResource->Allocate		(GraphBuilder, EHairResourceLoadingType::Async, ResourceStatus, RequestedCurveCount, RequestedPointCount); }
 #if RHI_RAYTRACING
-		if (bHasPathTracingView)						{ AllocateRaytracingResources(Instance); }
-		if (Instance->Strands.RenRaytracingResource)	{ Instance->Strands.RenRaytracingResource->Allocate	(GraphBuilder, EHairResourceLoadingType::Async, ResourceStatus, RequestedCurveCount, RequestedPointCount, -1/*MeshLODIndex*/, Instance->Strands.RenRaytracingResourceOwned/*bAllowDeallocation*/); }
+		if (bHasPathTracingView)						{ AllocateRaytracingResources(Instance, MeshLODIndex); }
+		if (Instance->Strands.RenRaytracingResource)	{ Instance->Strands.RenRaytracingResource->Allocate	(GraphBuilder, EHairResourceLoadingType::Async, ResourceStatus, RequestedCurveCount, RequestedPointCount, Instance->Strands.RenRaytracingResourceOwned/*bAllowDeallocation*/); }
 #endif
 		Instance->Strands.VertexFactory->InitResources(GraphBuilder.RHICmdList);
 
@@ -2850,13 +2851,22 @@ static void RunHairLODSelection(
 			// Groom binding
 			if (BindingType == EHairBindingType::Skinning)
 			{
-				FUniqueGroomBinding* Unique = UniqueGroomBindingAssets.FindByPredicate([Instance](const FUniqueGroomBinding& A) { return A.AssetHash == Instance->Debug.GroomBindingAssetHash && A.GroupIndex == Instance->Debug.GroupIndex; });
+				FUniqueGroomBinding* Unique = UniqueGroomBindingAssets.FindByPredicate([&InstanceData](const FUniqueGroomBinding& A) 
+				{ 
+					return 
+						A.AssetHash == InstanceData.Instance->Debug.GroomBindingAssetHash && 
+						A.GroupIndex == InstanceData.Instance->Debug.GroupIndex && 
+						A.MeshLODIndex == InstanceData.MeshLODIndex; 
+				});
 				if (Unique)
 				{
+					// Sanity check
+					check(Unique->MeshLODIndex == InstanceData.MeshLODIndex);
+					check(Unique->GroupIndex == InstanceData.Instance->Debug.GroupIndex);
+
 					if (InstanceData.HairLOD.ContinuousLODCurveCount > Unique->HairLOD.ContinuousLODCurveCount)
 					{
-						Unique->HairLOD 	= InstanceData.HairLOD;
-						Unique->MeshLODIndex= InstanceData.MeshLODIndex;
+						Unique->HairLOD = InstanceData.HairLOD;
 					}
 				}
 				else
@@ -2876,6 +2886,7 @@ static void RunHairLODSelection(
 	}
 
 	// Emit streaming request for each unique Groom/GroomBinding group
+	bool bAllowToDeallocate = true;
 	for (FUniqueGroom& Unique : UniqueGroomAssets)
 	{
 		const uint32 RequestedCurveCount = Unique.HairLOD.ContinuousLODCurveCount;
@@ -2901,7 +2912,7 @@ static void RunHairLODSelection(
 		ResourceStatus.Status 				= EHairResourceStatus::EStatus::None;
 		ResourceStatus.AvailableCurveCount 	= Unique.RestCurveCount;
 
-		if (Unique.RestRootResource)		{ Unique.RestRootResource->Allocate(GraphBuilder, EHairResourceLoadingType::Async, ResourceStatus, RequestedCurveCount, RequestedPointCount, Unique.MeshLODIndex); }
+		if (Unique.RestRootResource)		{ Unique.RestRootResource->Allocate(GraphBuilder, EHairResourceLoadingType::Async, ResourceStatus, RequestedCurveCount, RequestedPointCount, bAllowToDeallocate, Unique.MeshLODIndex); }
 	}
 
 	// LOD selection & load resources
diff --git a/Engine/Plugins/Runtime/HairStrands/Source/HairStrandsCore/Private/GroomRBFDeformer.cpp b/Engine/Plugins/Runtime/HairStrands/Source/HairStrandsCore/Private/GroomRBFDeformer.cpp
--- a/Engine/Plugins/Runtime/HairStrands/Source/HairStrandsCore/Private/GroomRBFDeformer.cpp
+++ b/Engine/Plugins/Runtime/HairStrands/Source/HairStrandsCore/Private/GroomRBFDeformer.cpp
@@ -183,7 +183,7 @@ void DeformStrands(
 
 // Compute the triangle positions for each curve's roots
 void ExtractUniqueTrianglePositions(
-	const FHairStrandsRootData::FMeshProjectionLOD& RestLODData,
+	const FHairStrandsRootData& RestLODData,
 	const uint32 MeshLODIndex,
 	const FSkeletalMeshRenderData* InMeshRenderData, 
 	TArray<FHairStrandsMeshTrianglePositionFormat::Type>& OutDeformUniqueTrianglePositionBuffer)
@@ -226,7 +226,7 @@ TArray<FVector3f> GetDeformedHairStrandsPositions(
 	const uint32 MeshLODIndex,
 	const FSkeletalMeshRenderData* InMeshRenderData,
 	const TArray<FHairStrandsIndexFormat::Type>& PointToCurveBuffer,
-	const FHairStrandsRootData::FMeshProjectionLOD& RestLODData)
+	const FHairStrandsRootData& RestLODData)
 {
 	// Init the mesh samples with the target mesh vertices
 	const int32 MaxVertexCount = MeshVertexPositionsBuffer_Target.Num();
@@ -550,7 +550,7 @@ static void ExtractSkeletalVertexPosition(
 void DeformStaticMeshPositions(
 	UStaticMesh* OutMesh,
 	const TArray<FVector3f>& MeshVertexPositionsBuffer_Target,
-	const FHairStrandsRootData::FMeshProjectionLOD& RestLODData)
+	const FHairStrandsRootData& RestLODData)
 {
 	// Init the mesh samples with the target mesh vertices
 	const int32 MaxVertexCount = MeshVertexPositionsBuffer_Target.Num();
@@ -718,6 +718,10 @@ void FGroomRBFDeformer::GetRBFDeformedGroomAsset(const UGroomAsset* InGroomAsset
 			BuildPointToCurveMapping(GuidesData, SimRootDataPointToCurveBuffer);
 			BuildPointToCurveMapping(StrandsData, RenRootDataPointToCurveBuffer);
 
+			// Sanity check
+			check(BindingAsset->GetHairGroupsPlatformData()[GroupIndex].SimRootBulkDatas.IsValidIndex(MeshLODIndex));
+			check(BindingAsset->GetHairGroupsPlatformData()[GroupIndex].RenRootBulkDatas.IsValidIndex(MeshLODIndex));
+
 			// Get deformed guides
 			// If the groom override the value, we output dummy value for the guides, since they won't be used
 			if (InGroomAsset->GetHairGroupsInterpolation()[GroupIndex].InterpolationSettings.GuideType != EGroomGuideType::Imported)
@@ -728,7 +732,7 @@ void FGroomRBFDeformer::GetRBFDeformedGroomAsset(const UGroomAsset* InGroomAsset
 			else
 			{
 				FHairStrandsRootData SimRootData;
-				FGroomBindingBuilder::GetRootData(SimRootData, BindingAsset->GetHairGroupsPlatformData()[GroupIndex].SimRootBulkData);
+				FGroomBindingBuilder::GetRootData(SimRootData, BindingAsset->GetHairGroupsPlatformData()[GroupIndex].SimRootBulkDatas[MeshLODIndex]);
 
 				DeformedPositions[GroupIndex].GuideStrands = GetDeformedHairStrandsPositions(
 					MeshVertexPositionsBuffer_Target,
@@ -736,22 +740,23 @@ void FGroomRBFDeformer::GetRBFDeformedGroomAsset(const UGroomAsset* InGroomAsset
 					MeshLODIndex,
 					SkeletalMeshData_Target,
 					SimRootDataPointToCurveBuffer,
-					SimRootData.MeshProjectionLODs[MeshLODIndex]);
+					SimRootData);
 			}
 
 			// Get deformed render strands
 			{
 				FHairStrandsRootData RenRootData;
-				FGroomBindingBuilder::GetRootData(RenRootData, BindingAsset->GetHairGroupsPlatformData()[GroupIndex].RenRootBulkData);
+				FGroomBindingBuilder::GetRootData(RenRootData, BindingAsset->GetHairGroupsPlatformData()[GroupIndex].RenRootBulkDatas[MeshLODIndex]);
 
 				// Transfer RBF weights from SimRootData to RenRootData, as RenRootData does not hold RBF sample data
 				// This is transient data.
 				{
+
 					FHairStrandsRootData SimRootData;
-					FGroomBindingBuilder::GetRootData(SimRootData, BindingAsset->GetHairGroupsPlatformData()[GroupIndex].SimRootBulkData);
+					FGroomBindingBuilder::GetRootData(SimRootData, BindingAsset->GetHairGroupsPlatformData()[GroupIndex].SimRootBulkDatas[MeshLODIndex]);
 
-					FHairStrandsRootData::FMeshProjectionLOD& SimLODData = SimRootData.MeshProjectionLODs[MeshLODIndex];
-					FHairStrandsRootData::FMeshProjectionLOD& RenLODData = RenRootData.MeshProjectionLODs[MeshLODIndex];
+					FHairStrandsRootData& SimLODData = SimRootData;
+					FHairStrandsRootData& RenLODData = RenRootData;
 					RenLODData.SampleCount 						= SimLODData.SampleCount;
 					RenLODData.MeshInterpolationWeightsBuffer 	= SimLODData.MeshInterpolationWeightsBuffer;
 					RenLODData.MeshSampleIndicesBuffer 			= SimLODData.MeshSampleIndicesBuffer;
@@ -764,7 +769,7 @@ void FGroomRBFDeformer::GetRBFDeformedGroomAsset(const UGroomAsset* InGroomAsset
 					MeshLODIndex,
 					SkeletalMeshData_Target,
 					RenRootDataPointToCurveBuffer,
-					RenRootData.MeshProjectionLODs[MeshLODIndex]);
+					RenRootData);
 			}
 
 			// Trim strands based on mask texture
@@ -862,11 +867,13 @@ void FGroomRBFDeformer::GetRBFDeformedGroomAsset(const UGroomAsset* InGroomAsset
 				{
 					Mesh->ConditionalPostLoad();
 	
+					check(BindingAsset->GetHairGroupsPlatformData()[Desc.GroupIndex].SimRootBulkDatas.IsValidIndex(MeshLODIndex));
+
 					// Load Sim root data, as they contains RBF weights
 					FHairStrandsRootData SimRootData;
-					FGroomBindingBuilder::GetRootData(SimRootData, BindingAsset->GetHairGroupsPlatformData()[Desc.GroupIndex].SimRootBulkData);
+					FGroomBindingBuilder::GetRootData(SimRootData, BindingAsset->GetHairGroupsPlatformData()[Desc.GroupIndex].SimRootBulkDatas[MeshLODIndex]);
 	
-					DeformStaticMeshPositions(Mesh, MeshVertexPositionsBuffer_Target, SimRootData.MeshProjectionLODs[MeshLODIndex]);
+					DeformStaticMeshPositions(Mesh, MeshVertexPositionsBuffer_Target, SimRootData);
 				}
 			}
 		} 
@@ -883,11 +890,13 @@ void FGroomRBFDeformer::GetRBFDeformedGroomAsset(const UGroomAsset* InGroomAsset
 
 				Mesh->ConditionalPostLoad();
 
+				check(BindingAsset->GetHairGroupsPlatformData()[Desc.GroupIndex].SimRootBulkDatas.IsValidIndex(MeshLODIndex));
+
 				// Load Sim root data, as they contains RBF weights
 				FHairStrandsRootData SimRootData;
-				FGroomBindingBuilder::GetRootData(SimRootData, BindingAsset->GetHairGroupsPlatformData()[Desc.GroupIndex].SimRootBulkData);
+				FGroomBindingBuilder::GetRootData(SimRootData, BindingAsset->GetHairGroupsPlatformData()[Desc.GroupIndex].SimRootBulkDatas[MeshLODIndex]);
 
-				DeformStaticMeshPositions(Mesh, MeshVertexPositionsBuffer_Target, SimRootData.MeshProjectionLODs[MeshLODIndex]);
+				DeformStaticMeshPositions(Mesh, MeshVertexPositionsBuffer_Target, SimRootData);
 			}
 		}
 
diff --git a/Engine/Plugins/Runtime/HairStrands/Source/HairStrandsCore/Private/GroomResources.cpp b/Engine/Plugins/Runtime/HairStrands/Source/HairStrandsCore/Private/GroomResources.cpp
--- a/Engine/Plugins/Runtime/HairStrands/Source/HairStrandsCore/Private/GroomResources.cpp
+++ b/Engine/Plugins/Runtime/HairStrands/Source/HairStrandsCore/Private/GroomResources.cpp
@@ -677,8 +677,7 @@ void FHairCommonResource::InitRHI(FRHICommandListBase&)
 	}
 
 	// 3. Loaded the data
-	const int32 DummyLODIndex = -1;
-	const bool bReady = InternalGetOrRequestData(HAIR_MAX_NUM_CURVE_PER_GROUP, HAIR_MAX_NUM_POINT_PER_GROUP, DummyLODIndex);
+	const bool bReady = InternalGetOrRequestData(HAIR_MAX_NUM_CURVE_PER_GROUP, HAIR_MAX_NUM_POINT_PER_GROUP);
 	check(bReady);
 
 	// 4. Allocate the resource, and update the data
@@ -686,7 +685,7 @@ void FHairCommonResource::InitRHI(FRHICommandListBase&)
 	{
 		FRHICommandListImmediate& RHICmdList = FRHICommandListExecutor::GetImmediateCommandList();
 		FRDGBuilder GraphBuilder(RHICmdList);
-		InternalAllocate(GraphBuilder, HAIR_MAX_NUM_CURVE_PER_GROUP, HAIR_MAX_NUM_POINT_PER_GROUP, DummyLODIndex);
+		InternalAllocate(GraphBuilder, HAIR_MAX_NUM_CURVE_PER_GROUP, HAIR_MAX_NUM_POINT_PER_GROUP);
 		GraphBuilder.Execute();
 	}
 	else
@@ -702,13 +701,12 @@ void FHairCommonResource::ReleaseRHI()
 	bIsInitialized = false;
 }
 
-bool FHairCommonResource::InternalGetOrRequestData(uint32 InRequestedCurveCount, uint32 InRequestedPointCount, int32 InLODIndex) 
+bool FHairCommonResource::InternalGetOrRequestData(uint32 InRequestedCurveCount, uint32 InRequestedPointCount) 
 { 
 	// Initialize the streaming request so that resource which actually don't emit streaming request get correct #curve/#point/LODIndx
 	// request when allocating its resources
 	StreamingRequest.CurveCount = InRequestedCurveCount; 
 	StreamingRequest.PointCount = InRequestedPointCount;  
-	StreamingRequest.LODIndex   = InLODIndex;
 	return true; 
 }
 
@@ -717,16 +715,16 @@ void FHairCommonResource::Allocate(FRDGBuilder& GraphBuilder, EHairResourceLoadi
 	EHairResourceStatus Status;
 	Status.Status = EHairResourceStatus::EStatus::None;
 	Status.AvailableCurveCount = HAIR_MAX_NUM_CURVE_PER_GROUP;
-	Allocate(GraphBuilder, LoadingType, Status, HAIR_MAX_NUM_CURVE_PER_GROUP, HAIR_MAX_NUM_POINT_PER_GROUP, -1/*LODIndex*/);
+	Allocate(GraphBuilder, LoadingType, Status, HAIR_MAX_NUM_CURVE_PER_GROUP, HAIR_MAX_NUM_POINT_PER_GROUP);
 }
 
-void FHairCommonResource::Allocate(FRDGBuilder& GraphBuilder, EHairResourceLoadingType LoadingType, EHairResourceStatus& Status, int32 InLODIndex)
+void FHairCommonResource::Allocate(FRDGBuilder& GraphBuilder, EHairResourceLoadingType LoadingType, EHairResourceStatus& Status)
 {
-	Allocate(GraphBuilder, LoadingType, Status, HAIR_MAX_NUM_CURVE_PER_GROUP, HAIR_MAX_NUM_POINT_PER_GROUP, InLODIndex);
+	Allocate(GraphBuilder, LoadingType, Status, HAIR_MAX_NUM_CURVE_PER_GROUP, HAIR_MAX_NUM_POINT_PER_GROUP);
 }
 
 bool NeedDeallocation(uint32 InRequest, uint32 InAvailable);
-void FHairCommonResource::Allocate(FRDGBuilder& GraphBuilder, EHairResourceLoadingType LoadingType, EHairResourceStatus& Status, uint32 InRequestedCurveCount, uint32 InRequestedPointCount, int32 InLODIndex, bool bAllowDeallocation)
+void FHairCommonResource::Allocate(FRDGBuilder& GraphBuilder, EHairResourceLoadingType LoadingType, EHairResourceStatus& Status, uint32 InRequestedCurveCount, uint32 InRequestedPointCount, bool bAllowDeallocation)
 {
 	check(AllocationType == EHairStrandsAllocationType::Deferred);
 
@@ -738,9 +736,9 @@ void FHairCommonResource::Allocate(FRDGBuilder& GraphBuilder, EHairResourceLoadi
 		{
 			FRenderResource::InitResource(GraphBuilder.RHICmdList); // Call RenderResource InitResource() so that the resources is marked as initialized
 		}
-		if (!bIsInitialized || !InternalIsLODDataLoaded(InRequestedCurveCount, InRequestedPointCount, InLODIndex))
+		if (!bIsInitialized)
 		{
-			InternalAllocate(GraphBuilder, InRequestedCurveCount, InRequestedPointCount, InLODIndex);
+			InternalAllocate(GraphBuilder, InRequestedCurveCount, InRequestedPointCount);
 		}
 		bIsInitialized = true;
 
@@ -752,19 +750,18 @@ void FHairCommonResource::Allocate(FRDGBuilder& GraphBuilder, EHairResourceLoadi
 	else if (LoadingType == EHairResourceLoadingType::Async)
 	{
 		// 1. If all requested curve are already loaded, nothing to do
-		if (bIsInitialized && MaxAvailableCurveCount >= InRequestedCurveCount && InternalIsLODDataLoaded(InRequestedCurveCount, InRequestedPointCount, InLODIndex)) 
+		if (bIsInitialized && MaxAvailableCurveCount >= InRequestedCurveCount) 
 		{ 
 			// Trim/Un-stream data if needed (ensure no streaming request is in-flight)
 			if (bAllowDeallocation && NeedDeallocation(InRequestedCurveCount, MaxAvailableCurveCount) && StreamingRequest.IsNone())
 			{
 				StreamingRequest.CurveCount = InRequestedCurveCount;
 				StreamingRequest.PointCount = InRequestedPointCount;
-				StreamingRequest.LODIndex = InLODIndex;
 				if (FHairStrandsBulkCommon* Bulk = InternalGetBulkData())
 				{
 					Bulk->Unload(&StreamingRequest);
 				}
-				InternalAllocate(GraphBuilder, InRequestedCurveCount, InRequestedPointCount, StreamingRequest.LODIndex); // Do we need to pass something else for the deallocation
+				InternalAllocate(GraphBuilder, InRequestedCurveCount, InRequestedPointCount); // Do we need to pass something else for the deallocation
 
 				// Update the max curve count available
 				MaxAvailableCurveCount = StreamingRequest.CurveCount;
@@ -776,14 +773,14 @@ void FHairCommonResource::Allocate(FRDGBuilder& GraphBuilder, EHairResourceLoadi
 			Status |= EHairResourceStatus::EStatus::Valid; 
 		}
 		// 2. If more curves are requested, issue a streaming request
-		else if (InternalGetOrRequestData(InRequestedCurveCount, InRequestedPointCount, InLODIndex))
+		else if (InternalGetOrRequestData(InRequestedCurveCount, InRequestedPointCount))
 		{ 
 			// 2.1 Curve data are availble, and update GPU resources
 			if (!bIsInitialized)
 			{
 				FRenderResource::InitResource(GraphBuilder.RHICmdList); // Call RenderResource InitResource() so that the resources is marked as initialized
 			}
-			InternalAllocate(GraphBuilder, StreamingRequest.CurveCount, StreamingRequest.PointCount, StreamingRequest.LODIndex);
+			InternalAllocate(GraphBuilder, StreamingRequest.CurveCount, StreamingRequest.PointCount);
 			bIsInitialized = true;
 
 			// Update the max curve count available
@@ -792,17 +789,7 @@ void FHairCommonResource::Allocate(FRDGBuilder& GraphBuilder, EHairResourceLoadi
 			// Reset streaming request. When the request is delete, the DDC request becomes cancelled. 
 			StreamingRequest = FHairStreamingRequest();
 
-			// Validate the resource the resource again, as the LODIndex at request time (StreamingRequest.LODIndex) 
-			// and the current LODIndex (InLODIndex) might not match.
-			// If they don't match a new streaming request will be emitted next frame
-			if (InternalIsLODDataLoaded(InRequestedCurveCount, InRequestedPointCount, InLODIndex))
-			{
-				Status |= EHairResourceStatus::EStatus::Valid;
-			}
-			else
-			{
-				Status |= EHairResourceStatus::EStatus::Loading;
-			}
+			Status |= EHairResourceStatus::EStatus::Valid;
 		}
 		else
 		{
@@ -814,12 +801,12 @@ void FHairCommonResource::Allocate(FRDGBuilder& GraphBuilder, EHairResourceLoadi
 	Status.AddAvailableCurve(MaxAvailableCurveCount);
 }
 
-void FHairCommonResource::StreamInData(int32 InLODIndex)
+void FHairCommonResource::StreamInData()
 {
 	if (!bIsInitialized)
 	{
 		// TODO
-		// InternalGetOrRequestData(HAIR_MAX_NUM_CURVE_PER_GROUP, HAIR_MAX_NUM_POINT_PER_GROUP, InLODIndex);
+		// InternalGetOrRequestData(HAIR_MAX_NUM_CURVE_PER_GROUP, HAIR_MAX_NUM_POINT_PER_GROUP);
 	}
 }
 
@@ -1021,11 +1008,11 @@ FHairStrandsRestResource::FHairStrandsRestResource(FHairStrandsBulkData& InBulkD
 	check(!!(BulkData.Header.Flags & FHairStrandsBulkData::DataFlags_HasData));
 }
 
-bool FHairStrandsRestResource::InternalGetOrRequestData(uint32 InRequestedCurveCount, uint32 InRequestedPointCount, int32 InLODIndex)
+bool FHairStrandsRestResource::InternalGetOrRequestData(uint32 InRequestedCurveCount, uint32 InRequestedPointCount)
 {
 	if (StreamingRequest.IsNone())
 	{
-		StreamingRequest.Request(InRequestedCurveCount, InRequestedPointCount, InLODIndex, BulkData, false, false, false, OwnerName);
+		StreamingRequest.Request(InRequestedCurveCount, InRequestedPointCount, BulkData, false, false, false, OwnerName);
 	}
 	return StreamingRequest.IsCompleted();
 }
@@ -1141,7 +1128,7 @@ FHairStrandsDeformedResource::FHairStrandsDeformedResource(FHairStrandsBulkData&
 	GetPositionOffset(FHairStrandsDeformedResource::EFrameType::Previous) = BulkData.GetPositionOffset();
 }
 
-void FHairStrandsDeformedResource::InternalAllocate(FRDGBuilder& GraphBuilder, uint32 InCurveCount, uint32 InPointCount, int32 InLODIndex)
+void FHairStrandsDeformedResource::InternalAllocate(FRDGBuilder& GraphBuilder, uint32 InCurveCount, uint32 InPointCount)
 {
 	// Need to update these to by dynamic based on the active number of point
 	const uint32 PointCount = FMath::Min(InPointCount, BulkData.GetNumPoints());
@@ -1226,11 +1213,11 @@ FHairStrandsClusterResource::FHairStrandsClusterResource(FHairStrandsClusterBulk
 	MaxAvailableCurveCount = 0;
 }
 
-bool FHairStrandsClusterResource::InternalGetOrRequestData(uint32 InRequestedCurveCount, uint32 InRequestedPointCount, int32 InLODIndex)
+bool FHairStrandsClusterResource::InternalGetOrRequestData(uint32 InRequestedCurveCount, uint32 InRequestedPointCount)
 {
 	if (StreamingRequest.IsNone())
 	{
-		StreamingRequest.Request(InRequestedCurveCount, InRequestedPointCount, InLODIndex, BulkData, false, false, false, OwnerName);
+		StreamingRequest.Request(InRequestedCurveCount, InRequestedPointCount, BulkData, false, false, false, OwnerName);
 	}
 	return StreamingRequest.IsCompleted();
 }
@@ -1306,7 +1293,7 @@ FHairCommonResource(EHairStrandsAllocationType::Deferred, InResourceName, InOwne
 	MaxAvailableCurveCount = 0;
 }
 
-void FHairStrandsCullingResource::InternalAllocate(FRDGBuilder& GraphBuilder, uint32 InCurveCount, uint32 InPointCount, int32 InLODIndex)
+void FHairStrandsCullingResource::InternalAllocate(FRDGBuilder& GraphBuilder, uint32 InCurveCount, uint32 InPointCount)
 {
 	if (ClusterCount > 0)
 	{	
@@ -1340,220 +1327,277 @@ void FHairStrandsCullingResource::InternalRelease()
 
 /////////////////////////////////////////////////////////////////////////////////////////
 
-FHairStrandsRestRootResource::FHairStrandsRestRootResource(FHairStrandsRootBulkData& InBulkData, EHairStrandsResourcesType InCurveType, const FHairResourceName& InResourceName, const FName& InOwnerName) :
+FHairStrandsLODRestRootResource::FHairStrandsLODRestRootResource(FHairStrandsRootBulkData& InBulkData, EHairStrandsResourcesType InCurveType, uint32 InMeshLODIndex, const FHairResourceName& InResourceName, const FName& InOwnerName) :
 	FHairCommonResource(EHairStrandsAllocationType::Deferred, InResourceName, InOwnerName),
-	BulkData(InBulkData), CurveType(InCurveType)
+	BulkData(InBulkData), CurveType(InCurveType), MeshLODIndex(InMeshLODIndex)
 {
 	// MaxAvailableCurveCount is not set to 0, as we need to track the loaded curve data *per LOD*. 
 	// This is done through callback to InternalIsLODDataLoaded()
 	// MaxAvailableCurveCount = 0;
 
-	PopulateFromRootData();
+	Reset();
 }
 
-void FHairStrandsRestRootResource::PopulateFromRootData()
+void FHairStrandsLODRestRootResource::Reset()
 {
-	LODs.Reserve(BulkData.Header.LODs.Num());
-	for (const FHairStrandsRootBulkData::FHeader::FLOD& InLOD : BulkData.Header.LODs)
-	{
-		FLOD& OutLOD = LODs.AddDefaulted_GetRef();
-
-		OutLOD.LODIndex 	= InLOD.LODIndex;
-		OutLOD.Status 		= FLOD::EStatus::Invalid;
-		OutLOD.SampleCount 	= InLOD.SampleCount;
-	}
-}
-
-bool FHairStrandsRestRootResource::InternalIsLODDataLoaded(uint32 InRequestedCurveCount, uint32 InRequestedPointCount, int32 InLODIndex) const
-{
-	if (InLODIndex >= 0 && InLODIndex < LODs.Num())
-	{
-		// Clamp request to the max supported
-		InRequestedCurveCount = FMath::Min(InRequestedCurveCount, BulkData.Header.RootCount);
-
-		const FLOD& GPUData = LODs[InLODIndex];
-		return (GPUData.Status == FLOD::EStatus::Completed || GPUData.Status == FLOD::EStatus::Initialized) && InRequestedCurveCount <= GPUData.AvailableCurveCount;
-	}
-	return true;
+	check(MeshLODIndex == BulkData.Header.LODIndex);
+	Status 		= EStatus::Invalid;
+	SampleCount = BulkData.Header.SampleCount;
 }
 
-bool FHairStrandsRestRootResource::InternalGetOrRequestData(uint32 InRequestedCurveCount, uint32 InRequestedPointCount, int32 InLODIndex)
+bool FHairStrandsLODRestRootResource::InternalGetOrRequestData(uint32 InRequestedCurveCount, uint32 InRequestedPointCount)
 {
-	if (StreamingRequest.IsNone() && InLODIndex >= 0)
+	if (StreamingRequest.IsNone())
 	{
-		StreamingRequest.Request(InRequestedCurveCount, InRequestedPointCount, InLODIndex, BulkData, false, false, false, OwnerName);
+		StreamingRequest.Request(InRequestedCurveCount, InRequestedPointCount, BulkData, false, false, false, OwnerName);
 	}
-	return InLODIndex >= 0 && StreamingRequest.IsCompleted();
+	return StreamingRequest.IsCompleted();
 }
 
-void FHairStrandsRestRootResource::InternalAllocate(FRDGBuilder& GraphBuilder, uint32 InCurveCount, uint32 InPointCount, int32 InLODIndex)
+void FHairStrandsLODRestRootResource::InternalAllocate(FRDGBuilder& GraphBuilder, uint32 InCurveCount, uint32 InPointCount)
 {
 	// Once empty, the MeshProjectionLODsneeds to be repopulate as it might be re-initialized. 
 	// E.g., when a resource is updated, it is first released, then re-init. 
-	if (LODs.IsEmpty())
+	if (!bIsInitialized)
 	{
-		PopulateFromRootData();
+		Reset();
 	}
 
 	check(BulkData.Header.PointCount > 0);
-	check(LODs.Num() == BulkData.Header.LODs.Num());
-	if (InLODIndex >= 0 && InLODIndex < LODs.Num())
 	{
 		// Clamp to the max supported
 		InCurveCount = FMath::Min(InCurveCount, BulkData.Header.RootCount);
 
-		FLOD& GPUData = LODs[InLODIndex];
-		const bool bIsLODInitialized = (GPUData.Status == FLOD::EStatus::Completed || GPUData.Status == FLOD::EStatus::Initialized) && InCurveCount <= GPUData.AvailableCurveCount;
+		const bool bIsLODInitialized = (Status == EStatus::Completed || Status == EStatus::Initialized) && InCurveCount <= AvailableCurveCount;
 		if (bIsLODInitialized)
 		{
 			return;
 		}
 
-		const FHairStrandsRootBulkData::FHeader::FLOD& LODHeader = BulkData.Header.LODs[InLODIndex];
-		FHairStrandsRootBulkData::FData::FLOD& CPUData = BulkData.Data.LODs[InLODIndex];
 		{
-			InternalCreateVertexBufferRDG_FromHairBulkData<FHairStrandsRootBarycentricFormat>(GraphBuilder, CPUData.RootBarycentricBuffer, BulkData.Header.RootCount, GPUData.RootBarycentricBuffer, ToHairResourceDebugName(HAIRSTRANDS_RESOUCE_NAME(CurveType, Hair.StrandsRestRoot_RootTriangleBarycentricBuffer), ResourceName), OwnerName, EHairResourceUsageType::Static);
-			InternalCreateVertexBufferRDG_FromHairBulkData<FHairStrandsRootToUniqueTriangleIndexFormat>(GraphBuilder, CPUData.RootToUniqueTriangleIndexBuffer, BulkData.Header.RootCount, GPUData.RootToUniqueTriangleIndexBuffer, ToHairResourceDebugName(HAIRSTRANDS_RESOUCE_NAME(CurveType, Hair.StrandsRestRoot_RootToUniqueTriangleIndexBuffer), ResourceName), OwnerName, EHairResourceUsageType::Static);
-			InternalCreateVertexBufferRDG_FromHairBulkData<FHairStrandsUniqueTriangleIndexFormat>(GraphBuilder, CPUData.UniqueTriangleIndexBuffer, LODHeader.UniqueTriangleCount, GPUData.UniqueTriangleIndexBuffer, ToHairResourceDebugName(HAIRSTRANDS_RESOUCE_NAME(CurveType, Hair.StrandsRestRoot_UniqueTriangleIndexBuffer), ResourceName), OwnerName, EHairResourceUsageType::Static);
-			InternalCreateVertexBufferRDG_FromHairBulkData<FHairStrandsMeshTrianglePositionFormat>(GraphBuilder, CPUData.RestUniqueTrianglePositionBuffer, LODHeader.UniqueTriangleCount * 3, GPUData.RestUniqueTrianglePositionBuffer, ToHairResourceDebugName(HAIRSTRANDS_RESOUCE_NAME(CurveType, Hair.StrandsRestRoot_RestUniqueTrianglePosition0Buffer), ResourceName), OwnerName, EHairResourceUsageType::Static);
+			InternalCreateVertexBufferRDG_FromHairBulkData<FHairStrandsRootBarycentricFormat>(GraphBuilder, BulkData.Data.RootBarycentricBuffer, BulkData.Header.RootCount, RootBarycentricBuffer, ToHairResourceDebugName(HAIRSTRANDS_RESOUCE_NAME(CurveType, Hair.StrandsRestRoot_RootTriangleBarycentricBuffer), ResourceName), OwnerName, EHairResourceUsageType::Static);
+			InternalCreateVertexBufferRDG_FromHairBulkData<FHairStrandsRootToUniqueTriangleIndexFormat>(GraphBuilder, BulkData.Data.RootToUniqueTriangleIndexBuffer, BulkData.Header.RootCount, RootToUniqueTriangleIndexBuffer, ToHairResourceDebugName(HAIRSTRANDS_RESOUCE_NAME(CurveType, Hair.StrandsRestRoot_RootToUniqueTriangleIndexBuffer), ResourceName), OwnerName, EHairResourceUsageType::Static);
+			InternalCreateVertexBufferRDG_FromHairBulkData<FHairStrandsUniqueTriangleIndexFormat>(GraphBuilder, BulkData.Data.UniqueTriangleIndexBuffer, BulkData.Header.UniqueTriangleCount, UniqueTriangleIndexBuffer, ToHairResourceDebugName(HAIRSTRANDS_RESOUCE_NAME(CurveType, Hair.StrandsRestRoot_UniqueTriangleIndexBuffer), ResourceName), OwnerName, EHairResourceUsageType::Static);
+			InternalCreateVertexBufferRDG_FromHairBulkData<FHairStrandsMeshTrianglePositionFormat>(GraphBuilder, BulkData.Data.RestUniqueTrianglePositionBuffer, BulkData.Header.UniqueTriangleCount * 3, RestUniqueTrianglePositionBuffer, ToHairResourceDebugName(HAIRSTRANDS_RESOUCE_NAME(CurveType, Hair.StrandsRestRoot_RestUniqueTrianglePosition0Buffer), ResourceName), OwnerName, EHairResourceUsageType::Static);
 		}
 
 		// Allocate only once
-		if (LODHeader.SampleCount > 0 && GPUData.MeshInterpolationWeightsBuffer.Buffer == nullptr)
+		if (BulkData.Header.SampleCount > 0 && MeshInterpolationWeightsBuffer.Buffer == nullptr)
 		{
-			const uint32 InteroplationWeightCount = CPUData.MeshInterpolationWeightsBuffer.GetBulkDataSize() / sizeof(FHairStrandsWeightFormat::Type);
-			const uint32 WeightCount = FGroomRBFDeformer::GetWeightCount(LODHeader.SampleCount);
+			const uint32 InteroplationWeightCount = BulkData.Data.MeshInterpolationWeightsBuffer.GetBulkDataSize() / sizeof(FHairStrandsWeightFormat::Type);
+			const uint32 WeightCount = FGroomRBFDeformer::GetWeightCount(BulkData.Header.SampleCount);
 			check(WeightCount == InteroplationWeightCount); // Sanity check
 
-			InternalCreateStructuredBufferRDG_FromHairBulkData<FHairStrandsWeightFormat>(GraphBuilder, CPUData.MeshInterpolationWeightsBuffer, InteroplationWeightCount, GPUData.MeshInterpolationWeightsBuffer, ToHairResourceDebugName(HAIRSTRANDS_RESOUCE_NAME(CurveType, Hair.StrandsRestRoot_MeshInterpolationWeightsBuffer), ResourceName), OwnerName, EHairResourceUsageType::Static);
-			InternalCreateVertexBufferRDG_FromHairBulkData<FHairStrandsRBFSampleIndexFormat>(GraphBuilder, CPUData.MeshSampleIndicesAndSectionsBuffer, LODHeader.SampleCount, GPUData.MeshSampleIndicesAndSectionsBuffer, ToHairResourceDebugName(HAIRSTRANDS_RESOUCE_NAME(CurveType, Hair.StrandsRestRoot_MeshSampleIndicesAndSectionsBuffer), ResourceName), OwnerName, EHairResourceUsageType::Static);
-			InternalCreateStructuredBufferRDG_FromHairBulkData<FHairStrandsMeshTrianglePositionFormat>(GraphBuilder, CPUData.RestSamplePositionsBuffer, LODHeader.SampleCount, GPUData.RestSamplePositionsBuffer, ToHairResourceDebugName(HAIRSTRANDS_RESOUCE_NAME(CurveType, Hair.StrandsRestRoot_RestSamplePositionsBuffer), ResourceName), OwnerName, EHairResourceUsageType::Static);
+			InternalCreateStructuredBufferRDG_FromHairBulkData<FHairStrandsWeightFormat>(GraphBuilder, BulkData.Data.MeshInterpolationWeightsBuffer, InteroplationWeightCount, MeshInterpolationWeightsBuffer, ToHairResourceDebugName(HAIRSTRANDS_RESOUCE_NAME(CurveType, Hair.StrandsRestRoot_MeshInterpolationWeightsBuffer), ResourceName), OwnerName, EHairResourceUsageType::Static);
+			InternalCreateVertexBufferRDG_FromHairBulkData<FHairStrandsRBFSampleIndexFormat>(GraphBuilder, BulkData.Data.MeshSampleIndicesAndSectionsBuffer, BulkData.Header.SampleCount, MeshSampleIndicesAndSectionsBuffer, ToHairResourceDebugName(HAIRSTRANDS_RESOUCE_NAME(CurveType, Hair.StrandsRestRoot_MeshSampleIndicesAndSectionsBuffer), ResourceName), OwnerName, EHairResourceUsageType::Static);
+			InternalCreateStructuredBufferRDG_FromHairBulkData<FHairStrandsMeshTrianglePositionFormat>(GraphBuilder, BulkData.Data.RestSamplePositionsBuffer, BulkData.Header.SampleCount, RestSamplePositionsBuffer, ToHairResourceDebugName(HAIRSTRANDS_RESOUCE_NAME(CurveType, Hair.StrandsRestRoot_RestSamplePositionsBuffer), ResourceName), OwnerName, EHairResourceUsageType::Static);
 		}
 
-		GPUData.Status = FLOD::EStatus::Completed;
-		GPUData.AvailableCurveCount = InCurveCount;
-		GPUData.SampleCount = LODHeader.SampleCount;
+		Status = EStatus::Completed;
+		AvailableCurveCount = InCurveCount;
+		SampleCount = BulkData.Header.SampleCount;
 	}
 }
 
-void FHairStrandsRestRootResource::InternalRelease()
+void FHairStrandsLODRestRootResource::InternalRelease()
 {
-	for (FLOD& GPUData : LODs)
-	{
-		GPUData.Status = FLOD::EStatus::Invalid;
-		GPUData.RootBarycentricBuffer.Release();
-		GPUData.RootToUniqueTriangleIndexBuffer.Release();
-		GPUData.UniqueTriangleIndexBuffer.Release();
-		GPUData.RestUniqueTrianglePositionBuffer.Release();
-		GPUData.SampleCount = 0;
-		GPUData.MeshInterpolationWeightsBuffer.Release();
-		GPUData.MeshSampleIndicesAndSectionsBuffer.Release();
-		GPUData.RestSamplePositionsBuffer.Release();
-		GPUData.AvailableCurveCount = 0;
-	}
-	LODs.Empty();
+	Status = EStatus::Invalid;
+	RootBarycentricBuffer.Release();
+	RootToUniqueTriangleIndexBuffer.Release();
+	UniqueTriangleIndexBuffer.Release();
+	RestUniqueTrianglePositionBuffer.Release();
+	SampleCount = 0;
+	MeshInterpolationWeightsBuffer.Release();
+	MeshSampleIndicesAndSectionsBuffer.Release();
+	RestSamplePositionsBuffer.Release();
+	AvailableCurveCount = 0;
 }
 
-void FHairStrandsRestRootResource::InternalResetLoadedSize()
+void FHairStrandsLODRestRootResource::InternalResetLoadedSize()
 {
 	BulkData.ResetLoadedSize();
 }
 
 /////////////////////////////////////////////////////////////////////////////////////////
 
-FHairStrandsDeformedRootResource::FHairStrandsDeformedRootResource(EHairStrandsResourcesType InCurveType, const FHairResourceName& InResourceName, const FName& InOwnerName):
-	FHairCommonResource(EHairStrandsAllocationType::Deferred, InResourceName, InOwnerName),
-	CurveType(InCurveType)
+FHairStrandsRestRootResource::FHairStrandsRestRootResource(TArray<FHairStrandsRootBulkData>& InBulkDatas, EHairStrandsResourcesType InCurveType, const FHairResourceName& InResourceName, const FName& InOwnerName)
+: CurveType(InCurveType)
 {
+	for (FHairStrandsRootBulkData& LODBulkData : InBulkDatas)
+	{
+		LODs.Add(new FHairStrandsLODRestRootResource(LODBulkData, InCurveType, LODBulkData.Header.LODIndex, InResourceName, InOwnerName));
+	}
+}
 
+void FHairStrandsRestRootResource::StreamInData(int32 InMeshLODIndex)
+{
+	if (IsDataValid(InMeshLODIndex))
+	{
+		GetLOD(InMeshLODIndex)->StreamInData();
+	}
 }
 
-FHairStrandsDeformedRootResource::FHairStrandsDeformedRootResource(const FHairStrandsRestRootResource* InRestResources, EHairStrandsResourcesType InCurveType, const FHairResourceName& InResourceName, const FName& InOwnerName):
-	FHairCommonResource(EHairStrandsAllocationType::Deferred, InResourceName, InOwnerName),
-	CurveType(InCurveType)
+void FHairStrandsRestRootResource::Allocate(FRDGBuilder& GraphBuilder, EHairResourceLoadingType LoadingType, EHairResourceStatus& Status, int32 InMeshLODIndex)
 {
-	check(InRestResources);
-	RootCount = InRestResources->BulkData.Header.RootCount;
-	LODs.Reserve(InRestResources->LODs.Num());
-	for (const FHairStrandsRestRootResource::FLOD& InLOD : InRestResources->LODs)
+	if (IsDataValid(InMeshLODIndex))
 	{
-		FLOD& LOD = LODs.AddDefaulted_GetRef();
+		GetLOD(InMeshLODIndex)->Allocate(GraphBuilder, LoadingType, Status);
+	}
+}
 
-		LOD.Status = FLOD::EStatus::Invalid;
-		LOD.LODIndex = InLOD.LODIndex;
-		LOD.SampleCount = InLOD.SampleCount;
+void FHairStrandsRestRootResource::Allocate(FRDGBuilder& GraphBuilder, EHairResourceLoadingType LoadingType, EHairResourceStatus& Status, uint32 InRequestedCurveCount, uint32 InRequestedPointCount, bool bAllowDeallocation, int32 InMeshLODIndex)
+{
+	if (IsDataValid(InMeshLODIndex))
+	{
+		GetLOD(InMeshLODIndex)->Allocate(GraphBuilder, LoadingType, Status, InRequestedCurveCount, InRequestedPointCount, bAllowDeallocation);
 	}
 }
 
-bool FHairStrandsDeformedRootResource::InternalIsLODDataLoaded(uint32 InRequestedCurveCount, uint32 InRequestedPointCount, int32 InLODIndex) const
+void FHairStrandsRestRootResource::ReleaseResource()
 {
-	if (InLODIndex >= 0 && InLODIndex < LODs.Num())
+	for (FHairStrandsLODRestRootResource* LOD : LODs)
 	{
-		// Clamp request to the max supported
-		InRequestedCurveCount = FMath::Min(InRequestedCurveCount, RootCount);
+		LOD->ReleaseResource();
+		delete LOD;
+	}
+	LODs.Empty();
+}
 
-		const FLOD& GPUData = LODs[InLODIndex];
-		return (GPUData.Status == FLOD::EStatus::Completed || GPUData.Status == FLOD::EStatus::Initialized) && InRequestedCurveCount <= GPUData.AvailableCurveCount;
+void FHairStrandsRestRootResource::InternalResetLoadedSize()
+{
+	for (FHairStrandsLODRestRootResource* LOD : LODs)
+	{
+		LOD->InternalResetLoadedSize();
 	}
-	return true;
 }
 
-void FHairStrandsDeformedRootResource::InternalAllocate(FRDGBuilder& GraphBuilder, uint32 InCurveCount, uint32 InPointCount, int32 InLODIndex)
+void FHairStrandsRestRootResource::BeginInitResource()
 {
-	if (RootCount > 0 && InLODIndex >= 0 && InLODIndex < LODs.Num())
+	for (FHairStrandsLODRestRootResource* LOD : LODs)
 	{
-		FLOD& LOD = LODs[InLODIndex];
-		if (LOD.Status == FLOD::EStatus::Invalid)
-		{		
-			LOD.Status = FLOD::EStatus::Initialized;
-			LOD.AvailableCurveCount = FMath::Min(InCurveCount, RootCount);
-			if (LOD.SampleCount > 0)
-			{
-				const uint32 RBFEntryCount = FGroomRBFDeformer::GetEntryCount(LOD.SampleCount);
+		::BeginInitResource(LOD);
+	}
+}
 
-				InternalCreateStructuredBufferRDG<FHairStrandsMeshTrianglePositionFormat>(GraphBuilder, LOD.SampleCount, LOD.DeformedSamplePositionsBuffer[0], ToHairResourceDebugName(HAIRSTRANDS_RESOUCE_NAME(CurveType, Hair.StrandsRootDeformed_DeformedSamplePositionsBuffer0), ResourceName), OwnerName, EHairResourceUsageType::Dynamic);
-				InternalCreateStructuredBufferRDG<FHairStrandsMeshTrianglePositionFormat>(GraphBuilder, RBFEntryCount, LOD.MeshSampleWeightsBuffer[0], ToHairResourceDebugName(HAIRSTRANDS_RESOUCE_NAME(CurveType, Hair.StrandsRootDeformed_MeshSampleWeightsBuffer0), ResourceName), OwnerName, EHairResourceUsageType::Dynamic);
+void FHairStrandsRestRootResource::BeginUpdateResourceRHI()
+{
+	for (FHairStrandsLODRestRootResource* LOD : LODs)
+	{
+		::BeginUpdateResourceRHI(LOD);
+	}
+}
 
-				// Double buffering is disabled by default unless the read-only cvar r.HairStrands.ContinuousDecimationReordering is set
-				if (IsHairStrandContinuousDecimationReorderingEnabled())
-				{
-					InternalCreateStructuredBufferRDG<FHairStrandsMeshTrianglePositionFormat>(GraphBuilder, LOD.SampleCount, LOD.DeformedSamplePositionsBuffer[1], ToHairResourceDebugName(HAIRSTRANDS_RESOUCE_NAME(CurveType, Hair.StrandsRootDeformed_DeformedSamplePositionsBuffer1), ResourceName), OwnerName, EHairResourceUsageType::Dynamic);
-					InternalCreateStructuredBufferRDG<FHairStrandsMeshTrianglePositionFormat>(GraphBuilder, RBFEntryCount, LOD.MeshSampleWeightsBuffer[1], ToHairResourceDebugName(HAIRSTRANDS_RESOUCE_NAME(CurveType, Hair.StrandsRootDeformed_MeshSampleWeightsBuffer1), ResourceName), OwnerName, EHairResourceUsageType::Dynamic);
-				}
-			}
+/////////////////////////////////////////////////////////////////////////////////////////
 
-			InternalCreateVertexBufferRDG<FHairStrandsMeshTrianglePositionFormat>(GraphBuilder, RootCount * 3, LOD.DeformedUniqueTrianglePositionBuffer[0], ToHairResourceDebugName(HAIRSTRANDS_RESOUCE_NAME(CurveType, Hair.StrandsRootDeformed_DeformedUniqueTrianglePosition0Buffer0), ResourceName), OwnerName, EHairResourceUsageType::Dynamic);
+FHairStrandsLODDeformedRootResource::FHairStrandsLODDeformedRootResource(FHairStrandsRestRootResource* InRestResources, EHairStrandsResourcesType InCurveType, const uint32 InMeshLODIndex, const FHairResourceName& InResourceName, const FName& InOwnerName)
+:	FHairCommonResource(EHairStrandsAllocationType::Deferred, InResourceName, InOwnerName)
+,	MeshLODIndex(InMeshLODIndex) 
+,	CurveType(InCurveType)
+{
+	check(InRestResources);
+	check(InRestResources->IsDataValid(InMeshLODIndex));
+	FHairStrandsLODRestRootResource* LODData = InRestResources->GetLOD(InMeshLODIndex);
+	RootCount = LODData->BulkData.Header.RootCount;
+	SampleCount = LODData->BulkData.Header.SampleCount;
+	Status = EStatus::Invalid;
+}
+
+void FHairStrandsLODDeformedRootResource::InternalAllocate(FRDGBuilder& GraphBuilder, uint32 InCurveCount, uint32 InPointCount)
+{
+	if (RootCount > 0 && Status == EStatus::Invalid)
+	{
+		Status = EStatus::Initialized;
+		AvailableCurveCount = FMath::Min(InCurveCount, RootCount);
+		if (SampleCount > 0)
+		{
+			const uint32 RBFEntryCount = FGroomRBFDeformer::GetEntryCount(SampleCount);
+
+			InternalCreateStructuredBufferRDG<FHairStrandsMeshTrianglePositionFormat>(GraphBuilder, SampleCount, DeformedSamplePositionsBuffer[0], ToHairResourceDebugName(HAIRSTRANDS_RESOUCE_NAME(CurveType, Hair.StrandsRootDeformed_DeformedSamplePositionsBuffer0), ResourceName), OwnerName, EHairResourceUsageType::Dynamic);
+			InternalCreateStructuredBufferRDG<FHairStrandsMeshTrianglePositionFormat>(GraphBuilder, RBFEntryCount, MeshSampleWeightsBuffer[0], ToHairResourceDebugName(HAIRSTRANDS_RESOUCE_NAME(CurveType, Hair.StrandsRootDeformed_MeshSampleWeightsBuffer0), ResourceName), OwnerName, EHairResourceUsageType::Dynamic);
 
 			// Double buffering is disabled by default unless the read-only cvar r.HairStrands.ContinuousDecimationReordering is set
 			if (IsHairStrandContinuousDecimationReorderingEnabled())
 			{
-				InternalCreateVertexBufferRDG<FHairStrandsMeshTrianglePositionFormat>(GraphBuilder, RootCount * 3, LOD.DeformedUniqueTrianglePositionBuffer[1], ToHairResourceDebugName(HAIRSTRANDS_RESOUCE_NAME(CurveType, Hair.StrandsRootDeformed_DeformedUniqueTrianglePosition0Buffer1), ResourceName), OwnerName, EHairResourceUsageType::Dynamic);
+				InternalCreateStructuredBufferRDG<FHairStrandsMeshTrianglePositionFormat>(GraphBuilder, SampleCount, DeformedSamplePositionsBuffer[1], ToHairResourceDebugName(HAIRSTRANDS_RESOUCE_NAME(CurveType, Hair.StrandsRootDeformed_DeformedSamplePositionsBuffer1), ResourceName), OwnerName, EHairResourceUsageType::Dynamic);
+				InternalCreateStructuredBufferRDG<FHairStrandsMeshTrianglePositionFormat>(GraphBuilder, RBFEntryCount, MeshSampleWeightsBuffer[1], ToHairResourceDebugName(HAIRSTRANDS_RESOUCE_NAME(CurveType, Hair.StrandsRootDeformed_MeshSampleWeightsBuffer1), ResourceName), OwnerName, EHairResourceUsageType::Dynamic);
 			}
 		}
-	}
-}
 
-void FHairStrandsDeformedRootResource::InternalRelease()
-{
-	for (FLOD& GPUData : LODs)
-	{
-		GPUData.Status = FLOD::EStatus::Invalid;
-		GPUData.AvailableCurveCount = 0;
-		GPUData.DeformedUniqueTrianglePositionBuffer[0].Release();
-		GPUData.DeformedSamplePositionsBuffer[0].Release();
-		GPUData.MeshSampleWeightsBuffer[0].Release();
-	
+		InternalCreateVertexBufferRDG<FHairStrandsMeshTrianglePositionFormat>(GraphBuilder, RootCount * 3, DeformedUniqueTrianglePositionBuffer[0], ToHairResourceDebugName(HAIRSTRANDS_RESOUCE_NAME(CurveType, Hair.StrandsRootDeformed_DeformedUniqueTrianglePosition0Buffer0), ResourceName), OwnerName, EHairResourceUsageType::Dynamic);
+
 		// Double buffering is disabled by default unless the read-only cvar r.HairStrands.ContinuousDecimationReordering is set
 		if (IsHairStrandContinuousDecimationReorderingEnabled())
 		{
-			GPUData.DeformedUniqueTrianglePositionBuffer[1].Release();
-			GPUData.DeformedSamplePositionsBuffer[1].Release();
-			GPUData.MeshSampleWeightsBuffer[1].Release();
+			InternalCreateVertexBufferRDG<FHairStrandsMeshTrianglePositionFormat>(GraphBuilder, RootCount * 3, DeformedUniqueTrianglePositionBuffer[1], ToHairResourceDebugName(HAIRSTRANDS_RESOUCE_NAME(CurveType, Hair.StrandsRootDeformed_DeformedUniqueTrianglePosition0Buffer1), ResourceName), OwnerName, EHairResourceUsageType::Dynamic);
 		}
 	}
+}
+
+void FHairStrandsLODDeformedRootResource::InternalRelease()
+{
+	Status = EStatus::Invalid;
+	AvailableCurveCount = 0;
+	DeformedUniqueTrianglePositionBuffer[0].Release();
+	DeformedSamplePositionsBuffer[0].Release();
+	MeshSampleWeightsBuffer[0].Release();
+
+	// Double buffering is disabled by default unless the read-only cvar r.HairStrands.ContinuousDecimationReordering is set
+	if (IsHairStrandContinuousDecimationReorderingEnabled())
+	{
+		DeformedUniqueTrianglePositionBuffer[1].Release();
+		DeformedSamplePositionsBuffer[1].Release();
+		MeshSampleWeightsBuffer[1].Release();
+	}
+}
+
+/////////////////////////////////////////////////////////////////////////////////////////
+
+FHairStrandsDeformedRootResource::FHairStrandsDeformedRootResource(FHairStrandsRestRootResource* InRestResources, EHairStrandsResourcesType InCurveType, const FHairResourceName& InResourceName, const FName& InOwnerName)
+:	CurveType(InCurveType)
+{
+	check(InRestResources);
+	LODs.Reserve(InRestResources->LODs.Num());
+	for (FHairStrandsLODRestRootResource* InLOD : InRestResources->LODs)
+	{
+		check(InLOD);
+		LODs.Add(new FHairStrandsLODDeformedRootResource(InRestResources, InCurveType, InLOD->MeshLODIndex, InResourceName, InOwnerName));
+	}
+}
+
+void FHairStrandsDeformedRootResource::Allocate(FRDGBuilder& GraphBuilder, EHairResourceLoadingType LoadingType, EHairResourceStatus& Status, int32 InMeshLODIndex)
+{
+	if (IsDataValid(InMeshLODIndex))
+	{
+		GetLOD(InMeshLODIndex)->Allocate(GraphBuilder, LoadingType, Status);
+	}
+}
+
+void FHairStrandsDeformedRootResource::Allocate(FRDGBuilder& GraphBuilder, EHairResourceLoadingType LoadingType, EHairResourceStatus& Status, uint32 InRequestedCurveCount, uint32 InRequestedPointCount, bool bAllowDeallocation, int32 InMeshLODIndex)
+{
+	if (IsDataValid(InMeshLODIndex))
+	{
+		GetLOD(InMeshLODIndex)->Allocate(GraphBuilder, LoadingType, Status, InRequestedCurveCount, InRequestedPointCount, bAllowDeallocation);
+	}
+}
+
+void FHairStrandsDeformedRootResource::ReleaseResource()
+{
+	for (FHairStrandsLODDeformedRootResource* LOD : LODs)
+	{
+		LOD->ReleaseResource();
+		delete LOD;
+	}
 	LODs.Empty();
 }
 
+void FHairStrandsDeformedRootResource::SwapBuffer()
+{
+	for (FHairStrandsLODDeformedRootResource* LOD : LODs)
+	{
+		LOD->SwapBuffer();
+	}
+}
+
 /////////////////////////////////////////////////////////////////////////////////////////
 
 FHairStrandsInterpolationResource::FHairStrandsInterpolationResource(FHairStrandsInterpolationBulkData& InBulkData, const FHairResourceName& InResourceName, const FName& InOwnerName) :
@@ -1566,11 +1610,11 @@ CurveInterpolationBuffer(), PointInterpolationBuffer(),BulkData(InBulkData)
 	check(!!(BulkData.Header.Flags & FHairStrandsInterpolationBulkData::DataFlags_HasData));
 }
 
-bool FHairStrandsInterpolationResource::InternalGetOrRequestData(uint32 InRequestedCurveCount, uint32 InRequestedPointCount, int32 InLODIndex)
+bool FHairStrandsInterpolationResource::InternalGetOrRequestData(uint32 InRequestedCurveCount, uint32 InRequestedPointCount)
 {
 	if (StreamingRequest.IsNone())
 	{
-		StreamingRequest.Request(InRequestedCurveCount, InRequestedPointCount, InLODIndex, BulkData, false, false, false, OwnerName);
+		StreamingRequest.Request(InRequestedCurveCount, InRequestedPointCount, BulkData, false, false, false, OwnerName);
 	}
 	return StreamingRequest.IsCompleted();
 }
@@ -1661,7 +1705,7 @@ FHairStrandsRaytracingResource::FHairStrandsRaytracingResource(const FHairMeshes
 	MaxVertexCount(InData.GetNumVertices()), bOwnBuffers(false)
 {}
 
-void FHairStrandsRaytracingResource::InternalAllocate(FRDGBuilder& GraphBuilder, uint32 InCurveCount, uint32 InPointCount, int32 InLODIndex)
+void FHairStrandsRaytracingResource::InternalAllocate(FRDGBuilder& GraphBuilder, uint32 InCurveCount, uint32 InPointCount)
 {
 	if (bOwnBuffers)
 	{
diff --git a/Engine/Plugins/Runtime/HairStrands/Source/HairStrandsCore/Private/HairStrandsDatas.cpp b/Engine/Plugins/Runtime/HairStrands/Source/HairStrandsCore/Private/HairStrandsDatas.cpp
--- a/Engine/Plugins/Runtime/HairStrands/Source/HairStrandsCore/Private/HairStrandsDatas.cpp
+++ b/Engine/Plugins/Runtime/HairStrands/Source/HairStrandsCore/Private/HairStrandsDatas.cpp
@@ -130,7 +130,7 @@ namespace HairStrands
 #endif
 }
 
-void FHairStrandsBulkCommon::Write_DDC(UObject* Owner, TArray<UE::DerivedData::FCachePutValueRequest>& Out, int32 InLODIndex)
+void FHairStrandsBulkCommon::Write_DDC(UObject* Owner, TArray<UE::DerivedData::FCachePutValueRequest>& Out)
 {
 #if WITH_EDITORONLY_DATA
 	FHairStrandsBulkCommon::FQuery Q;
@@ -138,7 +138,6 @@ void FHairStrandsBulkCommon::Write_DDC(UObject* Owner, TArray<UE::DerivedData::F
 	Q.OutWriteDDC = &Out;
 	Q.DerivedDataKey = &DerivedDataKey;
 	Q.Owner = Owner;
-	Q.LODIndex = InLODIndex;
 	GetResources(Q);
 #endif
 }
@@ -168,7 +167,7 @@ void FHairStrandsBulkCommon::Read_IO(FHairStreamingRequest* In, FBulkDataBatchRe
 	GetResources(Q);
 	Q.OutReadIO->Issue(Out);
 }
-void FHairStrandsBulkCommon::Write_IO(UObject* Owner, FArchive& Out, int32 InLODIndex)
+void FHairStrandsBulkCommon::Write_IO(UObject* Owner, FArchive& Out)
 {
 	GetResourceVersion(Out);
 
@@ -176,7 +175,6 @@ void FHairStrandsBulkCommon::Write_IO(UObject* Owner, FArchive& Out, int32 InLOD
 	Q.Type = FHairStrandsBulkCommon::FQuery::ReadWriteIO;
 	Q.OutWriteIO = &Out;
 	Q.Owner = Owner;
-	Q.LODIndex = InLODIndex;
 	GetResources(Q);
 }
 void FHairStrandsBulkCommon::Unload(FHairStreamingRequest* In)
@@ -191,12 +189,12 @@ void FHairStrandsBulkCommon::Unload(FHairStreamingRequest* In)
 void FHairStrandsBulkCommon::Serialize(FArchive& Ar, UObject* Owner)
 {
 	SerializeHeader(Ar, Owner);
-	SerializeData(Ar, Owner, -1 /*InLODIndex*/);
+	SerializeData(Ar, Owner);
 }
 
-void FHairStrandsBulkCommon::SerializeData(FArchive& Ar, UObject* Owner, int32 InLODIndex)
+void FHairStrandsBulkCommon::SerializeData(FArchive& Ar, UObject* Owner)
 {
-	Write_IO(Owner, Ar, InLODIndex);
+	Write_IO(Owner, Ar);
 }
 
 void FHairStrandsBulkCommon::FQuery::Add(FHairBulkContainer& In, const TCHAR* InSuffix, uint32& InOffset, uint32 InSize) 
@@ -211,13 +209,10 @@ void FHairStrandsBulkCommon::FQuery::Add(FHairBulkContainer& In, const TCHAR* In
 		FMemory::Memcpy(WriteData.GetData(), In.Data.Lock(LOCK_READ_ONLY), DataSizeInByte);
 		In.Data.Unlock();
 
-		// Optional LOD suffix
-		const FString InLODSuffix = GetLODIndex() >= 0 ? FString(TEXT("_LOD")) + FString::FromInt(GetLODIndex()) : FString();
-
 		using namespace UE::DerivedData;
 		FCachePutValueRequest& Out = OutWriteDDC->AddDefaulted_GetRef();
 		if (Owner) { Out.Name = Owner->GetPathName(); }
-		Out.Key 	= ConvertLegacyCacheKey(*DerivedDataKey + InSuffix + InLODSuffix);
+		Out.Key 	= ConvertLegacyCacheKey(*DerivedDataKey + InSuffix);
 		Out.Value 	= FValue::Compress(MakeSharedBufferFromArray(MoveTemp(WriteData)));
 		Out.Policy 	= ECachePolicy::Default;
 		Out.UserData= 0;
@@ -238,15 +233,12 @@ void FHairStrandsBulkCommon::FQuery::Add(FHairBulkContainer& In, const TCHAR* In
 		Chunk.TotalSize = InSize;
 		In.ChunkRequest = &Chunk;
 
-		// Optional LOD suffix
-		const FString InLODSuffix = GetLODIndex() >= 0 ? FString(TEXT("_LOD")) + FString::FromInt(GetLODIndex()) : FString();
-
 		// 2. Fill in actual DDC request
 		check(OutReadDDC);
 		using namespace UE::DerivedData;
 		FCacheGetChunkRequest& Out = OutReadDDC->AddDefaulted_GetRef();
 		Out.Id			= FValueId::Null; 	// HairStrands::HairStrandsValueId : This is only needed for cache record, not cache value.
-		Out.Key			= ConvertLegacyCacheKey(*DerivedDataKey + InSuffix + InLODSuffix);
+		Out.Key			= ConvertLegacyCacheKey(*DerivedDataKey + InSuffix);
 		Out.RawOffset	= InOffset;
 		Out.RawSize		= InSize != 0 ? InSize-InOffset : MAX_uint64;
 		Out.RawHash		= FIoHash();
@@ -378,18 +370,16 @@ static bool RequestWarmCache(UE::DerivedData::FRequestOwner* RequestOwner, const
 	return bHasDataInCache;
 }
 
-bool FHairStreamingRequest::WarmCache(uint32 InRequestedCurveCount, uint32 InRequestedPointCount, int32 InLODIndex, FHairStrandsBulkCommon& In)
+bool FHairStreamingRequest::WarmCache(uint32 InRequestedCurveCount, uint32 InRequestedPointCount, FHairStrandsBulkCommon& In)
 {
 	if (In.GetResourceCount() == 0 || InRequestedCurveCount == 0)
 	{
 		CurveCount = 0;
 		PointCount = 0;
-		LODIndex = -1;
 		return true;
 	}
 	CurveCount = InRequestedCurveCount;
 	PointCount = InRequestedPointCount;
-	LODIndex = InLODIndex;
 	bSupportOffsetLoad = false; // Whole resource loading/caching
 
 	using namespace UE::DerivedData;
@@ -405,7 +395,7 @@ bool FHairStreamingRequest::WarmCache(uint32 InRequestedCurveCount, uint32 InReq
 // Request fullfil 2 use cases:
 // * Load IO/DDC data and upload them to GPU
 // * Load DDC data and store them into bulkdata for serialization
-void FHairStreamingRequest::Request(uint32 InRequestedCurveCount, uint32 InRequestedPointCount, int32 InLODIndex, FHairStrandsBulkCommon& In,
+void FHairStreamingRequest::Request(uint32 InRequestedCurveCount, uint32 InRequestedPointCount, FHairStrandsBulkCommon& In,
 	bool bWait, bool bFillBulkdata, bool bWarmCache, const FName& InOwnerName, bool* bWaitResult)
 {
 	bool bLogErrors = bWaitResult == nullptr;
@@ -422,12 +412,10 @@ void FHairStreamingRequest::Request(uint32 InRequestedCurveCount, uint32 InReque
 	{
 		CurveCount = 0;
 		PointCount = 0;
-		LODIndex = -1;
 		return;
 	}
 	CurveCount = InRequestedCurveCount;
 	PointCount = InRequestedPointCount;
-	LODIndex = InLODIndex;
 
 #if !WITH_EDITORONLY_DATA
 	{ 
@@ -853,28 +841,15 @@ uint32 FHairStrandsClusterBulkData::GetCurveCount(float InLODIndex) const
 
 bool FHairStrandsRootData::HasProjectionData() const
 {
-	bool bIsValid = MeshProjectionLODs.Num() > 0;
-	for (const FMeshProjectionLOD& LOD : MeshProjectionLODs)
-	{
-		const bool bHasValidCPUData = LOD.RootBarycentricBuffer.Num() > 0;
-		if (bHasValidCPUData)
-		{
-			bIsValid = bIsValid && LOD.RootBarycentricBuffer.Num() > 0;
-			bIsValid = bIsValid && LOD.RootToUniqueTriangleIndexBuffer.Num() > 0;
-			bIsValid = bIsValid && LOD.UniqueTriangleIndexBuffer.Num() > 0;
-			bIsValid = bIsValid && LOD.RestUniqueTrianglePositionBuffer.Num() > 0;
-
-			if (!bIsValid) break;
-		}
-	}
-
-	return bIsValid;
+	return RootBarycentricBuffer.Num() > 0
+	&& RootToUniqueTriangleIndexBuffer.Num() > 0
+	&& UniqueTriangleIndexBuffer.Num() > 0
+	&& RestUniqueTrianglePositionBuffer.Num() > 0;
 }
 
 void FHairStrandsRootData::Reset()
 {
-	RootCount = 0;
-	MeshProjectionLODs.Empty();
+	*this = FHairStrandsRootData();
 }
 
 /////////////////////////////////////////////////////////////////////////////////////////
@@ -882,8 +857,16 @@ void FHairStrandsRootData::Reset()
 
 void FHairStrandsRootBulkData::SerializeHeader(FArchive& Ar, UObject* Owner)
 {
-	Ar << Header.RootCount;
-	Ar << Header.PointCount;
+	// Header
+	{
+		Ar << Header.RootCount;
+		Ar << Header.PointCount;
+		Ar << Header.LODIndex;
+		Ar << Header.UniqueTriangleCount;
+		Ar << Header.SampleCount;
+		Ar << Header.UniqueSectionIndices;
+		Ar << Header.MeshSectionCount;
+	}
 
 	// Strides
 	{
@@ -896,33 +879,11 @@ void FHairStrandsRootBulkData::SerializeHeader(FArchive& Ar, UObject* Owner)
 		Ar << Header.Strides.MeshSampleIndicesAndSectionsBufferStride;
 		Ar << Header.Strides.RestSamplePositionsBufferStride;
 	}
-
-	// LODS
-	uint32 LODCount = Header.LODs.Num();
-	Ar << LODCount;
-	if (Ar.IsLoading())
-	{
-		Header.LODs.SetNum(LODCount);
-		Data.LODs.SetNum(LODCount);
-	}
-	for (uint32 LODIt = 0; LODIt < LODCount; ++LODIt)
-	{
-		if (Ar.IsLoading())
-		{
-			Header.LODs[LODIt].LODIndex = LODIt;
-		}
-
-		Ar << Header.LODs[LODIt].UniqueTriangleCount;
-		Ar << Header.LODs[LODIt].SampleCount;
-		Ar << Header.LODs[LODIt].UniqueSectionIndices;
-		Ar << Header.LODs[LODIt].MeshSectionCount;
-	}
 }
 
-const TArray<uint32>& FHairStrandsRootBulkData::GetValidSectionIndices(int32 LODIndex) const
+const TArray<uint32>& FHairStrandsRootBulkData::GetValidSectionIndices() const
 {
-	check(LODIndex >= 0 && LODIndex < Header.LODs.Num());
-	return Header.LODs[LODIndex].UniqueSectionIndices;
+	return Header.UniqueSectionIndices;
 }
 
 uint32 FHairStrandsRootBulkData::GetResourceCount() const
@@ -932,14 +893,6 @@ uint32 FHairStrandsRootBulkData::GetResourceCount() const
 
 void FHairStrandsRootBulkData::GetResources(FQuery& Out)
 {
-	const int32 LODIndex = Out.GetLODIndex();
-	check(LODIndex >= 0);
-	check(Header.LODs.Num() == Data.LODs.Num());
-	check(LODIndex < Data.LODs.Num());
-
-	const FHeader::FLOD& HeaderLOD = Header.LODs[LODIndex];
-	FData::FLOD& DataLOD = Data.LODs[LODIndex];
-
 	// Translate requested curve count into chunk/offset/size to be read
 	// * RBF points are loaded all at once
 	// * Unique triangle data are loaded all at one
@@ -950,85 +903,83 @@ void FHairStrandsRootBulkData::GetResources(FQuery& Out)
 		RootCount = FMath::Min(Header.RootCount, Out.GetCurveCount());
 	}
 
-	Out.Add(DataLOD.RootToUniqueTriangleIndexBuffer, 	TEXT("_RootToUniqueTriangleIndexBuffer"), 		DataLOD.RootToUniqueTriangleIndexBuffer.LoadedSize, 	RootCount * Header.Strides.RootToUniqueTriangleIndexBufferStride);
-	Out.Add(DataLOD.RootBarycentricBuffer, 				TEXT("_RootBarycentricBuffer"), 				DataLOD.RootBarycentricBuffer.LoadedSize, 				RootCount * Header.Strides.RootBarycentricBufferStride);
-	Out.Add(DataLOD.UniqueTriangleIndexBuffer, 			TEXT("_UniqueTriangleIndexBuffer"), 			DataLOD.UniqueTriangleIndexBuffer.LoadedSize, 			HeaderLOD.UniqueTriangleCount * Header.Strides.UniqueTriangleIndexBufferStride);		 // Load all data
-	Out.Add(DataLOD.RestUniqueTrianglePositionBuffer, 	TEXT("_RestUniqueTrianglePositionBuffer"), 		DataLOD.RestUniqueTrianglePositionBuffer.LoadedSize, 	HeaderLOD.UniqueTriangleCount * Header.Strides.RestUniqueTrianglePositionBufferStride);	 // Load all data
+	Out.Add(Data.RootToUniqueTriangleIndexBuffer, 	TEXT("_RootToUniqueTriangleIndexBuffer"), 		Data.RootToUniqueTriangleIndexBuffer.LoadedSize, 	RootCount * Header.Strides.RootToUniqueTriangleIndexBufferStride);
+	Out.Add(Data.RootBarycentricBuffer, 			TEXT("_RootBarycentricBuffer"), 				Data.RootBarycentricBuffer.LoadedSize, 				RootCount * Header.Strides.RootBarycentricBufferStride);
+	Out.Add(Data.UniqueTriangleIndexBuffer, 		TEXT("_UniqueTriangleIndexBuffer"), 			Data.UniqueTriangleIndexBuffer.LoadedSize, 			Header.UniqueTriangleCount * Header.Strides.UniqueTriangleIndexBufferStride);		 // Load all data
+	Out.Add(Data.RestUniqueTrianglePositionBuffer, 	TEXT("_RestUniqueTrianglePositionBuffer"), 		Data.RestUniqueTrianglePositionBuffer.LoadedSize, 	Header.UniqueTriangleCount * Header.Strides.RestUniqueTrianglePositionBufferStride); // Load all data
 
-	if (HeaderLOD.SampleCount > 0)
+	if (Header.SampleCount > 0)
 	{
-		const uint32 RBFWeightCount = FGroomRBFDeformer::GetWeightCount(HeaderLOD.SampleCount); 
+		const uint32 RBFWeightCount = FGroomRBFDeformer::GetWeightCount(Header.SampleCount); 
 
-		Out.Add(DataLOD.MeshInterpolationWeightsBuffer, 	TEXT("_MeshInterpolationWeightsBuffer"), 	DataLOD.MeshInterpolationWeightsBuffer.LoadedSize, 		RBFWeightCount * Header.Strides.MeshInterpolationWeightsBufferStride); 				// Load all data
-		Out.Add(DataLOD.MeshSampleIndicesAndSectionsBuffer, TEXT("_MeshSampleIndicesAndSectionsBuffer"),DataLOD.MeshSampleIndicesAndSectionsBuffer.LoadedSize, 	HeaderLOD.SampleCount * Header.Strides.MeshSampleIndicesAndSectionsBufferStride); 	// Load all data
-		Out.Add(DataLOD.RestSamplePositionsBuffer, 			TEXT("_RestSamplePositionsBuffer"), 		DataLOD.RestSamplePositionsBuffer.LoadedSize, 			HeaderLOD.SampleCount * Header.Strides.RestSamplePositionsBufferStride);			// Load all data
+		Out.Add(Data.MeshInterpolationWeightsBuffer, 	TEXT("_MeshInterpolationWeightsBuffer"), 	Data.MeshInterpolationWeightsBuffer.LoadedSize, 		RBFWeightCount * Header.Strides.MeshInterpolationWeightsBufferStride); 			// Load all data
+		Out.Add(Data.MeshSampleIndicesAndSectionsBuffer,TEXT("_MeshSampleIndicesAndSectionsBuffer"),Data.MeshSampleIndicesAndSectionsBuffer.LoadedSize, 	Header.SampleCount * Header.Strides.MeshSampleIndicesAndSectionsBufferStride); 	// Load all data
+		Out.Add(Data.RestSamplePositionsBuffer, 		TEXT("_RestSamplePositionsBuffer"), 		Data.RestSamplePositionsBuffer.LoadedSize, 				Header.SampleCount * Header.Strides.RestSamplePositionsBufferStride);			// Load all data
 	}
 }
 
-uint32 FHairStrandsRootBulkData::GetSize() const
+uint32 FHairStrandsRootBulkData::GetDataSize() const
 {
 	uint32 Out = 0;
-	for (const auto& LOD : Data.LODs)
-	{
-		Out += LOD.RootToUniqueTriangleIndexBuffer.GetBulkDataSize();
-		Out += LOD.RootBarycentricBuffer.GetBulkDataSize();
-		Out += LOD.UniqueTriangleIndexBuffer.GetBulkDataSize();
-		Out += LOD.RestUniqueTrianglePositionBuffer.GetBulkDataSize();
-
-		Out += LOD.MeshInterpolationWeightsBuffer.GetBulkDataSize();
-		Out += LOD.MeshSampleIndicesAndSectionsBuffer.GetBulkDataSize();
-		Out += LOD.RestSamplePositionsBuffer.GetBulkDataSize();
-	}
+	Out += Data.RootToUniqueTriangleIndexBuffer.GetBulkDataSize();
+	Out += Data.RootBarycentricBuffer.GetBulkDataSize();
+	Out += Data.UniqueTriangleIndexBuffer.GetBulkDataSize();
+	Out += Data.RestUniqueTrianglePositionBuffer.GetBulkDataSize();
+
+	Out += Data.MeshInterpolationWeightsBuffer.GetBulkDataSize();
+	Out += Data.MeshSampleIndicesAndSectionsBuffer.GetBulkDataSize();
+	Out += Data.RestSamplePositionsBuffer.GetBulkDataSize();
 	return Out;
 }
 
 void FHairStrandsRootBulkData::Reset()
 {
 	// Header
-	Header.RootCount = 0;
-	Header.PointCount = 0;
-	Header.LODs.Empty();
+	{
+		Header.LODIndex = -1;
+		Header.RootCount = 0;
+		Header.PointCount = 0;
+		Header.SampleCount = 0;
+		Header.MeshSectionCount = 0;
+		Header.UniqueTriangleCount = 0;
+		Header.UniqueSectionIndices.Empty();
+	}
 
 	// Data
-	for (FData::FLOD& LOD : Data.LODs)
 	{
 		// Binding
-		LOD.RootBarycentricBuffer.RemoveBulkData();
-		LOD.RootToUniqueTriangleIndexBuffer.RemoveBulkData();
-		LOD.UniqueTriangleIndexBuffer.RemoveBulkData();
-		LOD.RestUniqueTrianglePositionBuffer.RemoveBulkData();
+		Data.RootBarycentricBuffer.RemoveBulkData();
+		Data.RootToUniqueTriangleIndexBuffer.RemoveBulkData();
+		Data.UniqueTriangleIndexBuffer.RemoveBulkData();
+		Data.RestUniqueTrianglePositionBuffer.RemoveBulkData();
 
 		// Reset the bulk byte buffer to ensure the (serialize) data size is reset to 0
-		LOD.RootBarycentricBuffer 			 = FHairBulkContainer();
-		LOD.RootToUniqueTriangleIndexBuffer  = FHairBulkContainer();
-		LOD.UniqueTriangleIndexBuffer		 = FHairBulkContainer();
-		LOD.RestUniqueTrianglePositionBuffer = FHairBulkContainer();
+		Data.RootBarycentricBuffer 			  = FHairBulkContainer();
+		Data.RootToUniqueTriangleIndexBuffer  = FHairBulkContainer();
+		Data.UniqueTriangleIndexBuffer		  = FHairBulkContainer();
+		Data.RestUniqueTrianglePositionBuffer = FHairBulkContainer();
 
 		// RBF
-		LOD.MeshInterpolationWeightsBuffer.RemoveBulkData();
-		LOD.MeshSampleIndicesAndSectionsBuffer.RemoveBulkData();
-		LOD.RestSamplePositionsBuffer.RemoveBulkData();
+		Data.MeshInterpolationWeightsBuffer.RemoveBulkData();
+		Data.MeshSampleIndicesAndSectionsBuffer.RemoveBulkData();
+		Data.RestSamplePositionsBuffer.RemoveBulkData();
 
-		LOD.MeshInterpolationWeightsBuffer 		 = FHairBulkContainer();
-		LOD.MeshSampleIndicesAndSectionsBuffer 	 = FHairBulkContainer();
-		LOD.RestSamplePositionsBuffer 			 = FHairBulkContainer();
+		Data.MeshInterpolationWeightsBuffer 	 = FHairBulkContainer();
+		Data.MeshSampleIndicesAndSectionsBuffer  = FHairBulkContainer();
+		Data.RestSamplePositionsBuffer 			 = FHairBulkContainer();
 	}
-	Data.LODs.Empty();
 }
 
 void FHairStrandsRootBulkData::ResetLoadedSize()
 {
-	for (FData::FLOD& LOD : Data.LODs)
-	{
-		// Binding
-		LOD.RootBarycentricBuffer.LoadedSize 			= 0;
-		LOD.RootToUniqueTriangleIndexBuffer.LoadedSize 	= 0;
-		LOD.UniqueTriangleIndexBuffer.LoadedSize 		= 0;
-		LOD.RestUniqueTrianglePositionBuffer.LoadedSize = 0;
-		      
-		// RBF
-		LOD.MeshInterpolationWeightsBuffer.LoadedSize	 = 0;
-		LOD.MeshSampleIndicesAndSectionsBuffer.LoadedSize= 0;
-		LOD.RestSamplePositionsBuffer.LoadedSize 		 = 0;
-	}
+	// Binding
+	Data.RootBarycentricBuffer.LoadedSize 			= 0;
+	Data.RootToUniqueTriangleIndexBuffer.LoadedSize = 0;
+	Data.UniqueTriangleIndexBuffer.LoadedSize 		= 0;
+	Data.RestUniqueTrianglePositionBuffer.LoadedSize= 0;
+		
+	// RBF
+	Data.MeshInterpolationWeightsBuffer.LoadedSize	  = 0;
+	Data.MeshSampleIndicesAndSectionsBuffer.LoadedSize= 0;
+	Data.RestSamplePositionsBuffer.LoadedSize 		  = 0;
 }
diff --git a/Engine/Plugins/Runtime/HairStrands/Source/HairStrandsCore/Private/HairStrandsInterpolation.cpp b/Engine/Plugins/Runtime/HairStrands/Source/HairStrandsCore/Private/HairStrandsInterpolation.cpp
--- a/Engine/Plugins/Runtime/HairStrands/Source/HairStrandsCore/Private/HairStrandsInterpolation.cpp
+++ b/Engine/Plugins/Runtime/HairStrands/Source/HairStrandsCore/Private/HairStrandsInterpolation.cpp
@@ -383,17 +383,15 @@ void AddDeformSimHairStrandsPass(
 			const bool bSupportDynamicMesh = 
 				RootCount > 0 && 
 				MeshLODIndex >= 0 && 
-				MeshLODIndex < uint32(SimRestRootResources->LODs.Num()) &&
-				MeshLODIndex < uint32(SimDeformedRootResources->LODs.Num()) &&
-				SimRestRootResources->LODs[MeshLODIndex].IsValid() &&
-				SimDeformedRootResources->LODs[MeshLODIndex].IsValid() &&
+				SimRestRootResources->IsValid(MeshLODIndex) &&
+				SimDeformedRootResources->IsValid(MeshLODIndex) &&
 				bIsPointToCurveBuffersValid;
 			
 			bool bSupportGlobalInterpolation = false;
 			if (bSupportDynamicMesh)
 			{
-				FHairStrandsRestRootResource::FLOD& RestLODDatas = SimRestRootResources->LODs[MeshLODIndex];
-				FHairStrandsDeformedRootResource::FLOD& DeformedLODDatas = SimDeformedRootResources->LODs[MeshLODIndex];
+				FHairStrandsLODRestRootResource& RestLODDatas = *SimRestRootResources->GetLOD(MeshLODIndex);
+				FHairStrandsLODDeformedRootResource& DeformedLODDatas = *SimDeformedRootResources->GetLOD(MeshLODIndex);
 
 				bSupportGlobalInterpolation = bHasGlobalInterpolation && (RestLODDatas.SampleCount > 0);
 				if (!bSupportGlobalInterpolation) 
@@ -401,13 +399,13 @@ void AddDeformSimHairStrandsPass(
 					InternalDeformationType = InternalDeformationType_Skinned;
 					Parameters->SimRootToUniqueTriangleIndexBuffer = RegisterAsSRV(GraphBuilder, RestLODDatas.RootToUniqueTriangleIndexBuffer);
 					Parameters->SimRestPositionBuffer = RegisterAsSRV(GraphBuilder, RestLODDatas.RestUniqueTrianglePositionBuffer);
-					Parameters->SimDeformedPositionBuffer = RegisterAsSRV(GraphBuilder, DeformedLODDatas.GetDeformedUniqueTrianglePositionBuffer(FHairStrandsDeformedRootResource::FLOD::Current));
+					Parameters->SimDeformedPositionBuffer = RegisterAsSRV(GraphBuilder, DeformedLODDatas.GetDeformedUniqueTrianglePositionBuffer(FHairStrandsLODDeformedRootResource::Current));
 					Parameters->SimRootBarycentricBuffer = RegisterAsSRV(GraphBuilder, RestLODDatas.RootBarycentricBuffer);
 				}
 				else
 				{
 					InternalDeformationType = InternalDeformationType_RBF;
-					Parameters->MeshSampleWeightsBuffer = RegisterAsSRV(GraphBuilder, DeformedLODDatas.GetMeshSampleWeightsBuffer(FHairStrandsDeformedRootResource::FLOD::Current));
+					Parameters->MeshSampleWeightsBuffer = RegisterAsSRV(GraphBuilder, DeformedLODDatas.GetMeshSampleWeightsBuffer(FHairStrandsLODDeformedRootResource::Current));
 					Parameters->RestSamplePositionsBuffer = RegisterAsSRV(GraphBuilder, RestLODDatas.RestSamplePositionsBuffer);
 					Parameters->SampleCount = RestLODDatas.SampleCount;
 				}
@@ -558,7 +556,7 @@ void AddHairStrandsInterpolationPass(
 	const FRDGBufferSRVRef& SimDeformedPositionBuffer,
 	const FRDGBufferSRVRef& RenDeformerPositionBuffer,
 	FRDGBufferUAVRef& OutRenPositionBuffer,
-	const FHairStrandsDeformedRootResource::FLOD::EFrameType DeformedFrame)
+	const FHairStrandsLODDeformedRootResource::EFrameType DeformedFrame)
 {
 	FHairInterpolationCS::FParameters* Parameters = GraphBuilder.AllocParameters<FHairInterpolationCS::FParameters>();
 	Parameters->RenRestPosePositionBuffer = RenRestPosePositionBuffer;
@@ -591,17 +589,15 @@ void AddHairStrandsInterpolationPass(
 		RenRestRootResources &&
 		RenRestRootResources->GetRootCount() > 0 && 
 		MeshLODIndex >= 0 && 
-		MeshLODIndex < RenRestRootResources->LODs.Num() &&
-		MeshLODIndex < RenDeformedRootResources->LODs.Num() &&
-		RenRestRootResources->LODs[MeshLODIndex].IsValid() &&
-		RenDeformedRootResources->LODs[MeshLODIndex].IsValid();
+		RenRestRootResources->IsValid(MeshLODIndex) &&
+		RenDeformedRootResources->IsValid(MeshLODIndex);
 
 	// Guides
 	bool bSupportGlobalInterpolation = false;
 	if (bSupportDynamicMesh && (Instance->Guides.bIsSimulationEnable || Instance->Guides.bHasGlobalInterpolation || Instance->Guides.bIsDeformationEnable)) // No need for Instance->Guides.bHasSimulationCache since it is incompatible with binding
 	{
-		const FHairStrandsRestRootResource::FLOD& Sim_RestLODDatas = SimRestRootResources->LODs[MeshLODIndex];
-		const FHairStrandsDeformedRootResource::FLOD& Sim_DeformedLODDatas = SimDeformedRootResources->LODs[MeshLODIndex];
+		const FHairStrandsLODRestRootResource& Sim_RestLODDatas = *SimRestRootResources->GetLOD(MeshLODIndex);
+		const FHairStrandsLODDeformedRootResource& Sim_DeformedLODDatas = *SimDeformedRootResources->GetLOD(MeshLODIndex);
 		bSupportGlobalInterpolation = Instance->Guides.bHasGlobalInterpolation && (Sim_RestLODDatas.SampleCount > 0);
 		{
 			Parameters->SimRootRestPositionBuffer = RegisterAsSRV(GraphBuilder, Sim_RestLODDatas.RestUniqueTrianglePositionBuffer);
@@ -616,8 +612,8 @@ void AddHairStrandsInterpolationPass(
 	// Strands
 	if (bSupportDynamicMesh)
 	{
-		const FHairStrandsRestRootResource::FLOD& Ren_RestLODDatas = RenRestRootResources->LODs[MeshLODIndex];
-		const FHairStrandsDeformedRootResource::FLOD& Ren_DeformedLODDatas = RenDeformedRootResources->LODs[MeshLODIndex];
+		const FHairStrandsLODRestRootResource& Ren_RestLODDatas = *RenRestRootResources->GetLOD(MeshLODIndex);
+		const FHairStrandsLODDeformedRootResource& Ren_DeformedLODDatas = *RenDeformedRootResources->GetLOD(MeshLODIndex);
 		{
 			Parameters->RenRootRestPositionBuffer = RegisterAsSRV(GraphBuilder, Ren_RestLODDatas.RestUniqueTrianglePositionBuffer);
 			Parameters->RenRootBarycentricBuffer = RegisterAsSRV(GraphBuilder, Ren_RestLODDatas.RootBarycentricBuffer);
@@ -927,9 +923,9 @@ void AddHairCardsDeformationPass(
 	const ERHIFeatureLevel::Type FeatureLevel,
 	const FShaderPrintData* ShaderPrintData,
 	FHairGroupInstance* Instance,
+	const int32 HairLODIndex,
 	const int32 MeshLODIndex)
 {
-	const int32 HairLODIndex = Instance->HairGroupPublicData->GetIntLODIndex();
 	if (!Instance->Cards.IsValid(HairLODIndex))
 		return;
 
@@ -978,20 +974,18 @@ void AddHairCardsDeformationPass(
 		RestRootResources &&
 		RestRootResources->GetRootCount() > 0 &&
 		MeshLODIndex >= 0 &&
-		MeshLODIndex < RestRootResources->LODs.Num() &&
-		MeshLODIndex < DeformedRootResources->LODs.Num() &&
-		RestRootResources->LODs[MeshLODIndex].IsValid() &&
-		DeformedRootResources->LODs[MeshLODIndex].IsValid();
+		RestRootResources->IsValid(MeshLODIndex) &&
+		DeformedRootResources->IsValid(MeshLODIndex);
 	if (bSupportDynamicMesh)
 	{
 		Parameters->GuideVertexToRootIndexBuffer = RegisterAsSRV(GraphBuilder, RestResources->PointToCurveBuffer);
-		const FHairStrandsRestRootResource::FLOD& RestLODDatas = RestRootResources->LODs[MeshLODIndex];
-		const FHairStrandsDeformedRootResource::FLOD& DeformedLODDatas = DeformedRootResources->LODs[MeshLODIndex];
+		const FHairStrandsLODRestRootResource& RestLODDatas = *RestRootResources->GetLOD(MeshLODIndex);
+		const FHairStrandsLODDeformedRootResource& DeformedLODDatas = *DeformedRootResources->GetLOD(MeshLODIndex);
 
 		Parameters->GuideRootBarycentricBuffer = RegisterAsSRV(GraphBuilder, RestLODDatas.RootBarycentricBuffer);
 		Parameters->GuideRootToUniqueTriangleIndexBuffer = RegisterAsSRV(GraphBuilder, RestLODDatas.RootToUniqueTriangleIndexBuffer);
 		Parameters->TriangleRestPositionBuffer = RegisterAsSRV(GraphBuilder, RestLODDatas.RestUniqueTrianglePositionBuffer);
-		Parameters->TriangleDeformedPositionBuffer = RegisterAsSRV(GraphBuilder, DeformedLODDatas.GetDeformedUniqueTrianglePositionBuffer(FHairStrandsDeformedRootResource::FLOD::Current));
+		Parameters->TriangleDeformedPositionBuffer = RegisterAsSRV(GraphBuilder, DeformedLODDatas.GetDeformedUniqueTrianglePositionBuffer(FHairStrandsLODDeformedRootResource::Current));
 	}
 
 	if (ShaderPrintData)
diff --git a/Engine/Plugins/Runtime/HairStrands/Source/HairStrandsCore/Private/HairStrandsMeshProjection.cpp b/Engine/Plugins/Runtime/HairStrands/Source/HairStrandsCore/Private/HairStrandsMeshProjection.cpp
--- a/Engine/Plugins/Runtime/HairStrands/Source/HairStrandsCore/Private/HairStrandsMeshProjection.cpp
+++ b/Engine/Plugins/Runtime/HairStrands/Source/HairStrandsCore/Private/HairStrandsMeshProjection.cpp
@@ -305,27 +305,26 @@ void AddHairStrandUpdateMeshTrianglesPass(
 	FHairStrandsRestRootResource* RestResources,
 	FHairStrandsDeformedRootResource* DeformedResources)
 {	
-	if (RestResources->GetRootCount() == 0 || !RestResources->LODs.IsValidIndex(MeshLODIndex))
+	if (RestResources->GetRootCount() == 0 || !RestResources->IsValid(MeshLODIndex))
 	{
 		return;
 	}
 
-	FHairStrandsRestRootResource::FLOD& RestLODData = RestResources->LODs[MeshLODIndex];
-	check(RestLODData.LODIndex == MeshLODIndex);
-	check(RestResources->BulkData.Header.LODs.IsValidIndex(MeshLODIndex));
-	const uint32 UniqueTriangleCount = RestResources->BulkData.Header.LODs[MeshLODIndex].UniqueTriangleCount;
-	const uint32 TotalMeshSectionCount = RestResources->BulkData.Header.LODs[MeshLODIndex].MeshSectionCount;
+	FHairStrandsLODRestRootResource& RestLODData = *RestResources->GetLOD(MeshLODIndex);
+	check(RestLODData.BulkData.Header.LODIndex == MeshLODIndex);
+	const uint32 UniqueTriangleCount = RestLODData.BulkData.Header.UniqueTriangleCount;
+	const uint32 TotalMeshSectionCount = RestLODData.BulkData.Header.MeshSectionCount;
 	if (UniqueTriangleCount == 0)
 	{
 		return;
 	}
 
-	FRDGImportedBuffer OutputCurrBuffer = Register(GraphBuilder, DeformedResources->LODs[MeshLODIndex].GetDeformedUniqueTrianglePositionBuffer(FHairStrandsDeformedRootResource::FLOD::Current), ERDGImportedBufferFlags::CreateUAV, ERDGUnorderedAccessViewFlags::SkipBarrier);
+	FRDGImportedBuffer OutputCurrBuffer = Register(GraphBuilder, DeformedResources->GetLOD(MeshLODIndex)->GetDeformedUniqueTrianglePositionBuffer(FHairStrandsLODDeformedRootResource::Current), ERDGImportedBufferFlags::CreateUAV, ERDGUnorderedAccessViewFlags::SkipBarrier);
 	FRDGImportedBuffer OutputPrevBuffer;
 	const bool bComputePreviousDeformedPosition = IsHairStrandContinuousDecimationReorderingEnabled();
 	if (bComputePreviousDeformedPosition)
 	{
-		Register(GraphBuilder, DeformedResources->LODs[MeshLODIndex].GetDeformedUniqueTrianglePositionBuffer(FHairStrandsDeformedRootResource::FLOD::Previous), ERDGImportedBufferFlags::CreateUAV, ERDGUnorderedAccessViewFlags::SkipBarrier);
+		Register(GraphBuilder, DeformedResources->GetLOD(MeshLODIndex)->GetDeformedUniqueTrianglePositionBuffer(FHairStrandsLODDeformedRootResource::Previous), ERDGImportedBufferFlags::CreateUAV, ERDGUnorderedAccessViewFlags::SkipBarrier);
 	}
 
 	if (AddHairStrandUpdateMeshTrianglesPass(
@@ -346,8 +345,7 @@ void AddHairStrandUpdateMeshTrianglesPass(
 		}
 	
 		// Update the last known mesh LOD for which the root resources has been updated
-		DeformedResources->MeshLODIndex = MeshLODIndex;
-		DeformedResources->LODs[MeshLODIndex].Status = FHairStrandsDeformedRootResource::FLOD::EStatus::Completed;	
+		DeformedResources->GetLOD(MeshLODIndex)->Status = FHairStrandsLODDeformedRootResource::EStatus::Completed;
 	}
 }
 
@@ -405,8 +403,8 @@ void InternalAddHairRBFInterpolationPass(
 	FRDGImportedBuffer DeformedPositionBuffer_Prev = Register(GraphBuilder, DeformedResources->GetBuffer(TDeformedResource::EFrameType::Previous), ERDGImportedBufferFlags::CreateUAV);
 	AddCopyBufferPass(GraphBuilder, DeformedPositionBuffer_Prev.Buffer, DeformedPositionBuffer_Curr.Buffer);
 
-	FHairStrandsRestRootResource::FLOD& RestLODData = RestRootResources->LODs[MeshLODIndex];
-	FHairStrandsDeformedRootResource::FLOD& DeformedLODData = DeformedRootResources->LODs[MeshLODIndex];
+	FHairStrandsLODRestRootResource& RestLODData = *RestRootResources->GetLOD(MeshLODIndex);
+	FHairStrandsLODDeformedRootResource& DeformedLODData = *DeformedRootResources->GetLOD(MeshLODIndex);
 	FHairMeshesInterpolateCS::FParameters* Parameters = GraphBuilder.AllocParameters<FHairMeshesInterpolateCS::FParameters>();
 	Parameters->VertexCount = VertexCount;
 	Parameters->MaxSampleCount = RestLODData.SampleCount;
@@ -415,7 +413,7 @@ void InternalAddHairRBFInterpolationPass(
 	Parameters->OutDeformedPositionBuffer	= DeformedPositionBuffer_Curr.UAV;
 
 	Parameters->RestSamplePositionsBuffer	= RegisterAsSRV(GraphBuilder, RestLODData.RestSamplePositionsBuffer);
-	Parameters->MeshSampleWeightsBuffer		= RegisterAsSRV(GraphBuilder, DeformedLODData.GetMeshSampleWeightsBuffer(FHairStrandsDeformedRootResource::FLOD::Current));
+	Parameters->MeshSampleWeightsBuffer		= RegisterAsSRV(GraphBuilder, DeformedLODData.GetMeshSampleWeightsBuffer(FHairStrandsLODDeformedRootResource::Current));
 
 	const FIntVector DispatchGroupCount = FComputeShaderUtils::GetGroupCount(VertexCount, 128);
 	check(DispatchGroupCount.X <= GRHIMaxDispatchThreadGroupsPerDimension.X);
@@ -519,22 +517,21 @@ void AddHairStrandInitMeshSamplesPass(
 	FHairStrandsRestRootResource* RestResources,
 	FHairStrandsDeformedRootResource* DeformedResources)
 {
-	if (!RestResources->LODs.IsValidIndex(LODIndex))
+	if (!RestResources->IsValid(LODIndex))
 	{
 		return;
 	}
 	
-	FHairStrandsRestRootResource::FLOD& RestLODData = RestResources->LODs[LODIndex];
-	check(RestLODData.LODIndex == LODIndex);
+	FHairStrandsLODRestRootResource& RestLODData = *RestResources->GetLOD(LODIndex);
 
-	const uint32 TotalSectionCount = RestResources->BulkData.Header.LODs[LODIndex].MeshSectionCount;
+	const uint32 TotalSectionCount = RestLODData.BulkData.Header.MeshSectionCount;
 	const uint32 EffectiveSectionCount = MeshData.Sections.Num();
 	if (EffectiveSectionCount == 0 || RestLODData.SampleCount == 0)
 	{
 		return;
 	}
 
-	FRDGImportedBuffer OutBuffer = Register(GraphBuilder, DeformedResources->LODs[LODIndex].GetDeformedSamplePositionsBuffer(FHairStrandsDeformedRootResource::FLOD::Current), ERDGImportedBufferFlags::CreateUAV);
+	FRDGImportedBuffer OutBuffer = Register(GraphBuilder, DeformedResources->GetLOD(LODIndex)->GetDeformedSamplePositionsBuffer(FHairStrandsLODDeformedRootResource::Current), ERDGImportedBufferFlags::CreateUAV);
 	
 	FHairInitMeshSamplesCS::FParameters CommonParameters;
 	CommonParameters.MaxSampleCount 				= RestLODData.SampleCount;
@@ -635,27 +632,25 @@ void AddHairStrandUpdateMeshSamplesPass(
 	FHairStrandsRestRootResource* RestResources,
 	FHairStrandsDeformedRootResource* DeformedResources)
 {
-	if (LODIndex < 0 || LODIndex >= RestResources->LODs.Num() || LODIndex >= DeformedResources->LODs.Num())
+	if (LODIndex < 0 || !RestResources->IsValid(LODIndex) || !DeformedResources->IsValid(LODIndex))
 	{
 		return;
 	}
 
-	FHairStrandsRestRootResource::FLOD& RestLODData = RestResources->LODs[LODIndex];
-	FHairStrandsDeformedRootResource::FLOD& DeformedLODData = DeformedResources->LODs[LODIndex];
-	check(RestLODData.LODIndex == LODIndex);
-	check(DeformedLODData.LODIndex == LODIndex);
+	FHairStrandsLODRestRootResource& RestLODData = *RestResources->GetLOD(LODIndex);
+	FHairStrandsLODDeformedRootResource& DeformedLODData = *DeformedResources->GetLOD(LODIndex);
 
 	const uint32 SectionCount = MeshData.Sections.Num();
 	if (SectionCount > 0 && RestLODData.SampleCount > 0)
 	{
 		FHairUpdateMeshSamplesCS::FParameters* Parameters = GraphBuilder.AllocParameters<FHairUpdateMeshSamplesCS::FParameters>();
 
-		FRDGImportedBuffer OutWeightsBuffer = Register(GraphBuilder, DeformedLODData.GetMeshSampleWeightsBuffer(FHairStrandsDeformedRootResource::FLOD::Current), ERDGImportedBufferFlags::CreateUAV);
+		FRDGImportedBuffer OutWeightsBuffer = Register(GraphBuilder, DeformedLODData.GetMeshSampleWeightsBuffer(FHairStrandsLODDeformedRootResource::Current), ERDGImportedBufferFlags::CreateUAV);
 
 		Parameters->MaxSampleCount					= RestLODData.SampleCount;
 		Parameters->InterpolationWeightsBuffer		= RegisterAsSRV(GraphBuilder, RestLODData.MeshInterpolationWeightsBuffer);
 		Parameters->SampleRestPositionsBuffer		= RegisterAsSRV(GraphBuilder, RestLODData.RestSamplePositionsBuffer);
-		Parameters->SampleDeformedPositionsBuffer	= RegisterAsSRV(GraphBuilder, DeformedLODData.GetDeformedSamplePositionsBuffer(FHairStrandsDeformedRootResource::FLOD::Current));
+		Parameters->SampleDeformedPositionsBuffer	= RegisterAsSRV(GraphBuilder, DeformedLODData.GetDeformedSamplePositionsBuffer(FHairStrandsLODDeformedRootResource::Current));
 		Parameters->OutSampleDeformationsBuffer		= OutWeightsBuffer.UAV;
 
 		const FIntVector DispatchGroupCount = FComputeShaderUtils::GetGroupCount(RestLODData.SampleCount+4, 128);
@@ -746,10 +741,10 @@ static void AddFollicleMaskPass(
 	FRDGTextureRef OutTexture)
 {
 	const uint32 RootCount = RestResources->GetRootCount();
-	if (LODIndex >= uint32(RestResources->LODs.Num()) || RootCount == 0)
+	if (!RestResources->IsValid(LODIndex) || RootCount == 0)
 		return;
 
-	FHairStrandsRestRootResource::FLOD& LODData = RestResources->LODs[LODIndex];
+	FHairStrandsLODRestRootResource& LODData = *RestResources->GetLOD(LODIndex);
 	if (!LODData.RootBarycentricBuffer.Buffer || !LODData.RestUniqueTrianglePositionBuffer.Buffer)
 		return;
 
@@ -761,7 +756,7 @@ static void AddFollicleMaskPass(
 	Parameters->RootUVsBuffer = nullptr;
 	Parameters->OutputResolution = OutputResolution;
 	Parameters->MaxRootCount = RootCount;
-	Parameters->MaxUniqueTriangleIndex = RestResources->BulkData.Header.LODs[LODIndex].UniqueTriangleCount;
+	Parameters->MaxUniqueTriangleIndex = LODData.BulkData.Header.UniqueTriangleCount;
 	Parameters->Channel = FMath::Min(Channel, 3u);
 	Parameters->KernelSizeInPixels = FMath::Clamp(KernelSizeInPixels, 2u, 200u);
 	Parameters->RenderTargets[0] = FRenderTargetBinding(OutTexture, bNeedClear ? ERenderTargetLoadAction::EClear : ERenderTargetLoadAction::ELoad, 0);
@@ -1043,11 +1038,11 @@ void AddHairStrandUpdatePositionOffsetPass(
 	FRDGImportedBuffer RootTrianglePrevPositionBuffer;
 	if (DeformedRootResources)
 	{
-		RootTriangleCurrPositionBuffer = Register(GraphBuilder, DeformedRootResources->LODs[MeshLODIndex].GetDeformedUniqueTrianglePositionBuffer(FHairStrandsDeformedRootResource::FLOD::Current), ERDGImportedBufferFlags::CreateSRV);
-		RootTrianglePrevPositionBuffer = Register(GraphBuilder, DeformedRootResources->LODs[MeshLODIndex].GetDeformedUniqueTrianglePositionBuffer(FHairStrandsDeformedRootResource::FLOD::Previous), ERDGImportedBufferFlags::CreateSRV);
+		RootTriangleCurrPositionBuffer = Register(GraphBuilder, DeformedRootResources->GetLOD(MeshLODIndex)->GetDeformedUniqueTrianglePositionBuffer(FHairStrandsLODDeformedRootResource::Current), ERDGImportedBufferFlags::CreateSRV);
+		RootTrianglePrevPositionBuffer = Register(GraphBuilder, DeformedRootResources->GetLOD(MeshLODIndex)->GetDeformedUniqueTrianglePositionBuffer(FHairStrandsLODDeformedRootResource::Previous), ERDGImportedBufferFlags::CreateSRV);
 	}
 
-	const bool bUseGPUOffset = DeformedRootResources != nullptr && DeformedRootResources->IsValid() && DeformedRootResources->IsInitialized() && GHairStrandsUseGPUPositionOffset > 0;
+	const bool bUseGPUOffset = DeformedRootResources != nullptr && DeformedRootResources->IsValid(MeshLODIndex) && DeformedRootResources->GetLOD(MeshLODIndex)->IsInitialized() && GHairStrandsUseGPUPositionOffset > 0;
 	const uint32 CurrOffsetIndex = DeformedResources->GetIndex(FHairStrandsDeformedResource::Current);
 	const uint32 PrevOffsetIndex = DeformedResources->GetIndex(FHairStrandsDeformedResource::Previous);
 
diff --git a/Engine/Plugins/Runtime/HairStrands/Source/HairStrandsCore/Private/Niagara/NiagaraDataInterfaceHairStrands.cpp b/Engine/Plugins/Runtime/HairStrands/Source/HairStrandsCore/Private/Niagara/NiagaraDataInterfaceHairStrands.cpp
--- a/Engine/Plugins/Runtime/HairStrands/Source/HairStrandsCore/Private/Niagara/NiagaraDataInterfaceHairStrands.cpp
+++ b/Engine/Plugins/Runtime/HairStrands/Source/HairStrandsCore/Private/Niagara/NiagaraDataInterfaceHairStrands.cpp
@@ -887,7 +887,7 @@ FMatrix44f ComputeWorldTransform(const FNDIHairStrandsData* InstanceData)
 
 static void InterpolateGroomGuides(FRDGBuilder& GraphBuilder, FNiagaraDataBuffer* ParticlesBuffer,
 	const uint32 NodePositionComponent,  const uint32 RestPositionComponent, FNDIHairStrandsBuffer* HairStrandsBuffer,
-	const uint32 StrandsSize, const bool bHasSkinningBinding, const bool bHasValidGeometry, const FMatrix44f& WorldToLocal)
+	const uint32 StrandsSize, const bool bHasSkinningBinding, const bool bHasValidGeometry, const FMatrix44f& WorldToLocal, int32 MeshLODIndex)
 {
 	const bool bIsHairValid = HairStrandsBuffer && HairStrandsBuffer->IsInitialized();
 	const bool bIsRestValid = bIsHairValid && HairStrandsBuffer->SourceRestResources && HairStrandsBuffer->SourceRestResources->IsInitialized();
@@ -895,13 +895,10 @@ static void InterpolateGroomGuides(FRDGBuilder& GraphBuilder, FNiagaraDataBuffer
 
 	if(bIsRestValid && bIsDeformedValid && bHasValidGeometry)
 	{
-		bool bIsRootValid = HairStrandsBuffer->SourceDeformedRootResources && HairStrandsBuffer->SourceDeformedRootResources->IsInitialized() &&
-								HairStrandsBuffer->SourceRestRootResources && HairStrandsBuffer->SourceRestRootResources->IsInitialized() && bHasSkinningBinding;
+		bool bIsRootValid = HairStrandsBuffer->SourceDeformedRootResources && HairStrandsBuffer->SourceDeformedRootResources->IsValid(MeshLODIndex) &&
+							HairStrandsBuffer->SourceRestRootResources && HairStrandsBuffer->SourceRestRootResources->IsValid(MeshLODIndex) && bHasSkinningBinding;
 
 		const uint32 NumPoints = HairStrandsBuffer->SourceRestResources->GetPointCount();
-		const int32 MeshLODIndex = bIsRootValid ? HairStrandsBuffer->SourceDeformedRootResources->MeshLODIndex : -1;
-		bIsRootValid = bIsRootValid && HairStrandsBuffer->SourceDeformedRootResources->IsValid(MeshLODIndex);
-		
 
 		FInterpolateGroomGuidesCS::FPermutationDomain InterpolationDomain;
 		InterpolationDomain.Set<FInterpolateGroomGuidesCS::FInterpolationType>(!bIsRootValid);
@@ -929,11 +926,11 @@ static void InterpolateGroomGuides(FRDGBuilder& GraphBuilder, FNiagaraDataBuffer
 		PassParameters->CurvesOffsetsBuffer = RegisterAsSRV(GraphBuilder, HairStrandsBuffer->SourceRestResources->CurveBuffer);
 		if(bIsRootValid)
 		{
-			const FHairStrandsRestRootResource::FLOD* RestMeshProjection = &(HairStrandsBuffer->SourceRestRootResources->LODs[MeshLODIndex]) ;
-			const FHairStrandsDeformedRootResource::FLOD* DeformedMeshProjection =  &(HairStrandsBuffer->SourceDeformedRootResources->LODs[MeshLODIndex]);
+			const FHairStrandsLODRestRootResource* RestMeshProjection = HairStrandsBuffer->SourceRestRootResources->GetLOD(MeshLODIndex) ;
+			const FHairStrandsLODDeformedRootResource* DeformedMeshProjection = HairStrandsBuffer->SourceDeformedRootResources->GetLOD(MeshLODIndex);
 
 			PassParameters->RestTrianglePositionBuffer = RegisterAsSRV(GraphBuilder, RestMeshProjection->RestUniqueTrianglePositionBuffer);
-			PassParameters->DeformedTrianglePositionBuffer = RegisterAsSRV(GraphBuilder, DeformedMeshProjection->GetDeformedUniqueTrianglePositionBuffer(FHairStrandsDeformedRootResource::FLOD::Current));
+			PassParameters->DeformedTrianglePositionBuffer = RegisterAsSRV(GraphBuilder, DeformedMeshProjection->GetDeformedUniqueTrianglePositionBuffer(FHairStrandsLODDeformedRootResource::Current));
 			PassParameters->RootBarycentricCoordinatesBuffer = RegisterAsSRV(GraphBuilder, RestMeshProjection->RootBarycentricBuffer);
 			PassParameters->RootToUniqueTriangleIndexBuffer = RegisterAsSRV(GraphBuilder, RestMeshProjection->RootToUniqueTriangleIndexBuffer);
 		}
@@ -985,12 +982,14 @@ void UNiagaraDataInterfaceHairStrands::SimCachePostReadFrame(void* OptionalPerIn
 					const bool bHasSkinningBinding = InstanceData->HairGroupInstance->BindingType == EHairBindingType::Skinning;
 					const bool bHasValidGeometry = InstanceData->HairGroupInstance->GeometryType != EHairGeometryType::NoneGeometry;
 				
+					FHairGroupInstance* LocalHairInstance = InstanceData->HairGroupInstance;
+
 					ENQUEUE_RENDER_COMMAND(NiagaraInterpolateGroomSimCache)(
-						[ParticlesBuffer, HairstrandsBuffer, NodePositionComponent, RestPositionComponent, StrandsSize, bHasSkinningBinding, bHasValidGeometry, WorldToLocal](FRHICommandListImmediate& RHICmdList)
+						[ParticlesBuffer, HairstrandsBuffer, NodePositionComponent, RestPositionComponent, StrandsSize, bHasSkinningBinding, bHasValidGeometry, WorldToLocal, LocalHairInstance](FRHICommandListImmediate& RHICmdList)
 						{
 							FMemMark MemMark(FMemStack::Get());
 							FRDGBuilder GraphBuilder(RHICmdList);
-							InterpolateGroomGuides(GraphBuilder, ParticlesBuffer, NodePositionComponent, RestPositionComponent, HairstrandsBuffer, StrandsSize, bHasSkinningBinding, bHasValidGeometry, WorldToLocal);
+							InterpolateGroomGuides(GraphBuilder, ParticlesBuffer, NodePositionComponent, RestPositionComponent, HairstrandsBuffer, StrandsSize, bHasSkinningBinding, bHasValidGeometry, WorldToLocal, LocalHairInstance->Debug.MeshLODIndex);
 							GraphBuilder.Execute();
 						}
 					);
@@ -3353,15 +3352,16 @@ void UNiagaraDataInterfaceHairStrands::SetShaderParameters(const FNiagaraDataInt
 	FNDIHairStrandsProxy& DIProxy = Context.GetProxy<FNDIHairStrandsProxy>();
 	FNDIHairStrandsData* ProxyData = DIProxy.SystemInstancesToProxyData.Find(Context.GetSystemInstanceID());
 
+	const int32 MeshLODIndex = ProxyData && ProxyData->HairGroupInstance ? ProxyData->HairGroupInstance->Debug.MeshLODIndex : -1;
 	const bool bIsHairValid = ProxyData != nullptr && ProxyData->HairStrandsBuffer && ProxyData->HairStrandsBuffer->IsInitialized();
 	const bool bIsHairGroupInstValid = ProxyData != nullptr && ProxyData->HairGroupInstSource != nullptr && ProxyData->HairGroupInstSource->ContainsGroupInstance(ProxyData->HairGroupInstance);
-	const bool bHasSkinningBinding = bIsHairValid && bIsHairGroupInstValid && ProxyData->HairGroupInstance->BindingType == EHairBindingType::Skinning;
-	const bool bIsRootValid = bIsHairValid && ProxyData->HairStrandsBuffer->SourceDeformedRootResources&& ProxyData->HairStrandsBuffer->SourceDeformedRootResources->IsInitialized() && bHasSkinningBinding;
+	const bool bHasSkinningBinding = bIsHairValid && bIsHairGroupInstValid && ProxyData->HairGroupInstance && ProxyData->HairGroupInstance->BindingType == EHairBindingType::Skinning;
+	const bool bIsRootValid = bIsHairValid && ProxyData->HairStrandsBuffer->SourceDeformedRootResources&& ProxyData->HairStrandsBuffer->SourceDeformedRootResources->IsValid(MeshLODIndex) && bHasSkinningBinding;
 	const bool bIsRestValid = bIsHairValid && ProxyData->HairStrandsBuffer->SourceRestResources && ProxyData->HairStrandsBuffer->SourceRestResources->IsInitialized()&&
 		// TEMP: These check are only temporary for avoiding crashes while we find the bottom of the issue.
 		ProxyData->HairStrandsBuffer->ParamsScaleBuffer.IsValid() && ProxyData->HairStrandsBuffer->BoundingBoxBuffer.IsValid();
 
-	const bool bIsGeometryValid = bIsHairValid && (!bIsHairGroupInstValid || (bIsHairGroupInstValid && (ProxyData->HairGroupInstance->GeometryType != EHairGeometryType::NoneGeometry)));
+	const bool bIsGeometryValid = bIsHairValid && (!bIsHairGroupInstValid || (bIsHairGroupInstValid && (ProxyData->HairGroupInstance && ProxyData->HairGroupInstance->GeometryType != EHairGeometryType::NoneGeometry)));
 	const bool bIsDeformedValid = bIsHairValid && ProxyData->HairStrandsBuffer->SourceDeformedResources && ProxyData->HairStrandsBuffer->SourceDeformedResources->IsInitialized();
 
 	NDIHairStrandsLocal::FShaderParameters* ShaderParameters = Context.GetParameterNestedStruct<NDIHairStrandsLocal::FShaderParameters>();
@@ -3397,20 +3397,19 @@ void UNiagaraDataInterfaceHairStrands::SetShaderParameters(const FNiagaraDataInt
 			DeformedPositionBufferUAV = GraphBuilder.CreateUAV(GraphBuilder.RegisterExternalBuffer(HairStrandsBuffer->DeformedPositionBuffer));
 			DeformedPositionOffsetSRV = DummyStructuredBufferSRV;
 		}
-		const int32 MeshLODIndex = bIsRootValid ? HairStrandsBuffer->SourceDeformedRootResources->MeshLODIndex : -1;
 
 		// Projection Buffers
 		const bool bHasSkinnedInterpolation = (bIsRootValid && HairStrandsBuffer->SourceDeformedRootResources->IsValid(MeshLODIndex));
 		const EHairSimulationInterpolationMode InterpolationModeValue = bHasSkinnedInterpolation ? (ProxyData->GlobalInterpolation ?
 			EHairSimulationInterpolationMode::RBF : EHairSimulationInterpolationMode::Skinned) : EHairSimulationInterpolationMode::Rigid;
 
-		const FHairStrandsRestRootResource::FLOD* RestMeshProjection = bHasSkinnedInterpolation ? &(HairStrandsBuffer->SourceRestRootResources->LODs[MeshLODIndex]) : nullptr;
-		const FHairStrandsDeformedRootResource::FLOD* DeformedMeshProjection = bHasSkinnedInterpolation ? &(HairStrandsBuffer->SourceDeformedRootResources->LODs[MeshLODIndex]) : nullptr;
+		const FHairStrandsLODRestRootResource* RestMeshProjection = bHasSkinnedInterpolation ? HairStrandsBuffer->SourceRestRootResources->GetLOD(MeshLODIndex) : nullptr;
+		const FHairStrandsLODDeformedRootResource* DeformedMeshProjection = bHasSkinnedInterpolation ? HairStrandsBuffer->SourceDeformedRootResources->GetLOD(MeshLODIndex) : nullptr;
 
 		FRDGBufferSRVRef RestTrianglePositionSRV = (bHasSkinnedInterpolation && RestMeshProjection) ? RegisterAsSRV(GraphBuilder, RestMeshProjection->RestUniqueTrianglePositionBuffer) : Context.GetComputeDispatchInterface().GetEmptyBufferSRV(GraphBuilder, FHairStrandsMeshTrianglePositionFormat::Format);
-		FRDGBufferSRVRef DeformedTrianglePositionSRV = (bHasSkinnedInterpolation && DeformedMeshProjection) ?  RegisterAsSRV(GraphBuilder, DeformedMeshProjection->GetDeformedUniqueTrianglePositionBuffer(FHairStrandsDeformedRootResource::FLOD::Current)) : Context.GetComputeDispatchInterface().GetEmptyBufferSRV(GraphBuilder, FHairStrandsMeshTrianglePositionFormat::Format);
+		FRDGBufferSRVRef DeformedTrianglePositionSRV = (bHasSkinnedInterpolation && DeformedMeshProjection) ?  RegisterAsSRV(GraphBuilder, DeformedMeshProjection->GetDeformedUniqueTrianglePositionBuffer(FHairStrandsLODDeformedRootResource::Current)) : Context.GetComputeDispatchInterface().GetEmptyBufferSRV(GraphBuilder, FHairStrandsMeshTrianglePositionFormat::Format);
 		
-		FRDGBufferSRVRef RootBarycentricCoordinatesSRV = (bHasSkinnedInterpolation && RestMeshProjection) ? RegisterAsSRV(GraphBuilder, RestMeshProjection->RootBarycentricBuffer) : Context.GetComputeDispatchInterface().GetEmptyBufferSRV(GraphBuilder, FHairStrandsRootBarycentricFormat::Format);
+ 		FRDGBufferSRVRef RootBarycentricCoordinatesSRV = (bHasSkinnedInterpolation && RestMeshProjection) ? RegisterAsSRV(GraphBuilder, RestMeshProjection->RootBarycentricBuffer) : Context.GetComputeDispatchInterface().GetEmptyBufferSRV(GraphBuilder, FHairStrandsRootBarycentricFormat::Format);
 		FRDGBufferSRVRef RootToUniqueTriangleIndexSRV = (bHasSkinnedInterpolation && RestMeshProjection) ? RegisterAsSRV(GraphBuilder, RestMeshProjection->RootToUniqueTriangleIndexBuffer) : Context.GetComputeDispatchInterface().GetEmptyBufferSRV(GraphBuilder, FHairStrandsRootToUniqueTriangleIndexFormat::Format);
 
 		// RBF buffers
@@ -3418,7 +3417,7 @@ void UNiagaraDataInterfaceHairStrands::SetShaderParameters(const FNiagaraDataInt
 		const int32 SampleCountValue = bHasSamples ? RestMeshProjection->SampleCount : 0;
 
 		FRDGBufferSRVRef RestSamplePositionsBufferSRV = (bHasSamples && RestMeshProjection) ? RegisterAsSRV(GraphBuilder,RestMeshProjection->RestSamplePositionsBuffer) : DummyStructuredBufferSRV;
-		FRDGBufferSRVRef MeshSampleWeightsBufferSRV = (bHasSamples && DeformedMeshProjection) ? RegisterAsSRV(GraphBuilder,DeformedMeshProjection->GetMeshSampleWeightsBuffer((FHairStrandsDeformedRootResource::FLOD::Current))) : DummyStructuredBufferSRV;
+		FRDGBufferSRVRef MeshSampleWeightsBufferSRV = (bHasSamples && DeformedMeshProjection) ? RegisterAsSRV(GraphBuilder,DeformedMeshProjection->GetMeshSampleWeightsBuffer((FHairStrandsLODDeformedRootResource::Current))) : DummyStructuredBufferSRV;
 		
 		// Simulation setup (we update the rest configuration based on the deformed positions 
 		// if in restupdate mode or if we are resetting the sim and using RBF transfer since the rest positions are not matrching the physics asset)
diff --git a/Engine/Plugins/Runtime/HairStrands/Source/HairStrandsCore/Public/GroomBindingAsset.h b/Engine/Plugins/Runtime/HairStrands/Source/HairStrandsCore/Public/GroomBindingAsset.h
--- a/Engine/Plugins/Runtime/HairStrands/Source/HairStrandsCore/Public/GroomBindingAsset.h
+++ b/Engine/Plugins/Runtime/HairStrands/Source/HairStrandsCore/Public/GroomBindingAsset.h
@@ -179,9 +179,9 @@ public:
 	/** Binding bulk data */
 	struct FHairGroupPlatformData
 	{
-		FHairStrandsRootBulkData		SimRootBulkData;
-		FHairStrandsRootBulkData		RenRootBulkData;
-		TArray<FHairStrandsRootBulkData>CardsRootBulkData;
+		TArray<FHairStrandsRootBulkData>		 SimRootBulkDatas;
+		TArray<FHairStrandsRootBulkData>		 RenRootBulkDatas;
+		TArray<TArray<FHairStrandsRootBulkData>> CardsRootBulkDatas;
 	};
 
 private:
diff --git a/Engine/Plugins/Runtime/HairStrands/Source/HairStrandsCore/Public/GroomResources.h b/Engine/Plugins/Runtime/HairStrands/Source/HairStrandsCore/Public/GroomResources.h
--- a/Engine/Plugins/Runtime/HairStrands/Source/HairStrandsCore/Public/GroomResources.h
+++ b/Engine/Plugins/Runtime/HairStrands/Source/HairStrandsCore/Public/GroomResources.h
@@ -104,17 +104,16 @@ struct FHairCommonResource : public FRenderResource
 
 	/* Init/Release buffers (FHairCommonResource) */
 	void Allocate(FRDGBuilder& GraphBuilder, EHairResourceLoadingType LoadingType);
-	void Allocate(FRDGBuilder& GraphBuilder, EHairResourceLoadingType LoadingType, EHairResourceStatus& Status, int32 InLODIndex=-1);
-	void Allocate(FRDGBuilder& GraphBuilder, EHairResourceLoadingType LoadingType, EHairResourceStatus& Status, uint32 InRequestedCurveCount, uint32 InRequestedPointCount, int32 InLODIndex=-1, bool bAllowDeallocation=true);
+	void Allocate(FRDGBuilder& GraphBuilder, EHairResourceLoadingType LoadingType, EHairResourceStatus& Status);
+	void Allocate(FRDGBuilder& GraphBuilder, EHairResourceLoadingType LoadingType, EHairResourceStatus& Status, uint32 InRequestedCurveCount, uint32 InRequestedPointCount, bool bAllowDeallocation=true);
 
-	void StreamInData(int32 InLODIndex=-1);
+	void StreamInData();
 
 	virtual void InternalAllocate() {}
 	virtual void InternalAllocate(FRDGBuilder& GraphBuilder) {}
-	virtual void InternalAllocate(FRDGBuilder& GraphBuilder, uint32 InCurveCount, uint32 InPointCount, int32 InLODIndex) { InternalAllocate(GraphBuilder); }
+	virtual void InternalAllocate(FRDGBuilder& GraphBuilder, uint32 InCurveCount, uint32 InPointCount) { InternalAllocate(GraphBuilder); }
 	virtual void InternalRelease() {}
-	virtual bool InternalGetOrRequestData(uint32 InRequestedCurveCount, uint32 InRequestedPointCount, int32 InLODIndex);
-	virtual bool InternalIsLODDataLoaded(uint32 InRequestedCurveCount, uint32 InRequestedPointCount, int32 InLODIndex) const { return true; }
+	virtual bool InternalGetOrRequestData(uint32 InRequestedCurveCount, uint32 InRequestedPointCount);
 	virtual void InternalResetLoadedSize() { }
 	virtual FHairStrandsBulkCommon* InternalGetBulkData() { return nullptr; }
 
@@ -143,7 +142,7 @@ struct FHairStrandsRestResource : public FHairCommonResource
 	/* Init/Release buffers */
 	virtual void InternalAllocate(FRDGBuilder& GraphBuilder) override;
 	virtual void InternalRelease() override;
-	virtual bool InternalGetOrRequestData(uint32 InRequestedCurveCount, uint32 InRequestedPointCount, int32 InLODIndex) override;
+	virtual bool InternalGetOrRequestData(uint32 InRequestedCurveCount, uint32 InRequestedPointCount) override;
 	virtual void InternalResetLoadedSize() override;
 	virtual FHairStrandsBulkCommon* InternalGetBulkData() override { return &BulkData; }
 
@@ -210,7 +209,7 @@ struct FHairStrandsDeformedResource : public FHairCommonResource
 	FHairStrandsDeformedResource(FHairStrandsBulkData& BulkData, EHairStrandsResourcesType CurveType, const FHairResourceName& ResourceName, const FName& OwnerName);
 
 	/* Init/Release buffers */
-	virtual void InternalAllocate(FRDGBuilder& GraphBuilder, uint32 InCurveCount, uint32 InPointCount, int32 InLODIndex) override;
+	virtual void InternalAllocate(FRDGBuilder& GraphBuilder, uint32 InCurveCount, uint32 InPointCount) override;
 	virtual void InternalRelease() override;
 
 	/* Get the resource name */
@@ -292,7 +291,7 @@ struct FHairStrandsInterpolationResource : public FHairCommonResource
 	/* Init/Release buffers */
 	virtual void InternalAllocate(FRDGBuilder& GraphBuilder) override;
 	virtual void InternalRelease() override;
-	virtual bool InternalGetOrRequestData(uint32 InRequestedCurveCount, uint32 InRequestedPointCount, int32 InLODIndex) override;
+	virtual bool InternalGetOrRequestData(uint32 InRequestedCurveCount, uint32 InRequestedPointCount) override;
 	virtual void InternalResetLoadedSize() override;
 	virtual FHairStrandsBulkCommon* InternalGetBulkData() override { return &BulkData; }
 
@@ -325,7 +324,7 @@ struct FHairStrandsClusterResource : public FHairCommonResource
 	/* Init/Release buffers */
 	virtual void InternalAllocate(FRDGBuilder& GraphBuilder) override;
 	virtual void InternalRelease() override;
-	virtual bool InternalGetOrRequestData(uint32 InRequestedCurveCount, uint32 InRequestedPointCount, int32 InLODIndex) override;
+	virtual bool InternalGetOrRequestData(uint32 InRequestedCurveCount, uint32 InRequestedPointCount) override;
 	virtual void InternalResetLoadedSize() override;
 	virtual FHairStrandsBulkCommon* InternalGetBulkData() override { return &BulkData; }
 
@@ -359,7 +358,7 @@ struct FHairStrandsCullingResource : public FHairCommonResource
 	FHairStrandsCullingResource(uint32 InPointCount, uint32 InCurveCount, uint32 InClusterCount, const FHairResourceName& InResourceName, const FName& InOwnerName);
 
 	/* Init/Release buffers */
-	virtual void InternalAllocate(FRDGBuilder& GraphBuilder, uint32 InCurveCount, uint32 InPointCount, int32 InLODIndex) override;
+	virtual void InternalAllocate(FRDGBuilder& GraphBuilder, uint32 InCurveCount, uint32 InPointCount) override;
 	virtual void InternalRelease() override;
 
 	/* Get the resource name */
@@ -384,174 +383,227 @@ struct FHairStrandsCullingResource : public FHairCommonResource
 };
 
 /* Render buffers for root deformation for dynamic meshes */
-struct FHairStrandsRestRootResource : public FHairCommonResource
+struct FHairStrandsLODRestRootResource : public FHairCommonResource
 {
 	/** Build the hair strands resource */
-	FHairStrandsRestRootResource(FHairStrandsRootBulkData& BulkData, EHairStrandsResourcesType CurveType, const FHairResourceName& ResourceName, const FName& OwnerName);
+	FHairStrandsLODRestRootResource(FHairStrandsRootBulkData& BulkData, EHairStrandsResourcesType CurveType, uint32 InMeshLODIndex, const FHairResourceName& ResourceName, const FName& OwnerName);
 
 	/* Init/Release buffers */
-	virtual void InternalAllocate(FRDGBuilder& GraphBuilder, uint32 InCurveCount, uint32 InPointCount, int32 InLODIndex) override;
+	virtual void InternalAllocate(FRDGBuilder& GraphBuilder, uint32 InCurveCount, uint32 InPointCount) override;
 	virtual void InternalRelease() override;
-	virtual bool InternalGetOrRequestData(uint32 InRequestedCurveCount, uint32 InRequestedPointCount, int32 InLODIndex) override;
-	virtual bool InternalIsLODDataLoaded(uint32 InRequestedCurveCount, uint32 InRequestedPointCount, int32 InLODIndex) const override;
+	virtual bool InternalGetOrRequestData(uint32 InRequestedCurveCount, uint32 InRequestedPointCount) override;
 	virtual void InternalResetLoadedSize() override;
 	virtual FHairStrandsBulkCommon* InternalGetBulkData() override { return &BulkData; }
 
 	/* Get the resource name */
-	virtual FString GetFriendlyName() const override { return TEXT("FHairStrandsRestRootResource"); }
+	virtual FString GetFriendlyName() const override { return TEXT("FHairStrandsLODRestRootResource"); }
 
 	/* Populate GPU LOD data from RootData (this function doesn't initialize resources) */
-	void PopulateFromRootData();
+	void Reset();
 
 	// Accessors
-	uint32 GetLODCount() const { return BulkData.GetLODCount(); }
 	uint32 GetRootCount()const { return BulkData.GetRootCount(); }
 
 	/* Return the memory size for GPU resources */
 	uint32 GetResourcesSize() const 
 	{
 		uint32 Total = 0;
-		for (const FLOD& LOD : LODs)
-		{
-			Total += GetBufferTotalNumBytes(LOD.UniqueTriangleIndexBuffer);
-			Total += GetBufferTotalNumBytes(LOD.RootToUniqueTriangleIndexBuffer);
-			Total += GetBufferTotalNumBytes(LOD.RootBarycentricBuffer);
-			Total += GetBufferTotalNumBytes(LOD.RestUniqueTrianglePositionBuffer);
-			Total += GetBufferTotalNumBytes(LOD.MeshInterpolationWeightsBuffer);
-			Total += GetBufferTotalNumBytes(LOD.MeshSampleIndicesAndSectionsBuffer);
-			Total += GetBufferTotalNumBytes(LOD.RestSamplePositionsBuffer);
-		}
+		Total += GetBufferTotalNumBytes(UniqueTriangleIndexBuffer);
+		Total += GetBufferTotalNumBytes(RootToUniqueTriangleIndexBuffer);
+		Total += GetBufferTotalNumBytes(RootBarycentricBuffer);
+		Total += GetBufferTotalNumBytes(RestUniqueTrianglePositionBuffer);
+		Total += GetBufferTotalNumBytes(MeshInterpolationWeightsBuffer);
+		Total += GetBufferTotalNumBytes(MeshSampleIndicesAndSectionsBuffer);
+		Total += GetBufferTotalNumBytes(RestSamplePositionsBuffer);
 		return Total;
 	}
 
-	struct FLOD
-	{
-		enum class EStatus { Invalid, Initialized, Completed };
-
-		const bool IsValid() const { return Status == EStatus::Completed; }
-		EStatus Status = EStatus::Invalid;
-		int32 LODIndex = -1;
-		uint32 AvailableCurveCount = 0;
-
-		/* Triangle on which a root is attached */
-		/* When the projection is done with source to target mesh transfer, the projection indices does not match.
-		   In this case we need to separate index computation. The barycentric coords remain the same however. */
-		FRDGExternalBuffer UniqueTriangleIndexBuffer;
-		/* Strands hair root to unique triangle index */
-		FRDGExternalBuffer RootToUniqueTriangleIndexBuffer;
-		FRDGExternalBuffer RootBarycentricBuffer;
-
-		/* Strand hair roots translation and rotation in rest position relative to the bound triangle. Positions are relative to the rest root center */
-		FRDGExternalBuffer RestUniqueTrianglePositionBuffer;
-
-		/* Strand hair mesh interpolation matrix and sample indices */
-		uint32 SampleCount = 0;
-		FRDGExternalBuffer MeshInterpolationWeightsBuffer;
-		FRDGExternalBuffer MeshSampleIndicesAndSectionsBuffer;
-		FRDGExternalBuffer RestSamplePositionsBuffer;
-	};
+	enum class EStatus { Invalid, Initialized, Completed };
 
-	/* Store the hair projection information for each mesh LOD */
-	TArray<FLOD> LODs;
+	const bool IsDataValid() const { return BulkData.Header.RootCount > 0 && BulkData.Header.LODIndex >= 0; }
+	const bool IsValid() const { return Status == EStatus::Completed; }
+	EStatus Status = EStatus::Invalid;
+	uint32 AvailableCurveCount = 0;
+
+	/* Triangle on which a root is attached */
+	/* When the projection is done with source to target mesh transfer, the projection indices does not match.
+	In this case we need to separate index computation. The barycentric coords remain the same however. */
+	FRDGExternalBuffer UniqueTriangleIndexBuffer;
+	/* Strands hair root to unique triangle index */
+	FRDGExternalBuffer RootToUniqueTriangleIndexBuffer;
+	FRDGExternalBuffer RootBarycentricBuffer;
+
+	/* Strand hair roots translation and rotation in rest position relative to the bound triangle. Positions are relative to the rest root center */
+	FRDGExternalBuffer RestUniqueTrianglePositionBuffer;
+
+	/* Strand hair mesh interpolation matrix and sample indices */
+	uint32 SampleCount = 0;
+	FRDGExternalBuffer MeshInterpolationWeightsBuffer;
+	FRDGExternalBuffer MeshSampleIndicesAndSectionsBuffer;
+	FRDGExternalBuffer RestSamplePositionsBuffer;
 
 	/* Store CPU data for root info & root binding */
 	FHairStrandsRootBulkData& BulkData;
 
 	/* Type of curves */
 	const EHairStrandsResourcesType CurveType;
+
+	/* Mesh LOD index for this root resource. Used for sanity check */
+	const uint32 MeshLODIndex = 0;
 };
 
 /* Render buffers for root deformation for dynamic meshes */
-struct FHairStrandsDeformedRootResource : public FHairCommonResource
+struct FHairStrandsRestRootResource
 {
 	/** Build the hair strands resource */
-	FHairStrandsDeformedRootResource(EHairStrandsResourcesType CurveType, const FHairResourceName& ResourceName, const FName& OwnerName);
-	FHairStrandsDeformedRootResource(const FHairStrandsRestRootResource* InRestResources, EHairStrandsResourcesType CurveType, const FHairResourceName& ResourceName, const FName& OwnerName);
+	FHairStrandsRestRootResource(TArray<FHairStrandsRootBulkData>& BulkData, EHairStrandsResourcesType CurveType, const FHairResourceName& ResourceName, const FName& OwnerName);
 
-	/* Init/Release buffers */
-	virtual void InternalAllocate(FRDGBuilder& GraphBuilder, uint32 InCurveCount, uint32 InPointCount, int32 InLODIndex) override;
-	virtual void InternalRelease() override;
-	virtual bool InternalIsLODDataLoaded(uint32 InRequestedCurveCount, uint32 InRequestedPointCount, int32 InLODIndex) const override;
+	void Allocate(FRDGBuilder& GraphBuilder, EHairResourceLoadingType LoadingType, EHairResourceStatus& Status, int32 InMeshLODIndex);
+	void Allocate(FRDGBuilder& GraphBuilder, EHairResourceLoadingType LoadingType, EHairResourceStatus& Status, uint32 InRequestedCurveCount, uint32 InRequestedPointCount, bool bAllowDeallocation/*=true*/, int32 InMeshLODIndex);
+	void StreamInData(int32 MeshLODIndex);
 
-	/* Get the resource name */
-	virtual FString GetFriendlyName() const override { return TEXT("FHairStrandsDeformedRootResource"); }
+	/* Forward call to each mesh LOD version */
+	void BeginInitResource();
+	void BeginUpdateResourceRHI();
+	void ReleaseResource();
+	void InternalResetLoadedSize();
 
-	/* Indirect if the current root resources are valid and up to date */
-	bool IsValid() const { return MeshLODIndex >= 0 && MeshLODIndex < LODs.Num() && LODs[MeshLODIndex].IsValid(); }
-	bool IsValid(int32 InMeshLODIndex) const { return InMeshLODIndex >= 0 && InMeshLODIndex < LODs.Num() && LODs[InMeshLODIndex].IsValid(); }
+	// Accessors
+	bool IsDataValid(int32 InMeshLODIndex) const { return LODs.IsValidIndex(InMeshLODIndex) && LODs[InMeshLODIndex]->IsDataValid(); }
+	bool IsValid(int32 InMeshLODIndex) const { return LODs.IsValidIndex(InMeshLODIndex) && LODs[InMeshLODIndex]->IsValid(); }
+	FHairStrandsLODRestRootResource* GetLOD(int32 InMeshLODIndex) { check(LODs.IsValidIndex(InMeshLODIndex)); return LODs[InMeshLODIndex]; }
+	const FHairStrandsLODRestRootResource* GetLOD(int32 InMeshLODIndex) const { check(LODs.IsValidIndex(InMeshLODIndex)); return LODs[InMeshLODIndex]; }
+
+	uint32 GetLODCount() const { return LODs.Num(); }
+	uint32 GetRootCount()const { return LODs.Num() > 0 && LODs[0] != nullptr? LODs[0]->GetRootCount() : 0u; }
 
 	/* Return the memory size for GPU resources */
-	uint32 GetResourcesSize() const
+	uint32 GetResourcesSize() const 
 	{
 		uint32 Total = 0;
-		for (const FLOD& LOD : LODs)
+		for (const FHairStrandsLODRestRootResource* LOD : LODs)
 		{
-			Total += GetBufferTotalNumBytes(LOD.DeformedUniqueTrianglePositionBuffer[0]);
-			Total += GetBufferTotalNumBytes(LOD.DeformedSamplePositionsBuffer[0]);
-			Total += GetBufferTotalNumBytes(LOD.MeshSampleWeightsBuffer[0]);
-
-			// Double buffering is disabled by default unless the read-only cvar r.HairStrands.ContinuousDecimationReordering is set
-			if (IsHairStrandContinuousDecimationReorderingEnabled())
-			{
-				Total += GetBufferTotalNumBytes(LOD.DeformedUniqueTrianglePositionBuffer[1]);
-				Total += GetBufferTotalNumBytes(LOD.DeformedSamplePositionsBuffer[1]);
-				Total += GetBufferTotalNumBytes(LOD.MeshSampleWeightsBuffer[1]);
-			}
+			Total += LOD->GetResourcesSize();
 		}
 		return Total;
 	}
 
-	void SwapBuffer()
+	/* Store the hair projection information for each mesh LOD */
+	TArray<FHairStrandsLODRestRootResource*> LODs;
+
+	/* Type of curves */
+	const EHairStrandsResourcesType CurveType;
+};
+
+/* Render buffers for root deformation for dynamic meshes */
+struct FHairStrandsLODDeformedRootResource : public FHairCommonResource
+{
+	/** Build the hair strands resource */
+	FHairStrandsLODDeformedRootResource (FHairStrandsRestRootResource* InRestResources, EHairStrandsResourcesType CurveType, const uint32 MeshLODIndex, const FHairResourceName& ResourceName, const FName& OwnerName);
+
+	/* Init/Release buffers */
+	virtual void InternalAllocate(FRDGBuilder& GraphBuilder, uint32 InCurveCount, uint32 InPointCount) override;
+	virtual void InternalRelease() override;
+	//virtual bool InternalIsLODDataLoaded(uint32 InRequestedCurveCount, uint32 InRequestedPointCount) const override;
+
+	/* Get the resource name */
+	virtual FString GetFriendlyName() const override { return TEXT("FHairStrandsLODDeformedRootResource "); }
+
+	/* Return the memory size for GPU resources */
+	uint32 GetResourcesSize() const
 	{
-		for (FLOD& LOD : LODs)
+		uint32 Total = 0;
+		Total += GetBufferTotalNumBytes(DeformedUniqueTrianglePositionBuffer[0]);
+		Total += GetBufferTotalNumBytes(DeformedSamplePositionsBuffer[0]);
+		Total += GetBufferTotalNumBytes(MeshSampleWeightsBuffer[0]);
+
+		// Double buffering is disabled by default unless the read-only cvar r.HairStrands.ContinuousDecimationReordering is set
+		if (IsHairStrandContinuousDecimationReorderingEnabled())
 		{
-			LOD.SwapBuffer();
+			Total += GetBufferTotalNumBytes(DeformedUniqueTrianglePositionBuffer[1]);
+			Total += GetBufferTotalNumBytes(DeformedSamplePositionsBuffer[1]);
+			Total += GetBufferTotalNumBytes(MeshSampleWeightsBuffer[1]);
 		}
+		return Total;
 	}
+	
+	enum class EStatus { Invalid, Initialized, Completed };
+
+	// A LOD is considered valid as long as its resources are initialized. 
+	// Its state will become completed once its triangle position will be 
+	// update, but in order to be update its status needs to be valid.
+	bool IsDataValid() const { return RootCount > 0 && MeshLODIndex >= 0; }
+	bool IsValid() const { return Status == EStatus::Initialized || Status == EStatus::Completed; }
+	EStatus Status = EStatus::Invalid;
+	uint32 AvailableCurveCount = 0;
 
-	struct FLOD
-	{
-		enum class EStatus { Invalid, Initialized, Completed };
-
-		// A LOD is considered valid as long as its resources are initialized. 
-		// Its state will become completed once its triangle position will be 
-		// update, but in order to be update its status needs to be valid.
-		const bool IsValid() const { return Status == EStatus::Initialized || Status == EStatus::Completed; }
-		EStatus Status = EStatus::Invalid;
-		int32 LODIndex = -1;
-		uint32 AvailableCurveCount = 0;
+	/* Strand hair roots translation and rotation in triangle-deformed position relative to the bound triangle. Positions are relative the deformed root center*/
+	FRDGExternalBuffer DeformedUniqueTrianglePositionBuffer[2];
 
-		/* Strand hair roots translation and rotation in triangle-deformed position relative to the bound triangle. Positions are relative the deformed root center*/
-		FRDGExternalBuffer DeformedUniqueTrianglePositionBuffer[2];
+	/* Strand hair mesh interpolation matrix and sample indices */
+	uint32 SampleCount = 0;
+	FRDGExternalBuffer DeformedSamplePositionsBuffer[2];
+	FRDGExternalBuffer MeshSampleWeightsBuffer[2];
 
-		/* Strand hair mesh interpolation matrix and sample indices */
-		uint32 SampleCount = 0;
-		FRDGExternalBuffer DeformedSamplePositionsBuffer[2];
-		FRDGExternalBuffer MeshSampleWeightsBuffer[2];
+	/* Whether the GPU data should be initialized with the asset data or not */
+	uint32 CurrentIndex = 0;
 
-		/* Whether the GPU data should be initialized with the asset data or not */
-		uint32 CurrentIndex = 0;
+	enum EFrameType
+	{
+		Previous,
+		Current
+	};
 
-		enum EFrameType
-		{
-			Previous,
-			Current
-		};
+	// Double buffering is disabled by default unless the read-only cvar r.HairStrands.ContinuousDecimationReordering is set
+	inline uint32 GetIndex(EFrameType T) const { return T == EFrameType::Current ? CurrentIndex : 1u - CurrentIndex; }
+	inline const FRDGExternalBuffer& GetDeformedUniqueTrianglePositionBuffer(EFrameType T) const { return IsHairStrandContinuousDecimationReorderingEnabled() ? DeformedUniqueTrianglePositionBuffer[GetIndex(T)] : DeformedUniqueTrianglePositionBuffer[0]; }
+	inline const FRDGExternalBuffer& GetDeformedSamplePositionsBuffer(EFrameType T) const { return IsHairStrandContinuousDecimationReorderingEnabled() ? DeformedSamplePositionsBuffer[GetIndex(T)] : DeformedSamplePositionsBuffer[0]; }
+	inline const FRDGExternalBuffer& GetMeshSampleWeightsBuffer(EFrameType T) const { return IsHairStrandContinuousDecimationReorderingEnabled() ? MeshSampleWeightsBuffer[GetIndex(T)] : MeshSampleWeightsBuffer[0]; }
+	inline void SwapBuffer() { CurrentIndex = 1u - CurrentIndex; }
 
-		// Double buffering is disabled by default unless the read-only cvar r.HairStrands.ContinuousDecimationReordering is set
-		inline uint32 GetIndex(EFrameType T) const { return T == EFrameType::Current ? CurrentIndex : 1u - CurrentIndex; }
-		inline const FRDGExternalBuffer& GetDeformedUniqueTrianglePositionBuffer(EFrameType T) const { return IsHairStrandContinuousDecimationReorderingEnabled() ? DeformedUniqueTrianglePositionBuffer[GetIndex(T)] : DeformedUniqueTrianglePositionBuffer[0]; }
-		inline const FRDGExternalBuffer& GetDeformedSamplePositionsBuffer(EFrameType T) const { return IsHairStrandContinuousDecimationReorderingEnabled() ? DeformedSamplePositionsBuffer[GetIndex(T)] : DeformedSamplePositionsBuffer[0]; }
-		inline const FRDGExternalBuffer& GetMeshSampleWeightsBuffer(EFrameType T) const { return IsHairStrandContinuousDecimationReorderingEnabled() ? MeshSampleWeightsBuffer[GetIndex(T)] : MeshSampleWeightsBuffer[0]; }
-		inline void SwapBuffer() { CurrentIndex = 1u - CurrentIndex; }
-	};
 
 	/* Store the hair projection information for each mesh LOD */
 	uint32 RootCount = 0;
-	TArray<FLOD> LODs;
 
-	/* Last update MeshLODIndex */
-	int32 MeshLODIndex = -1;
+	/* Mesh index corresponding for this root data. Used for sanity check. */
+	const int32 MeshLODIndex = -1;
+
+	/* Type of curves */
+	const EHairStrandsResourcesType CurveType;
+};
+
+/* Render buffers for root deformation for dynamic meshes */
+struct FHairStrandsDeformedRootResource
+{
+	/** Build the hair strands resource */
+	FHairStrandsDeformedRootResource(FHairStrandsRestRootResource* InRestResources, EHairStrandsResourcesType CurveType, const FHairResourceName& ResourceName, const FName& OwnerName);
+
+	/* Init/Release buffers */
+	void Allocate(FRDGBuilder& GraphBuilder, EHairResourceLoadingType LoadingType, EHairResourceStatus& Status, int32 InMeshLODIndex);
+	void Allocate(FRDGBuilder& GraphBuilder, EHairResourceLoadingType LoadingType, EHairResourceStatus& Status, uint32 InRequestedCurveCount, uint32 InRequestedPointCount, bool bAllowDeallocation, int32 InMeshLODIndex);
+	void ReleaseResource();
+	//void InternalRelease();
+	void SwapBuffer();
+
+	/* Indirect if the current root resources are valid and up to date */
+	bool IsDataValid(int32 InMeshLODIndex) const { return LODs.IsValidIndex(InMeshLODIndex) && LODs[InMeshLODIndex]->IsDataValid(); }
+	bool IsValid(int32 InMeshLODIndex) const { return LODs.IsValidIndex(InMeshLODIndex) && LODs[InMeshLODIndex]->IsValid(); }
+	FHairStrandsLODDeformedRootResource* GetLOD(int32 InMeshLODIndex) { check(LODs.IsValidIndex(InMeshLODIndex)); return LODs[InMeshLODIndex]; }
+	const FHairStrandsLODDeformedRootResource* GetLOD(int32 InMeshLODIndex) const { check(LODs.IsValidIndex(InMeshLODIndex)); return LODs[InMeshLODIndex]; }
+
+	/* Return the memory size for GPU resources */
+	uint32 GetResourcesSize() const
+	{
+		uint32 Total = 0;
+		for (const FHairStrandsLODDeformedRootResource* LOD : LODs)
+		{
+			Total += LOD->GetResourcesSize();
+		}
+		return Total;
+	}
+
+	/* Store the hair projection information for each mesh LOD */
+	TArray<FHairStrandsLODDeformedRootResource*> LODs;
 
 	/* Type of curves */
 	const EHairStrandsResourcesType CurveType;
@@ -566,7 +618,7 @@ struct FHairStrandsRaytracingResource : public FHairCommonResource
 	FHairStrandsRaytracingResource(const FHairMeshesBulkData& InData, const FHairResourceName& ResourceName, const FName& OwnerName);
 
 	/* Init/Release buffers */
-	virtual void InternalAllocate(FRDGBuilder& GraphBuilder, uint32 InCurveCount, uint32 InPointCount, int32 InLODIndex) override;
+	virtual void InternalAllocate(FRDGBuilder& GraphBuilder, uint32 InCurveCount, uint32 InPointCount) override;
 	virtual void InternalRelease() override;
 
 	/* Get the resource name */
diff --git a/Engine/Plugins/Runtime/HairStrands/Source/HairStrandsCore/Public/HairStrandsDatas.h b/Engine/Plugins/Runtime/HairStrands/Source/HairStrandsCore/Public/HairStrandsDatas.h
--- a/Engine/Plugins/Runtime/HairStrands/Source/HairStrandsCore/Public/HairStrandsDatas.h
+++ b/Engine/Plugins/Runtime/HairStrands/Source/HairStrandsCore/Public/HairStrandsDatas.h
@@ -343,7 +343,7 @@ struct FHairStreamingRequest
 		void Release();
 	};
 
-	void Request(uint32 InRequestedCurveCount, uint32 InRequestedPointCount, int32 InLODIndex, FHairStrandsBulkCommon& In,
+	void Request(uint32 InRequestedCurveCount, uint32 InRequestedPointCount, FHairStrandsBulkCommon& In,
 		bool bWait=false, bool bFillBulkData=false, bool bWarmCache=false, const FName& InOwnerName = NAME_None,
 		bool* bWaitResult = nullptr);
 	bool IsNone() const;
@@ -351,7 +351,7 @@ struct FHairStreamingRequest
 	bool IsUnloading() const;
 
 #if WITH_EDITORONLY_DATA
-	bool WarmCache(uint32 InRequestedCurveCount, uint32 InRequestedPointCount, int32 InLODIndex, FHairStrandsBulkCommon& In);
+	bool WarmCache(uint32 InRequestedCurveCount, uint32 InRequestedPointCount, FHairStrandsBulkCommon& In);
 #endif
 
 #if !WITH_EDITORONLY_DATA
@@ -396,11 +396,11 @@ struct HAIRSTRANDSCORE_API FHairStrandsBulkCommon
 	virtual ~FHairStrandsBulkCommon() { }
 	void Serialize(FArchive& Ar, UObject* Owner);
 	virtual void SerializeHeader(FArchive& Ar, UObject* Owner) = 0;
-	void SerializeData(FArchive& Ar, UObject* Owner, int32 LODIndex=-1);
+	void SerializeData(FArchive& Ar, UObject* Owner);
 
-	void Write_DDC(UObject* Owner, TArray<UE::DerivedData::FCachePutValueRequest>& Out, int32 LODIndex=-1);
+	void Write_DDC(UObject* Owner, TArray<UE::DerivedData::FCachePutValueRequest>& Out);
 	void Read_DDC(FHairStreamingRequest* In, TArray<UE::DerivedData::FCacheGetChunkRequest>& Out);
-	void Write_IO(UObject* Owner, FArchive& Out, int32 LODIndex=-1);
+	void Write_IO(UObject* Owner, FArchive& Out);
 	void Read_IO(FHairStreamingRequest* In, FBulkDataBatchRequest& Out);
 	void Unload(FHairStreamingRequest* In);
 
@@ -409,7 +409,6 @@ struct HAIRSTRANDSCORE_API FHairStrandsBulkCommon
 		void Add(FHairBulkContainer& In, const TCHAR* InSuffix, uint32& InOffset, uint32 InSize=0);
 		uint32 GetCurveCount() const { check(StreamingRequest); return StreamingRequest->CurveCount; }
 		uint32 GetPointCount() const { check(StreamingRequest); return StreamingRequest->PointCount; }
-		int32  GetLODIndex() const	 { return StreamingRequest ? StreamingRequest->LODIndex : LODIndex; }
 		enum EQueryType { None, ReadDDC, WriteDDC, ReadIO, ReadWriteIO /* i.e. regular Serialize() */, UnloadData};
 		EQueryType Type = None;
 		FHairStreamingRequest* StreamingRequest = nullptr;
@@ -421,7 +420,6 @@ struct HAIRSTRANDSCORE_API FHairStrandsBulkCommon
 		FString* DerivedDataKey = nullptr;
 	#endif
 		UObject* Owner = nullptr; 
-		int32 LODIndex = -1;
 	};
 
 	virtual uint32 GetResourceCount() const = 0;
@@ -773,50 +771,45 @@ struct FHairStrandsRootData
 	bool HasProjectionData() const;
 	bool IsValid() const { return RootCount > 0; }
 
-	struct FMeshProjectionLOD
-	{
-		int32 LODIndex = -1;
+	/* Triangle on which a root is attached */
+	/* When the projection is done with source to target mesh transfer, the projection indices does not match.
+		In this case we need to separate index computation. The barycentric coords remain the same however. */
+	TArray<FHairStrandsRootToUniqueTriangleIndexFormat::Type> RootToUniqueTriangleIndexBuffer;
+	TArray<FHairStrandsRootBarycentricFormat::Type> RootBarycentricBuffer;
 
-		/* Triangle on which a root is attached */
-		/* When the projection is done with source to target mesh transfer, the projection indices does not match.
-			In this case we need to separate index computation. The barycentric coords remain the same however. */
-		TArray<FHairStrandsRootToUniqueTriangleIndexFormat::Type> RootToUniqueTriangleIndexBuffer;
-		TArray<FHairStrandsRootBarycentricFormat::Type> RootBarycentricBuffer;
+	/* Strand hair roots translation and rotation in rest position relative to the bound triangle. Positions are relative to the rest root center */
+	TArray<FHairStrandsUniqueTriangleIndexFormat::Type> UniqueTriangleIndexBuffer;
+	TArray<FHairStrandsMeshTrianglePositionFormat::Type> RestUniqueTrianglePositionBuffer;
 
-		/* Strand hair roots translation and rotation in rest position relative to the bound triangle. Positions are relative to the rest root center */
-		TArray<FHairStrandsUniqueTriangleIndexFormat::Type> UniqueTriangleIndexBuffer;
-		TArray<FHairStrandsMeshTrianglePositionFormat::Type> RestUniqueTrianglePositionBuffer;
+	/* Number of samples used for the mesh interpolation */
+	uint32 SampleCount = 0;
 
-		/* Number of samples used for the mesh interpolation */
-		uint32 SampleCount = 0;
+	/* Store the hair interpolation weights | Size = SamplesCount * SamplesCount */
+	TArray<FHairStrandsWeightFormat::Type> MeshInterpolationWeightsBuffer;
 
-		/* Store the hair interpolation weights | Size = SamplesCount * SamplesCount */
-		TArray<FHairStrandsWeightFormat::Type> MeshInterpolationWeightsBuffer;
+	/* Store the samples vertex indices */
+	TArray<uint32> MeshSampleIndicesBuffer;
 
-		/* Store the samples vertex indices */
-		TArray<uint32> MeshSampleIndicesBuffer;
+	/* Store the samples rest positions */
+	TArray<FHairStrandsMeshTrianglePositionFormat::Type> RestSamplePositionsBuffer;
 
-		/* Store the samples rest positions */
-		TArray<FHairStrandsMeshTrianglePositionFormat::Type> RestSamplePositionsBuffer;
+	/* Store the samples rest section */
+	TArray<uint32> MeshSampleSectionsBuffer;
 
-		/* Store the samples rest section */
-		TArray<uint32> MeshSampleSectionsBuffer;
+	/* Store the mesh section indices which are relevant for this root LOD data */
+	TArray<uint32> UniqueSectionIds;
 
-		/* Store the mesh section indices which are relevant for this root LOD data */
-		TArray<uint32> UniqueSectionIds;
+	/* Mesh LOD index for which the root data are computed */
+	int32 LODIndex = -1;
 
-		/* Number of render section of the target mesh */
-		uint32 MeshSectionCount = 0;
-	};
+	/* Number of render section of the target mesh */
+	uint32 MeshSectionCount = 0;
 
 	/* Number of roots */
 	uint32 RootCount = 0;
 
 	/* Number of control points */
 	uint32 PointCount = 0;
-
-	/* Store the hair projection information for each mesh LOD */
-	TArray<FMeshProjectionLOD> MeshProjectionLODs;
 };
 
 /* Bulk data for root resources (GPU resources are stored into FHairStrandsRootResources) */
@@ -830,36 +823,19 @@ struct FHairStrandsRootBulkData : FHairStrandsBulkCommon
 	void Reset();
 	virtual void ResetLoadedSize() override;
 	bool IsValid() const { return Header.RootCount > 0; }
-	const TArray<uint32>& GetValidSectionIndices(int32 LODIndex) const;
-	uint32 GetLODCount() const { return Header.LODs.Num(); }
+	const TArray<uint32>& GetValidSectionIndices() const;
 	uint32 GetRootCount()const { return Header.RootCount; }
-
-	uint32 GetDataSize() const
-	{
-		uint32 Total = 0;
-		for (const FData::FLOD& LOD : Data.LODs)
-		{
-			Total += LOD.UniqueTriangleIndexBuffer.IsBulkDataLoaded() ?			LOD.UniqueTriangleIndexBuffer.GetBulkDataSize() : 0u;
-			Total += LOD.RootToUniqueTriangleIndexBuffer.IsBulkDataLoaded() ?	LOD.RootToUniqueTriangleIndexBuffer.GetBulkDataSize() : 0u;
-			Total += LOD.RootBarycentricBuffer.IsBulkDataLoaded() ?				LOD.RootBarycentricBuffer.GetBulkDataSize() : 0u;
-			Total += LOD.RestUniqueTrianglePositionBuffer.IsBulkDataLoaded() ?	LOD.RestUniqueTrianglePositionBuffer.GetBulkDataSize() : 0u;
-			Total += LOD.MeshInterpolationWeightsBuffer.IsBulkDataLoaded() ?	LOD.MeshInterpolationWeightsBuffer.GetBulkDataSize() : 0u;
-			Total += LOD.MeshSampleIndicesAndSectionsBuffer.IsBulkDataLoaded()?	LOD.MeshSampleIndicesAndSectionsBuffer.GetBulkDataSize() : 0u;
-			Total += LOD.RestSamplePositionsBuffer.IsBulkDataLoaded() ?			LOD.RestSamplePositionsBuffer.GetBulkDataSize() : 0u;
-		}
-		return Total;
-	}
+	uint32 GetDataSize() const;
 
 	struct FHeader
 	{		
-		struct FLOD
-		{
-			int32  LODIndex = -1;
-			uint32 UniqueTriangleCount = 0;			
-			uint32 SampleCount = 0; 				// Number of samples used for the mesh interpolation
-			uint32 MeshSectionCount = 0;			// Number of section of the target mesh
-			TArray<uint32> UniqueSectionIndices; 	// Store the mesh section indices which are relevant for this root LOD data
-		};
+		int32  LODIndex = -1;
+		uint32 RootCount = 0;					// Number of roots
+		uint32 PointCount = 0;					// Number of control points
+		uint32 SampleCount = 0; 				// Number of samples used for the mesh interpolation
+		uint32 MeshSectionCount = 0;			// Number of section of the target mesh
+		uint32 UniqueTriangleCount = 0;			
+		TArray<uint32> UniqueSectionIndices; 	// Store the mesh section indices which are relevant for this root LOD data
 
 		struct FStrides
 		{
@@ -871,30 +847,21 @@ struct FHairStrandsRootBulkData : FHairStrandsBulkCommon
 			uint32 MeshInterpolationWeightsBufferStride = 0;
 			uint32 MeshSampleIndicesAndSectionsBufferStride = 0;
 			uint32 RestSamplePositionsBufferStride = 0;
-		};
-
-		uint32 RootCount = 0;						// Number of roots
-		uint32 PointCount = 0;						// Number of control points
-		FStrides Strides;
-		TArray<FLOD> LODs;
+		} Strides;
 	} Header;
 
 	struct FData
 	{
-		struct FLOD
-		{
-			// Binding
-			FHairBulkContainer RootToUniqueTriangleIndexBuffer; // Map each root onto the unique triangle Id (per-root)
-			FHairBulkContainer RootBarycentricBuffer; 			// Root's barycentric (per-root)
-			FHairBulkContainer UniqueTriangleIndexBuffer; 		// Unique triangles list from skeleton mesh section IDs and triangle IDs (per-unique-triangle)
-			FHairBulkContainer RestUniqueTrianglePositionBuffer;// Rest triangle positions (per-unique-triangle)
-
-			// RBF
-			FHairBulkContainer MeshInterpolationWeightsBuffer; 		// Store the hair interpolation weights | Size = SamplesCount * SamplesCount (per-sample
-			FHairBulkContainer MeshSampleIndicesAndSectionsBuffer;	// Store the samples vertex indices (per-sample)
-			FHairBulkContainer RestSamplePositionsBuffer; 			// Store the samples rest positions (per-sample)
-		};
-		TArray<FLOD> LODs;
+		// Binding
+		FHairBulkContainer RootToUniqueTriangleIndexBuffer; 	// Map each root onto the unique triangle Id (per-root)
+		FHairBulkContainer RootBarycentricBuffer; 				// Root's barycentric (per-root)
+		FHairBulkContainer UniqueTriangleIndexBuffer; 			// Unique triangles list from skeleton mesh section IDs and triangle IDs (per-unique-triangle)
+		FHairBulkContainer RestUniqueTrianglePositionBuffer;	// Rest triangle positions (per-unique-triangle)
+
+		// RBF
+		FHairBulkContainer MeshInterpolationWeightsBuffer; 		// Store the hair interpolation weights | Size = SamplesCount * SamplesCount (per-sample
+		FHairBulkContainer MeshSampleIndicesAndSectionsBuffer;	// Store the samples vertex indices (per-sample)
+		FHairBulkContainer RestSamplePositionsBuffer; 			// Store the samples rest positions (per-sample)
 	} Data;
 };
 
diff --git a/Engine/Plugins/Runtime/HairStrands/Source/HairStrandsCore/Public/HairStrandsInterpolation.h b/Engine/Plugins/Runtime/HairStrands/Source/HairStrandsCore/Public/HairStrandsInterpolation.h
--- a/Engine/Plugins/Runtime/HairStrands/Source/HairStrandsCore/Public/HairStrandsInterpolation.h
+++ b/Engine/Plugins/Runtime/HairStrands/Source/HairStrandsCore/Public/HairStrandsInterpolation.h
@@ -145,7 +145,7 @@ void AddHairStrandsInterpolationPass(
 
 	const FRDGBufferSRVRef& RenDeformerPositionBuffer,
 	FRDGBufferUAVRef& OutRenPositionBuffer,
-	const FHairStrandsDeformedRootResource::FLOD::EFrameType DeformedFrame);
+	const FHairStrandsLODDeformedRootResource::EFrameType DeformedFrame);
 
 
 void AddHairTangentPass(
@@ -232,4 +232,5 @@ void AddHairCardsDeformationPass(
 	const ERHIFeatureLevel::Type FeatureLevel,
 	const FShaderPrintData* ShaderPrintData,
 	FHairGroupInstance* Instance,
+	const int32 HairLODIndex,
 	const int32 MeshLODIndex);
\ No newline at end of file
