Manually merged CL 30602124 from Release-28.20: Disable hair deformers on mobile and move hair deformers to it's own module.

diff --git a/Engine/Plugins/Runtime/HairStrands/Config/BaseHairStrands.ini b/Engine/Plugins/Runtime/HairStrands/Config/BaseHairStrands.ini
--- a/Engine/Plugins/Runtime/HairStrands/Config/BaseHairStrands.ini
+++ b/Engine/Plugins/Runtime/HairStrands/Config/BaseHairStrands.ini
@@ -7,6 +7,11 @@ fx.UseShaderStages=1
 +ClassRedirects=(OldName="HairStrandsActor",NewName="/Script/HairStrandsCore.GroomActor")
 +ClassRedirects=(OldName="HairStrandsAsset",NewName="/Script/HairStrandsCore.GroomAsset")
 +ClassRedirects=(OldName="HairStrandsComponent",NewName="/Script/HairStrandsCore.GroomComponent",ValueChanges=(("HairStrandsComponent0","GroomComponent0")))
++ClassRedirects=(OldName="/Script/HairStrandsCore.OptimusGroomComponentSource",NewName="/Script/HairStrandsDeformer.OptimusGroomComponentSource")
++ClassRedirects=(OldName="/Script/HairStrandsCore.OptimusGroomDataInterface",NewName="/Script/HairStrandsDeformer.OptimusGroomDataInterface")
++ClassRedirects=(OldName="/Script/HairStrandsCore.OptimusGroomExecDataInterface",NewName="/Script/HairStrandsDeformer.OptimusGroomExecDataInterface")
++ClassRedirects=(OldName="/Script/HairStrandsCore.OptimusGroomGuideDataInterface",NewName="/Script/HairStrandsDeformer.OptimusGroomGuideDataInterface")
++ClassRedirects=(OldName="/Script/HairStrandsCore.OptimusGroomWriteDataInterface",NewName="/Script/HairStrandsDeformer.OptimusGroomWriteDataInterface")
 +ClassRedirects=(OldName="/Script/HairStrandsEditor.GroomAssetImportData",NewName="/Script/HairStrandsCore.GroomAssetImportData")
 +ClassRedirects=(OldName="/Script/HairStrandsEditor.GroomImportOptions",NewName="/Script/HairStrandsCore.GroomImportOptions")
 +ClassRedirects=(OldName="/Script/HairStrandsNiagara.NiagaraDataInterfaceHairStrands",NewName="/Script/HairStrandsCore.NiagaraDataInterfaceHairStrands")
diff --git a/Engine/Plugins/Runtime/HairStrands/HairStrands.uplugin b/Engine/Plugins/Runtime/HairStrands/HairStrands.uplugin
--- a/Engine/Plugins/Runtime/HairStrands/HairStrands.uplugin
+++ b/Engine/Plugins/Runtime/HairStrands/HairStrands.uplugin
@@ -20,6 +20,15 @@
 			"Type": "Runtime",
 			"LoadingPhase": "PostConfigInit"
 		},
+		{
+			"Name": "HairStrandsDeformer",
+			"Type": "Runtime",
+			"LoadingPhase": "Default",
+			"PlatformDenyList": [
+				"Android",
+				"IOS"
+			]
+		},
 		{
 			"Name": "HairStrandsRuntime",
 			"Type": "Runtime",
diff --git a/Engine/Plugins/Runtime/HairStrands/Source/HairStrandsCore/HairStrandsCore.Build.cs b/Engine/Plugins/Runtime/HairStrands/Source/HairStrandsCore/HairStrandsCore.Build.cs
--- a/Engine/Plugins/Runtime/HairStrands/Source/HairStrandsCore/HairStrandsCore.Build.cs
+++ b/Engine/Plugins/Runtime/HairStrands/Source/HairStrandsCore/HairStrandsCore.Build.cs
@@ -23,9 +23,7 @@ namespace UnrealBuildTool.Rules
 					"Renderer",
 					"VectorVM",
 					"RHI",
-					"StaticMeshDescription",
-					"ComputeFramework",
-					"OptimusCore",
+					"StaticMeshDescription"
 				});
 
 			PrivateIncludePathModuleNames.AddRange(
diff --git a/Engine/Plugins/Runtime/HairStrands/Source/HairStrandsCore/Private/Deformer/DeformerDataInterfaceGroom.cpp b/Engine/Plugins/Runtime/HairStrands/Source/HairStrandsCore/Private/Deformer/DeformerDataInterfaceGroom.cpp
deleted file mode 100644
diff --git a/Engine/Plugins/Runtime/HairStrands/Source/HairStrandsCore/Private/Deformer/DeformerDataInterfaceGroom.h b/Engine/Plugins/Runtime/HairStrands/Source/HairStrandsCore/Private/Deformer/DeformerDataInterfaceGroom.h
deleted file mode 100644
diff --git a/Engine/Plugins/Runtime/HairStrands/Source/HairStrandsCore/Private/Deformer/DeformerDataInterfaceGroomExec.cpp b/Engine/Plugins/Runtime/HairStrands/Source/HairStrandsCore/Private/Deformer/DeformerDataInterfaceGroomExec.cpp
deleted file mode 100644
diff --git a/Engine/Plugins/Runtime/HairStrands/Source/HairStrandsCore/Private/Deformer/DeformerDataInterfaceGroomExec.h b/Engine/Plugins/Runtime/HairStrands/Source/HairStrandsCore/Private/Deformer/DeformerDataInterfaceGroomExec.h
deleted file mode 100644
diff --git a/Engine/Plugins/Runtime/HairStrands/Source/HairStrandsCore/Private/Deformer/DeformerDataInterfaceGroomGuide.cpp b/Engine/Plugins/Runtime/HairStrands/Source/HairStrandsCore/Private/Deformer/DeformerDataInterfaceGroomGuide.cpp
deleted file mode 100644
diff --git a/Engine/Plugins/Runtime/HairStrands/Source/HairStrandsCore/Private/Deformer/DeformerDataInterfaceGroomGuide.h b/Engine/Plugins/Runtime/HairStrands/Source/HairStrandsCore/Private/Deformer/DeformerDataInterfaceGroomGuide.h
deleted file mode 100644
diff --git a/Engine/Plugins/Runtime/HairStrands/Source/HairStrandsCore/Private/Deformer/DeformerDataInterfaceGroomWrite.cpp b/Engine/Plugins/Runtime/HairStrands/Source/HairStrandsCore/Private/Deformer/DeformerDataInterfaceGroomWrite.cpp
deleted file mode 100644
diff --git a/Engine/Plugins/Runtime/HairStrands/Source/HairStrandsCore/Private/Deformer/DeformerDataInterfaceGroomWrite.h b/Engine/Plugins/Runtime/HairStrands/Source/HairStrandsCore/Private/Deformer/DeformerDataInterfaceGroomWrite.h
deleted file mode 100644
diff --git a/Engine/Plugins/Runtime/HairStrands/Source/HairStrandsCore/Private/Deformer/DeformerGroomComponentSource.cpp b/Engine/Plugins/Runtime/HairStrands/Source/HairStrandsCore/Private/Deformer/DeformerGroomComponentSource.cpp
deleted file mode 100644
diff --git a/Engine/Plugins/Runtime/HairStrands/Source/HairStrandsCore/Private/Deformer/DeformerGroomComponentSource.h b/Engine/Plugins/Runtime/HairStrands/Source/HairStrandsCore/Private/Deformer/DeformerGroomComponentSource.h
deleted file mode 100644
diff --git a/Engine/Plugins/Runtime/HairStrands/Source/HairStrandsCore/Public/GroomResources.h b/Engine/Plugins/Runtime/HairStrands/Source/HairStrandsCore/Public/GroomResources.h
--- a/Engine/Plugins/Runtime/HairStrands/Source/HairStrandsCore/Public/GroomResources.h
+++ b/Engine/Plugins/Runtime/HairStrands/Source/HairStrandsCore/Public/GroomResources.h
@@ -277,9 +277,9 @@ struct FHairStrandsDeformedResource : public FHairCommonResource
 	//bool NeedsToUpdateTangent();
 
 	// Return deformer buffers
-	FRDGExternalBuffer& GetDeformerBuffer(FRDGBuilder& GraphBuilder);
-	FRDGExternalBuffer& GetDeformerCurveAttributeBuffer(FRDGBuilder& GraphBuilder);
-	FRDGExternalBuffer& GetDeformerPointAttributeBuffer(FRDGBuilder& GraphBuilder);
+	HAIRSTRANDSCORE_API FRDGExternalBuffer& GetDeformerBuffer(FRDGBuilder& GraphBuilder);
+	HAIRSTRANDSCORE_API FRDGExternalBuffer& GetDeformerCurveAttributeBuffer(FRDGBuilder& GraphBuilder);
+	HAIRSTRANDSCORE_API FRDGExternalBuffer& GetDeformerPointAttributeBuffer(FRDGBuilder& GraphBuilder);
 };
 
 
diff --git a/Engine/Plugins/Runtime/HairStrands/Source/HairStrandsCore/Public/HairStrandsInterpolation.h b/Engine/Plugins/Runtime/HairStrands/Source/HairStrandsCore/Public/HairStrandsInterpolation.h
--- a/Engine/Plugins/Runtime/HairStrands/Source/HairStrandsCore/Public/HairStrandsInterpolation.h
+++ b/Engine/Plugins/Runtime/HairStrands/Source/HairStrandsCore/Public/HairStrandsInterpolation.h
@@ -91,6 +91,8 @@ void AddGroomCacheUpdatePass(
 
 HAIRSTRANDSCORE_API void ComputeInterpolationWeights(UGroomBindingAsset* BindingAsset, FSkeletalMeshRenderData* TargetRenderData, TArray<FRWBuffer>& TransferedPositions);
 
+HAIRSTRANDSCORE_API FHairGroupPublicData::FVertexFactoryInput ComputeHairStrandsVertexInputData(const FHairGroupInstance* Instance, EGroomViewMode ViewMode);
+
 struct FRDGHairStrandsCullingData
 {
 	bool bCullingResultAvailable = false;
diff --git a/Engine/Plugins/Runtime/HairStrands/Source/HairStrandsDeformer/HairStrandsDeformer.Build.cs b/Engine/Plugins/Runtime/HairStrands/Source/HairStrandsDeformer/HairStrandsDeformer.Build.cs
new file mode 100644
--- /dev/null
+++ b/Engine/Plugins/Runtime/HairStrands/Source/HairStrandsDeformer/HairStrandsDeformer.Build.cs
@@ -0,0 +1,32 @@
+// Copyright Epic Games, Inc. All Rights Reserved.
+
+namespace UnrealBuildTool.Rules
+{
+	public class HairStrandsDeformer : ModuleRules
+	{
+		public HairStrandsDeformer(ReadOnlyTargetRules Target) : base(Target)
+		{
+			PrivateDependencyModuleNames.AddRange(
+				new string[]
+				{
+					"Core",
+					"CoreUObject",
+					"Engine",
+					"NiagaraCore",
+					"Niagara",
+					"RenderCore",
+					"Renderer",
+					"RHI",
+					"ComputeFramework",
+					"OptimusCore",
+					"HairStrandsCore"
+				});
+
+			PrivateIncludePathModuleNames.AddRange(
+				new string[]
+				{
+					"Shaders",
+				});
+		}
+	}
+}
diff --git a/Engine/Plugins/Runtime/HairStrands/Source/HairStrandsDeformer/Private/DeformerDataInterfaceGroom.cpp b/Engine/Plugins/Runtime/HairStrands/Source/HairStrandsDeformer/Private/DeformerDataInterfaceGroom.cpp
new file mode 100644
--- /dev/null
+++ b/Engine/Plugins/Runtime/HairStrands/Source/HairStrandsDeformer/Private/DeformerDataInterfaceGroom.cpp
@@ -0,0 +1,279 @@
+// Copyright Epic Games, Inc. All Rights Reserved.
+
+#include "DeformerDataInterfaceGroom.h"
+
+#include "ComputeFramework/ShaderParamTypeDefinition.h"
+#include "ComputeFramework/ComputeKernelPermutationVector.h"
+#include "OptimusDataDomain.h"
+#include "ShaderCompilerCore.h"
+#include "ShaderParameterMetadataBuilder.h"
+#include "GroomComponent.h"
+#include "GroomInstance.h"
+#include "RenderGraphBuilder.h"
+#include "GlobalRenderResources.h"
+#include "HairStrandsDefinitions.h"
+#include "DeformerGroomComponentSource.h"
+#include "HairStrandsInterpolation.h"
+
+FString UOptimusGroomDataInterface::GetDisplayName() const
+{
+	return TEXT("Groom");
+}
+
+TArray<FOptimusCDIPinDefinition> UOptimusGroomDataInterface::GetPinDefinitions() const
+{
+	FName ControlPoint(UOptimusGroomComponentSource::Domains::ControlPoint);
+	FName Curve(UOptimusGroomComponentSource::Domains::Curve);
+
+	TArray<FOptimusCDIPinDefinition> Defs;
+	Defs.Add({ "NumControlPoints", "ReadNumControlPoints" });
+	Defs.Add({ "NumCurves",        "ReadNumCurves" });
+	Defs.Add({ "Position",         "ReadPosition",          ControlPoint,   "ReadPosition" });
+	Defs.Add({ "Radius",           "ReadRadius",            ControlPoint,   "ReadRadius" });
+	Defs.Add({ "CoordU",           "ReadCoordU",            ControlPoint,   "ReadCoordU" });
+	Defs.Add({ "Length",           "ReadLength",            ControlPoint,   "ReadLength" });
+	Defs.Add({ "RootUV",           "ReadRootUV",            ControlPoint,   "ReadRootUV" });
+	Defs.Add({ "Seed",             "ReadSeed",              ControlPoint,   "ReadSeed" });
+	Defs.Add({ "ClumpId",          "ReadClumpId",           ControlPoint,   "ReadClumpId" });
+	Defs.Add({ "Color",            "ReadColor",             ControlPoint,   "ReadColor" });
+	Defs.Add({ "Roughness",        "ReadRoughness",         ControlPoint,   "ReadRoughness" });
+	Defs.Add({ "AO",               "ReadAO",                ControlPoint,   "ReadAO" });
+	Defs.Add({ "CurveOffsetPoint", "ReadCurveOffsetPoint",  Curve,          "ReadCurveOffsetPoint" });
+	Defs.Add({ "CurveNumPoint",    "ReadCurveNumPoint",     Curve,          "ReadCurveNumPoint" });
+	Defs.Add({ "GuideIndex",       "ReadGuideIndex",        ControlPoint,   "ReadGuideIndex" });
+	return Defs;
+}
+ 
+TSubclassOf<UActorComponent> UOptimusGroomDataInterface::GetRequiredComponentClass() const
+{
+	return UGroomComponent::StaticClass();
+}
+
+void UOptimusGroomDataInterface::GetSupportedInputs(TArray<FShaderFunctionDefinition>& OutFunctions) const
+{
+	OutFunctions.AddDefaulted_GetRef()
+		.SetName(TEXT("ReadNumControlPoints"))
+		.AddReturnType(EShaderFundamentalType::Uint);
+
+	OutFunctions.AddDefaulted_GetRef()
+		.SetName(TEXT("ReadNumCurves"))
+		.AddReturnType(EShaderFundamentalType::Uint);
+
+	OutFunctions.AddDefaulted_GetRef()
+		.SetName(TEXT("ReadPosition"))
+		.AddReturnType(EShaderFundamentalType::Float, 3)
+		.AddParam(EShaderFundamentalType::Uint);
+
+	OutFunctions.AddDefaulted_GetRef()
+		.SetName(TEXT("ReadRadius"))
+		.AddReturnType(EShaderFundamentalType::Float)
+		.AddParam(EShaderFundamentalType::Uint);
+
+	OutFunctions.AddDefaulted_GetRef()
+		.SetName(TEXT("ReadCoordU"))
+		.AddReturnType(EShaderFundamentalType::Float)
+		.AddParam(EShaderFundamentalType::Uint);
+
+	OutFunctions.AddDefaulted_GetRef()
+		.SetName(TEXT("ReadLength"))
+		.AddReturnType(EShaderFundamentalType::Float)
+		.AddParam(EShaderFundamentalType::Uint);
+
+	OutFunctions.AddDefaulted_GetRef()
+		.SetName(TEXT("ReadRootUV"))
+		.AddReturnType(EShaderFundamentalType::Float, 2)
+		.AddParam(EShaderFundamentalType::Uint);
+
+	OutFunctions.AddDefaulted_GetRef()
+		.SetName(TEXT("ReadSeed"))
+		.AddReturnType(EShaderFundamentalType::Float)
+		.AddParam(EShaderFundamentalType::Uint);
+
+	OutFunctions.AddDefaulted_GetRef()
+		.SetName(TEXT("ReadClumpId"))
+		.AddReturnType(EShaderFundamentalType::Uint)
+		.AddParam(EShaderFundamentalType::Uint);
+
+	OutFunctions.AddDefaulted_GetRef()
+		.SetName(TEXT("ReadColor"))
+		.AddReturnType(EShaderFundamentalType::Float, 3)
+		.AddParam(EShaderFundamentalType::Uint);
+
+	OutFunctions.AddDefaulted_GetRef()
+		.SetName(TEXT("ReadRoughnes"))
+		.AddReturnType(EShaderFundamentalType::Float)
+		.AddParam(EShaderFundamentalType::Uint);
+
+	OutFunctions.AddDefaulted_GetRef()
+		.SetName(TEXT("ReadAO"))
+		.AddReturnType(EShaderFundamentalType::Float)
+		.AddParam(EShaderFundamentalType::Uint);
+
+	OutFunctions.AddDefaulted_GetRef()
+		.SetName(TEXT("ReadCurveOffsetPoint"))
+		.AddReturnType(EShaderFundamentalType::Uint)
+		.AddParam(EShaderFundamentalType::Uint);
+
+	OutFunctions.AddDefaulted_GetRef()
+		.SetName(TEXT("ReadCurveNumPoint"))
+		.AddReturnType(EShaderFundamentalType::Uint)
+		.AddParam(EShaderFundamentalType::Uint);
+
+	OutFunctions.AddDefaulted_GetRef()
+		.SetName(TEXT("ReadGuideIndex"))
+		.AddReturnType(EShaderFundamentalType::Uint)
+		.AddParam(EShaderFundamentalType::Uint);
+}
+
+BEGIN_SHADER_PARAMETER_STRUCT(FGroomDataInterfaceParameters, )
+	SHADER_PARAMETER_STRUCT_INCLUDE(FHairStrandsInstanceCommonParameters, Common)
+	SHADER_PARAMETER_STRUCT_INCLUDE(FHairStrandsInstanceResourceParameters, Resources)
+	SHADER_PARAMETER_STRUCT_INCLUDE(FHairStrandsInstanceInterpolationParameters, Interpolation)
+END_SHADER_PARAMETER_STRUCT()
+
+void UOptimusGroomDataInterface::GetShaderParameters(TCHAR const* UID, FShaderParametersMetadataBuilder& InOutBuilder, FShaderParametersMetadataAllocations& InOutAllocations) const
+{
+	InOutBuilder.AddNestedStruct<FGroomDataInterfaceParameters>(UID);
+}
+
+TCHAR const* UOptimusGroomDataInterface::TemplateFilePath = TEXT("/Plugin/Runtime/HairStrands/Private/DeformerDataInterfaceGroom.ush");
+
+TCHAR const* UOptimusGroomDataInterface::GetShaderVirtualPath() const
+{
+	return TemplateFilePath;
+}
+
+void UOptimusGroomDataInterface::GetShaderHash(FString& InOutKey) const
+{
+	GetShaderFileHash(TemplateFilePath, EShaderPlatform::SP_PCD3D_SM5).AppendString(InOutKey);
+}
+
+void UOptimusGroomDataInterface::GetHLSL(FString& OutHLSL, FString const& InDataInterfaceName) const
+{
+	TMap<FString, FStringFormatArg> TemplateArgs =
+	{
+		{ TEXT("DataInterfaceName"), InDataInterfaceName },
+	};
+
+	FString TemplateFile;
+	LoadShaderSourceFile(TemplateFilePath, EShaderPlatform::SP_PCD3D_SM5, &TemplateFile, nullptr);
+	OutHLSL += FString::Format(*TemplateFile, TemplateArgs);
+}
+
+UComputeDataProvider* UOptimusGroomDataInterface::CreateDataProvider(TObjectPtr<UObject> InBinding, uint64 InInputMask, uint64 InOutputMask) const
+{
+	UOptimusGroomDataProvider* Provider = NewObject<UOptimusGroomDataProvider>();
+	Provider->Groom = Cast<UGroomComponent>(InBinding);
+	return Provider;
+}
+
+
+FComputeDataProviderRenderProxy* UOptimusGroomDataProvider::GetRenderProxy()
+{
+	return new FOptimusGroomDataProviderProxy(Groom);
+}
+
+
+FOptimusGroomDataProviderProxy::FOptimusGroomDataProviderProxy(UGroomComponent* InGroomComponent)
+{
+	GroomComponent = InGroomComponent;
+}
+
+bool FOptimusGroomDataProviderProxy::IsValid(FValidationData const& InValidationData) const
+{
+	if (InValidationData.ParameterStructSize != sizeof(FParameters))
+	{
+		return false;
+	}
+	if (GroomComponent == nullptr)
+	{
+		return false;
+	}
+	const uint32 GroupCount = GroomComponent->GetGroupCount();
+	if (InValidationData.NumInvocations != GroupCount)
+	{
+		return false;
+	}
+	// Earlier invalidation in case one of the instance is not ready/valid
+	for (uint32 GroupIndex = 0; GroupIndex < GroupCount; ++GroupIndex)
+	{
+		if (GroomComponent->GetGroupInstance(GroupIndex) == nullptr)
+		{
+			return false;
+		}
+	}
+	
+	return true;
+}
+
+void FOptimusGroomDataProviderProxy::AllocateResources(FRDGBuilder& GraphBuilder)
+{
+	Resources.Empty();
+	FallbackSRV = nullptr;
+	const uint32 InstanceCount = GroomComponent ? GroomComponent->GetGroupCount() : 0;
+	for (uint32 Index =0; Index <InstanceCount;++Index)
+	{
+		if (FHairGroupInstance* Instance = GroomComponent->GetGroupInstance(Index))
+		{
+			{
+				FHairStrandsInstanceResourceParameters& R = Resources.AddDefaulted_GetRef();
+				R.PositionBuffer		= RegisterAsSRV(GraphBuilder, Instance->Strands.RestResource->PositionBuffer);
+				R.PositionOffsetBuffer 	= RegisterAsSRV(GraphBuilder, Instance->Strands.RestResource->PositionOffsetBuffer);
+				R.CurveBuffer			= RegisterAsSRV(GraphBuilder, Instance->Strands.RestResource->CurveBuffer);
+				R.PointToCurveBuffer	= RegisterAsSRV(GraphBuilder, Instance->Strands.RestResource->PointToCurveBuffer);
+				R.CurveAttributeBuffer	= RegisterAsSRV(GraphBuilder, Instance->Strands.RestResource->CurveAttributeBuffer);
+				R.PointAttributeBuffer	= RegisterAsSRV(GraphBuilder, Instance->Strands.RestResource->PointAttributeBuffer);
+			}
+
+			{
+				FHairStrandsInstanceInterpolationParameters& R = Interpolations.AddDefaulted_GetRef();
+				R.CurveInterpolationBuffer = Instance->Strands.InterpolationResource ? RegisterAsSRV(GraphBuilder, Instance->Strands.InterpolationResource->CurveInterpolationBuffer) : nullptr;
+				R.PointInterpolationBuffer = Instance->Strands.InterpolationResource ? RegisterAsSRV(GraphBuilder, Instance->Strands.InterpolationResource->PointInterpolationBuffer) : nullptr;
+			}
+
+			if (!FallbackSRV)
+			{
+				FallbackSRV = GraphBuilder.CreateSRV(GraphBuilder.RegisterExternalBuffer(GWhiteVertexBufferWithRDG->Buffer), PF_R16G16B16A16_UINT); // This should be a ByteAddressBuffer
+			}
+		}
+	}
+}
+
+void FOptimusGroomDataProviderProxy::GatherDispatchData(FDispatchData const& InDispatchData)
+{
+	const uint32 InstanceCount = GroomComponent ? GroomComponent->GetGroupCount() : 0;
+	check(InDispatchData.NumInvocations == InstanceCount);
+	check(uint32(Resources.Num()) == InstanceCount);
+
+	const TStridedView<FParameters> ParameterArray = MakeStridedParameterView<FParameters>(InDispatchData);
+	for (int32 InvocationIndex = 0; InvocationIndex < ParameterArray.Num(); ++InvocationIndex)
+	{
+		// Default valid in case the instance is not valid
+		FParameters& Parameters = ParameterArray[InvocationIndex];
+		Parameters.Common.PointCount = 0;
+		Parameters.Common.CurveCount = 0;
+		Parameters.Resources.PositionBuffer			= FallbackSRV;
+		Parameters.Resources.PositionOffsetBuffer 	= FallbackSRV;
+		Parameters.Resources.CurveAttributeBuffer	= FallbackSRV;
+		Parameters.Resources.PointAttributeBuffer	= FallbackSRV;
+		Parameters.Resources.PointToCurveBuffer		= FallbackSRV;
+		Parameters.Resources.CurveBuffer			= FallbackSRV;
+		Parameters.Interpolation.CurveInterpolationBuffer = FallbackSRV;
+		Parameters.Interpolation.PointInterpolationBuffer = FallbackSRV;
+
+		if (FHairGroupInstance* Instance = GroomComponent->GetGroupInstance(InvocationIndex))
+		{
+			const FHairGroupPublicData::FVertexFactoryInput VFInput = ComputeHairStrandsVertexInputData(Instance, EGroomViewMode::None);
+			Parameters.Common = VFInput.Strands.Common;
+
+			if (Resources[InvocationIndex].PositionBuffer != nullptr)
+			{
+				Parameters.Resources = Resources[InvocationIndex];
+			}
+			if (Interpolations[InvocationIndex].CurveInterpolationBuffer != nullptr)
+			{
+				Parameters.Interpolation = Interpolations[InvocationIndex];
+			}
+		}
+	}
+}
\ No newline at end of file
diff --git a/Engine/Plugins/Runtime/HairStrands/Source/HairStrandsDeformer/Private/DeformerDataInterfaceGroom.h b/Engine/Plugins/Runtime/HairStrands/Source/HairStrandsDeformer/Private/DeformerDataInterfaceGroom.h
new file mode 100644
--- /dev/null
+++ b/Engine/Plugins/Runtime/HairStrands/Source/HairStrandsDeformer/Private/DeformerDataInterfaceGroom.h
@@ -0,0 +1,75 @@
+// Copyright Epic Games, Inc. All Rights Reserved.
+
+#pragma once
+
+#include "RenderGraphDefinitions.h"
+#include "OptimusComputeDataInterface.h"
+#include "HairStrandsInterface.h"
+#include "ComputeFramework/ComputeDataProvider.h"
+#include "DeformerDataInterfaceGroom.generated.h"
+
+class FGroomDataInterfaceParameters;
+class UGroomComponent;
+
+/** Compute Framework Data Interface for reading groom. */
+UCLASS(Category = ComputeFramework)
+class UOptimusGroomDataInterface : public UOptimusComputeDataInterface
+{
+	GENERATED_BODY()
+
+public:
+	//~ Begin UOptimusComputeDataInterface Interface
+	FString GetDisplayName() const override;
+	TArray<FOptimusCDIPinDefinition> GetPinDefinitions() const override;
+	TSubclassOf<UActorComponent> GetRequiredComponentClass() const override;
+	//~ End UOptimusComputeDataInterface Interface
+	
+	//~ Begin UComputeDataInterface Interface
+	TCHAR const* GetClassName() const override { return TEXT("Groom"); }
+	void GetSupportedInputs(TArray<FShaderFunctionDefinition>& OutFunctions) const override;
+	void GetShaderParameters(TCHAR const* UID, FShaderParametersMetadataBuilder& InOutBuilder, FShaderParametersMetadataAllocations& InOutAllocations) const override;
+	TCHAR const* GetShaderVirtualPath() const override;
+	void GetShaderHash(FString& InOutKey) const override;
+	void GetHLSL(FString& OutHLSL, FString const& InDataInterfaceName) const override;
+	UComputeDataProvider* CreateDataProvider(TObjectPtr<UObject> InBinding, uint64 InInputMask, uint64 InOutputMask) const override;
+	//~ End UComputeDataInterface Interface
+
+private:
+	static TCHAR const* TemplateFilePath;
+};
+
+/** Compute Framework Data Provider for reading groom. */
+UCLASS(BlueprintType, editinlinenew, Category = ComputeFramework)
+class UOptimusGroomDataProvider : public UComputeDataProvider
+{
+	GENERATED_BODY()
+
+public:
+	UPROPERTY(EditAnywhere, BlueprintReadWrite, Category = Binding)
+	TObjectPtr<UGroomComponent> Groom = nullptr;
+
+	//~ Begin UComputeDataProvider Interface
+	FComputeDataProviderRenderProxy* GetRenderProxy() override;
+	//~ End UComputeDataProvider Interface
+};
+
+class FOptimusGroomDataProviderProxy : public FComputeDataProviderRenderProxy
+{
+public:
+	FOptimusGroomDataProviderProxy(UGroomComponent* InGroomComponent);
+
+	//~ Begin FComputeDataProviderRenderProxy Interface
+	bool IsValid(FValidationData const& InValidationData) const override;
+	void AllocateResources(FRDGBuilder& GraphBuilder) override;
+	void GatherDispatchData(FDispatchData const& InDispatchData) override;
+	//~ End FComputeDataProviderRenderProxy Interface
+
+private:
+	using FParameters = FGroomDataInterfaceParameters;
+
+	UGroomComponent* GroomComponent = nullptr; // Should it be HairInstance instead?
+	TArray<FHairStrandsInstanceResourceParameters> Resources;
+	TArray<FHairStrandsInstanceInterpolationParameters> Interpolations;
+	FRDGBufferSRVRef FallbackSRV;
+
+};
diff --git a/Engine/Plugins/Runtime/HairStrands/Source/HairStrandsDeformer/Private/DeformerDataInterfaceGroomExec.cpp b/Engine/Plugins/Runtime/HairStrands/Source/HairStrandsDeformer/Private/DeformerDataInterfaceGroomExec.cpp
new file mode 100644
--- /dev/null
+++ b/Engine/Plugins/Runtime/HairStrands/Source/HairStrandsDeformer/Private/DeformerDataInterfaceGroomExec.cpp
@@ -0,0 +1,159 @@
+// Copyright Epic Games, Inc. All Rights Reserved.
+
+#include "DeformerDataInterfaceGroomExec.h"
+
+#include "DeformerGroomComponentSource.h"
+#include "GroomComponent.h"
+#include "ComputeFramework/ShaderParamTypeDefinition.h"
+#include "ShaderCompilerCore.h"
+#include "ShaderParameterMetadataBuilder.h"
+#include "SkeletalRenderPublic.h"
+
+FString UOptimusGroomExecDataInterface::GetDisplayName() const
+{
+	return TEXT("Execute Groom");
+}
+
+FName UOptimusGroomExecDataInterface::GetCategory() const
+{
+	return CategoryName::ExecutionDataInterfaces;
+}
+
+TArray<FOptimusCDIPinDefinition> UOptimusGroomExecDataInterface::GetPinDefinitions() const
+{
+	TArray<FOptimusCDIPinDefinition> Defs;
+	Defs.Add({ "NumThreads", "ReadNumThreads" });
+	return Defs;
+}
+
+
+TSubclassOf<UActorComponent> UOptimusGroomExecDataInterface::GetRequiredComponentClass() const
+{
+	return UGroomComponent::StaticClass();
+}
+
+
+void UOptimusGroomExecDataInterface::GetSupportedInputs(TArray<FShaderFunctionDefinition>& OutFunctions) const
+{
+	OutFunctions.AddDefaulted_GetRef()
+		.SetName(TEXT("ReadNumThreads"))
+		.AddReturnType(EShaderFundamentalType::Int, 3);
+}
+
+
+BEGIN_SHADER_PARAMETER_STRUCT(FGroomExecDataInterfaceParameters, )
+	SHADER_PARAMETER(FIntVector, NumThreads)
+END_SHADER_PARAMETER_STRUCT()
+
+void UOptimusGroomExecDataInterface::GetShaderParameters(TCHAR const* UID, FShaderParametersMetadataBuilder& InOutBuilder, FShaderParametersMetadataAllocations& InOutAllocations) const
+{
+	InOutBuilder.AddNestedStruct<FGroomExecDataInterfaceParameters>(UID);
+}
+
+TCHAR const* UOptimusGroomExecDataInterface::TemplateFilePath = TEXT("/Plugin/Runtime/HairStrands/Private/DeformerDataInterfaceGroomExec.ush");
+
+TCHAR const* UOptimusGroomExecDataInterface::GetShaderVirtualPath() const
+{
+	return TemplateFilePath;
+}
+
+void UOptimusGroomExecDataInterface::GetShaderHash(FString& InOutKey) const
+{
+	GetShaderFileHash(TemplateFilePath, EShaderPlatform::SP_PCD3D_SM5).AppendString(InOutKey);
+}
+
+void UOptimusGroomExecDataInterface::GetHLSL(FString& OutHLSL, FString const& InDataInterfaceName) const
+{
+	TMap<FString, FStringFormatArg> TemplateArgs =
+	{
+		{ TEXT("DataInterfaceName"), InDataInterfaceName },
+	};
+
+	FString TemplateFile;
+	LoadShaderSourceFile(TemplateFilePath, EShaderPlatform::SP_PCD3D_SM5, &TemplateFile, nullptr);
+	OutHLSL += FString::Format(*TemplateFile, TemplateArgs);
+}
+
+UComputeDataProvider* UOptimusGroomExecDataInterface::CreateDataProvider(TObjectPtr<UObject> InBinding, uint64 InInputMask, uint64 InOutputMask) const
+{
+	UOptimusGroomExecDataProvider* Provider = NewObject<UOptimusGroomExecDataProvider>();
+	Provider->GroomComponent = Cast<UGroomComponent>(InBinding);
+	Provider->Domain = Domain;
+	return Provider;
+}
+
+FName UOptimusGroomExecDataInterface::GetSelectedExecutionDomainName() const
+{
+	if (Domain == EOptimusGroomExecDomain::ControlPoint)
+	{
+		return UOptimusGroomComponentSource::Domains::ControlPoint;
+	}
+
+	if (Domain == EOptimusGroomExecDomain::Curve)
+	{
+		return UOptimusGroomComponentSource::Domains::Curve;
+	}
+
+	return NAME_None;
+}
+
+
+FComputeDataProviderRenderProxy* UOptimusGroomExecDataProvider::GetRenderProxy()
+{
+	return new FOptimusGroomExecDataProviderProxy(GroomComponent, Domain);
+}
+
+
+FOptimusGroomExecDataProviderProxy::FOptimusGroomExecDataProviderProxy(UGroomComponent* InGroomComponent, EOptimusGroomExecDomain InDomain)
+{
+	GroomComponent = InGroomComponent;
+	Domain = InDomain;
+}
+
+int32 FOptimusGroomExecDataProviderProxy::GetDispatchThreadCount(TArray<FIntVector>& ThreadCounts) const
+{
+	const int32 NumInvocations = GroomComponent ? GroomComponent->GetGroupCount() : 0;
+
+	ThreadCounts.Reset();
+	ThreadCounts.Reserve(NumInvocations);
+	for (int32 InvocationIndex = 0; InvocationIndex < NumInvocations; ++InvocationIndex)
+	{		
+		if (FHairGroupInstance* Instance = GroomComponent->GetGroupInstance(InvocationIndex))
+		{
+			const int32 NumControlPoints = Instance->Strands.GetData().GetNumPoints();
+			const int32 NumCurves = Instance->Strands.GetData().GetNumCurves();
+			const int32 NumThreads = Domain == EOptimusGroomExecDomain::ControlPoint ? NumControlPoints : NumCurves;
+			ThreadCounts.Add(FIntVector(NumThreads, 1, 1));
+		}
+	}
+	return NumInvocations;
+}
+
+bool FOptimusGroomExecDataProviderProxy::IsValid(FValidationData const& InValidationData) const
+{
+	if (InValidationData.ParameterStructSize != sizeof(FParameters))
+	{
+		return false;
+	}
+	if (GroomComponent == nullptr || GroomComponent->GetGroupCount() == 0)
+	{
+		return false;
+	}
+
+	return true;
+}
+
+void FOptimusGroomExecDataProviderProxy::GatherDispatchData(FDispatchData const& InDispatchData)
+{
+	const TStridedView<FParameters> ParameterArray = MakeStridedParameterView<FParameters>(InDispatchData);
+	for (int32 InvocationIndex = 0; InvocationIndex < ParameterArray.Num(); ++InvocationIndex)
+	{
+		FHairGroupInstance* Instance = GroomComponent->GetGroupInstance(InvocationIndex);
+		const int32 NumControlPoints = Instance->Strands.GetData().GetNumPoints();
+		const int32 NumCurves = Instance->Strands.GetData().GetNumCurves();
+		const int32 NumThreads = Domain == EOptimusGroomExecDomain::ControlPoint ? NumControlPoints : NumCurves;
+
+		FParameters& Parameters = ParameterArray[InvocationIndex];
+		Parameters.NumThreads = FIntVector(NumThreads, 1, 1);
+	}
+}
diff --git a/Engine/Plugins/Runtime/HairStrands/Source/HairStrandsDeformer/Private/DeformerDataInterfaceGroomExec.h b/Engine/Plugins/Runtime/HairStrands/Source/HairStrandsDeformer/Private/DeformerDataInterfaceGroomExec.h
new file mode 100644
--- /dev/null
+++ b/Engine/Plugins/Runtime/HairStrands/Source/HairStrandsDeformer/Private/DeformerDataInterfaceGroomExec.h
@@ -0,0 +1,96 @@
+// Copyright Epic Games, Inc. All Rights Reserved.
+
+#pragma once
+
+#include "IOptimusDeprecatedExecutionDataInterface.h"
+#include "OptimusComputeDataInterface.h"
+#include "ComputeFramework/ComputeDataProvider.h"
+#include "DeformerDataInterfaceGroomExec.generated.h"
+
+class FGroomExecDataInterfaceParameters;
+class UGroomComponent;
+
+UENUM()
+enum class EOptimusGroomExecDomain : uint8
+{
+	None = 0 UMETA(Hidden),
+	/** Run kernel with one thread per control point. */
+	ControlPoint = 1,
+	/** Run kernel with one thread per curve. */
+	Curve,
+};
+
+/** Compute Framework Data Interface for executing kernels over a skinned mesh domain. */
+UCLASS(Category = ComputeFramework)
+class UOptimusGroomExecDataInterface :
+	public UOptimusComputeDataInterface,
+	public IOptimusDeprecatedExecutionDataInterface
+{
+	GENERATED_BODY()
+
+public:
+	//~ Begin UOptimusComputeDataInterface Interface
+	FString GetDisplayName() const override;
+	FName GetCategory() const override;
+	bool IsVisible() const override {return false;};
+	TArray<FOptimusCDIPinDefinition> GetPinDefinitions() const override;
+	TSubclassOf<UActorComponent> GetRequiredComponentClass() const override;
+	//~ End UOptimusComputeDataInterface Interface
+	
+	//~ Begin UComputeDataInterface Interface
+	TCHAR const* GetClassName() const override { return TEXT("GroomExec"); }
+	bool IsExecutionInterface() const override { return true; }
+	void GetSupportedInputs(TArray<FShaderFunctionDefinition>& OutFunctions) const override;
+	void GetShaderParameters(TCHAR const* UID, FShaderParametersMetadataBuilder& InOutBuilder, FShaderParametersMetadataAllocations& InOutAllocations) const override;
+	TCHAR const* GetShaderVirtualPath() const override;
+	void GetShaderHash(FString& InOutKey) const override;
+	void GetHLSL(FString& OutHLSL, FString const& InDataInterfaceName) const override;
+	UComputeDataProvider* CreateDataProvider(TObjectPtr<UObject> InBinding, uint64 InInputMask, uint64 InOutputMask) const override;
+	//~ End UComputeDataInterface Interface
+
+	//~ Begin IOptimusDeprecatedExecutionDataInterface Interface
+	FName GetSelectedExecutionDomainName() const override;
+	//~ End IOptimusDeprecatedExecutionDataInterface Interface
+	
+	UPROPERTY(EditAnywhere, Category = Execution)
+	EOptimusGroomExecDomain Domain = EOptimusGroomExecDomain::ControlPoint;
+
+private:
+	static TCHAR const* TemplateFilePath;
+};
+
+/** Compute Framework Data Provider for executing kernels over a skinned mesh domain. */
+UCLASS(BlueprintType, editinlinenew, Category = ComputeFramework)
+class UOptimusGroomExecDataProvider : public UComputeDataProvider
+{
+	GENERATED_BODY()
+
+public:
+	UPROPERTY()
+	TObjectPtr<UGroomComponent> GroomComponent = nullptr;
+
+	UPROPERTY()
+	EOptimusGroomExecDomain Domain = EOptimusGroomExecDomain::ControlPoint;
+
+	//~ Begin UComputeDataProvider Interface
+	FComputeDataProviderRenderProxy* GetRenderProxy() override;
+	//~ End UComputeDataProvider Interface
+};
+
+class FOptimusGroomExecDataProviderProxy : public FComputeDataProviderRenderProxy
+{
+public:
+	FOptimusGroomExecDataProviderProxy(UGroomComponent* InGroomComponent, EOptimusGroomExecDomain InDomain);
+
+	//~ Begin FComputeDataProviderRenderProxy Interface
+	int32 GetDispatchThreadCount(TArray<FIntVector>& ThreadCounts) const override;
+	bool IsValid(FValidationData const& InValidationData) const override;
+	void GatherDispatchData(FDispatchData const& InDispatchData) override;
+	//~ End FComputeDataProviderRenderProxy Interface
+
+private:
+	using FParameters = FGroomExecDataInterfaceParameters;
+
+	UGroomComponent* GroomComponent = nullptr;
+	EOptimusGroomExecDomain Domain = EOptimusGroomExecDomain::ControlPoint;
+};
diff --git a/Engine/Plugins/Runtime/HairStrands/Source/HairStrandsDeformer/Private/DeformerDataInterfaceGroomGuide.cpp b/Engine/Plugins/Runtime/HairStrands/Source/HairStrandsDeformer/Private/DeformerDataInterfaceGroomGuide.cpp
new file mode 100644
--- /dev/null
+++ b/Engine/Plugins/Runtime/HairStrands/Source/HairStrandsDeformer/Private/DeformerDataInterfaceGroomGuide.cpp
@@ -0,0 +1,206 @@
+// Copyright Epic Games, Inc. All Rights Reserved.
+
+#include "DeformerDataInterfaceGroomGuide.h"
+
+#include "ComputeFramework/ShaderParamTypeDefinition.h"
+#include "ComputeFramework/ComputeKernelPermutationVector.h"
+#include "OptimusDataDomain.h"
+#include "ShaderCompilerCore.h"
+#include "ShaderParameterMetadataBuilder.h"
+#include "GroomComponent.h"
+#include "GroomInstance.h"
+#include "RenderGraphBuilder.h"
+#include "GlobalRenderResources.h"
+#include "HairStrandsDefinitions.h"
+#include "DeformerGroomComponentSource.h"
+#include "HairStrandsInterpolation.h"
+
+FString UOptimusGroomGuideDataInterface::GetDisplayName() const
+{
+	return TEXT("Guides");
+}
+
+TArray<FOptimusCDIPinDefinition> UOptimusGroomGuideDataInterface::GetPinDefinitions() const
+{
+	FName ControlPoint(UOptimusGroomComponentSource::Domains::ControlPoint);
+	FName Curve(UOptimusGroomComponentSource::Domains::Curve);
+
+	TArray<FOptimusCDIPinDefinition> Defs;
+	Defs.Add({ "NumGuidePoints",   "ReadNumGuidePoints" });
+	Defs.Add({ "NumGuideCurves",   "ReadNumGuideCurves" });
+	Defs.Add({ "Position",         "ReadPosition",          ControlPoint,   "ReadPosition" });
+	Defs.Add({ "CurveOffsetPoint", "ReadCurveOffsetPoint",  Curve,          "ReadCurveOffsetPoint" });
+	Defs.Add({ "CurveNumPoint",    "ReadCurveNumPoint",     Curve,          "ReadCurveNumPoint" });
+	return Defs;
+}
+ 
+TSubclassOf<UActorComponent> UOptimusGroomGuideDataInterface::GetRequiredComponentClass() const
+{
+	return UGroomComponent::StaticClass();
+}
+
+void UOptimusGroomGuideDataInterface::GetSupportedInputs(TArray<FShaderFunctionDefinition>& OutFunctions) const
+{
+	OutFunctions.AddDefaulted_GetRef()
+		.SetName(TEXT("ReadNumGuidePoints"))
+		.AddReturnType(EShaderFundamentalType::Uint);
+
+	OutFunctions.AddDefaulted_GetRef()
+		.SetName(TEXT("ReadNumGuideCurves"))
+		.AddReturnType(EShaderFundamentalType::Uint);
+
+	OutFunctions.AddDefaulted_GetRef()
+		.SetName(TEXT("ReadPosition"))
+		.AddReturnType(EShaderFundamentalType::Float, 3)
+		.AddParam(EShaderFundamentalType::Uint);
+
+	OutFunctions.AddDefaulted_GetRef()
+		.SetName(TEXT("ReadCurveOffsetPoint"))
+		.AddReturnType(EShaderFundamentalType::Uint)
+		.AddParam(EShaderFundamentalType::Uint);
+
+	OutFunctions.AddDefaulted_GetRef()
+		.SetName(TEXT("ReadCurveNumPoint"))
+		.AddReturnType(EShaderFundamentalType::Uint)
+		.AddParam(EShaderFundamentalType::Uint);
+}
+
+BEGIN_SHADER_PARAMETER_STRUCT(FGroomGuideDataInterfaceParameters, )
+	SHADER_PARAMETER_STRUCT_INCLUDE(FHairStrandsInstanceCommonParameters, Common)
+	SHADER_PARAMETER_STRUCT_INCLUDE(FHairStrandsInstanceResourceParameters, Resources)
+END_SHADER_PARAMETER_STRUCT()
+
+void UOptimusGroomGuideDataInterface::GetShaderParameters(TCHAR const* UID, FShaderParametersMetadataBuilder& InOutBuilder, FShaderParametersMetadataAllocations& InOutAllocations) const
+{
+	InOutBuilder.AddNestedStruct<FGroomGuideDataInterfaceParameters>(UID);
+}
+
+TCHAR const* UOptimusGroomGuideDataInterface::TemplateFilePath = TEXT("/Plugin/Runtime/HairStrands/Private/DeformerDataInterfaceGroomGuide.ush");
+
+TCHAR const* UOptimusGroomGuideDataInterface::GetShaderVirtualPath() const
+{
+	return TemplateFilePath;
+}
+
+void UOptimusGroomGuideDataInterface::GetShaderHash(FString& InOutKey) const
+{
+	GetShaderFileHash(TemplateFilePath, EShaderPlatform::SP_PCD3D_SM5).AppendString(InOutKey);
+}
+
+void UOptimusGroomGuideDataInterface::GetHLSL(FString& OutHLSL, FString const& InDataInterfaceName) const
+{
+	TMap<FString, FStringFormatArg> TemplateArgs =
+	{
+		{ TEXT("DataInterfaceName"), InDataInterfaceName },
+	};
+
+	FString TemplateFile;
+	LoadShaderSourceFile(TemplateFilePath, EShaderPlatform::SP_PCD3D_SM5, &TemplateFile, nullptr);
+	OutHLSL += FString::Format(*TemplateFile, TemplateArgs);
+}
+
+UComputeDataProvider* UOptimusGroomGuideDataInterface::CreateDataProvider(TObjectPtr<UObject> InBinding, uint64 InInputMask, uint64 InOutputMask) const
+{
+	UOptimusGroomGuideDataProvider* Provider = NewObject<UOptimusGroomGuideDataProvider>();
+	Provider->Groom = Cast<UGroomComponent>(InBinding);
+	return Provider;
+}
+
+
+FComputeDataProviderRenderProxy* UOptimusGroomGuideDataProvider::GetRenderProxy()
+{
+	return new FOptimusGroomGuideDataProviderProxy(Groom);
+}
+
+
+FOptimusGroomGuideDataProviderProxy::FOptimusGroomGuideDataProviderProxy(UGroomComponent* InGroomComponent)
+{
+	GroomComponent = InGroomComponent;
+}
+
+bool FOptimusGroomGuideDataProviderProxy::IsValid(FValidationData const& InValidationData) const
+{
+	if (InValidationData.ParameterStructSize != sizeof(FParameters))
+	{
+		return false;
+	}
+	if (GroomComponent == nullptr)
+	{
+		return false;
+	}
+	const uint32 GroupCount = GroomComponent->GetGroupCount();
+	if (InValidationData.NumInvocations != GroupCount)
+	{
+		return false;
+	}
+	// Earlier invalidation in case one of the instance is not ready/valid
+	for (uint32 GroupIndex = 0; GroupIndex < GroupCount; ++GroupIndex)
+	{
+		if (GroomComponent->GetGroupInstance(GroupIndex) == nullptr)
+		{
+			return false;
+		}
+	}
+	
+	return true;
+}
+
+void FOptimusGroomGuideDataProviderProxy::AllocateResources(FRDGBuilder& GraphBuilder)
+{
+	Resources.Empty();
+	FallbackSRV = nullptr;
+	const uint32 InstanceCount = GroomComponent ? GroomComponent->GetGroupCount() : 0;
+	for (uint32 Index =0; Index <InstanceCount;++Index)
+	{
+		if (FHairGroupInstance* Instance = GroomComponent->GetGroupInstance(Index))
+		{
+			FHairStrandsInstanceResourceParameters& R = Resources.AddDefaulted_GetRef();
+			if (Instance->Guides.RestResource)
+			{
+				R.PositionBuffer		= RegisterAsSRV(GraphBuilder, Instance->Guides.RestResource->PositionBuffer);
+				R.PositionOffsetBuffer 	= RegisterAsSRV(GraphBuilder, Instance->Guides.RestResource->PositionOffsetBuffer);
+				R.CurveBuffer			= RegisterAsSRV(GraphBuilder, Instance->Guides.RestResource->CurveBuffer);
+				R.PointToCurveBuffer	= RegisterAsSRV(GraphBuilder, Instance->Guides.RestResource->PointToCurveBuffer);
+			}
+
+			if (!FallbackSRV)
+			{
+				FallbackSRV = GraphBuilder.CreateSRV(GraphBuilder.RegisterExternalBuffer(GWhiteVertexBufferWithRDG->Buffer), PF_R16G16B16A16_UINT); // This should be a byte address buffer
+			}
+		}
+	}
+}
+
+void FOptimusGroomGuideDataProviderProxy::GatherDispatchData(FDispatchData const& InDispatchData)
+{
+	const uint32 InstanceCount = GroomComponent ? GroomComponent->GetGroupCount() : 0;
+	check(InDispatchData.NumInvocations == InstanceCount);
+	check(uint32(Resources.Num()) == InstanceCount);
+
+	const TStridedView<FParameters> ParameterArray = MakeStridedParameterView<FParameters>(InDispatchData);
+	for (int32 InvocationIndex = 0; InvocationIndex < ParameterArray.Num(); ++InvocationIndex)
+	{
+		// Default valid in case the instance is not valid
+		FParameters& Parameters = ParameterArray[InvocationIndex];
+		Parameters.Common.PointCount 				= 0;
+		Parameters.Common.CurveCount 				= 0;
+		Parameters.Resources.PositionBuffer			= FallbackSRV;
+		Parameters.Resources.PositionOffsetBuffer 	= FallbackSRV;
+		Parameters.Resources.PointToCurveBuffer		= FallbackSRV;
+		Parameters.Resources.CurveBuffer			= FallbackSRV;
+
+		if (FHairGroupInstance* Instance = GroomComponent->GetGroupInstance(InvocationIndex))
+		{
+			const bool bIsSRVValid = Resources[InvocationIndex].PositionBuffer != nullptr;
+			if (bIsSRVValid)
+			{
+				// Reuse the strands properties for guides, and apply guides values
+				const FHairGroupPublicData::FVertexFactoryInput VFInput = ComputeHairStrandsVertexInputData(Instance, EGroomViewMode::None);
+				Parameters.Common 				= VFInput.Strands.Common;
+				Parameters.Common.PointCount 	= Instance->Guides.IsValid() ? Instance->Guides.GetData().Header.PointCount : 0;
+				Parameters.Common.CurveCount 	= Instance->Guides.IsValid() ? Instance->Guides.GetData().Header.CurveCount : 0;
+				Parameters.Resources 			= Resources[InvocationIndex];
+			}
+		}
+	}
+}
\ No newline at end of file
diff --git a/Engine/Plugins/Runtime/HairStrands/Source/HairStrandsDeformer/Private/DeformerDataInterfaceGroomGuide.h b/Engine/Plugins/Runtime/HairStrands/Source/HairStrandsDeformer/Private/DeformerDataInterfaceGroomGuide.h
new file mode 100644
--- /dev/null
+++ b/Engine/Plugins/Runtime/HairStrands/Source/HairStrandsDeformer/Private/DeformerDataInterfaceGroomGuide.h
@@ -0,0 +1,74 @@
+// Copyright Epic Games, Inc. All Rights Reserved.
+
+#pragma once
+
+#include "RenderGraphDefinitions.h"
+#include "OptimusComputeDataInterface.h"
+#include "HairStrandsInterface.h"
+#include "ComputeFramework/ComputeDataProvider.h"
+#include "DeformerDataInterfaceGroomGuide.generated.h"
+
+class FGroomGuideDataInterfaceParameters;
+class UGroomComponent;
+
+/** Compute Framework Data Interface for reading groom guides. */
+UCLASS(Category = ComputeFramework)
+class UOptimusGroomGuideDataInterface : public UOptimusComputeDataInterface
+{
+	GENERATED_BODY()
+
+public:
+	//~ Begin UOptimusComputeDataInterface Interface
+	FString GetDisplayName() const override;
+	TArray<FOptimusCDIPinDefinition> GetPinDefinitions() const override;
+	TSubclassOf<UActorComponent> GetRequiredComponentClass() const override;
+	//~ End UOptimusComputeDataInterface Interface
+	
+	//~ Begin UComputeDataInterface Interface
+	TCHAR const* GetClassName() const override { return TEXT("Groom"); }
+	void GetSupportedInputs(TArray<FShaderFunctionDefinition>& OutFunctions) const override;
+	void GetShaderParameters(TCHAR const* UID, FShaderParametersMetadataBuilder& InOutBuilder, FShaderParametersMetadataAllocations& InOutAllocations) const override;
+	TCHAR const* GetShaderVirtualPath() const override;
+	void GetShaderHash(FString& InOutKey) const override;
+	void GetHLSL(FString& OutHLSL, FString const& InDataInterfaceName) const override;
+	UComputeDataProvider* CreateDataProvider(TObjectPtr<UObject> InBinding, uint64 InInputMask, uint64 InOutputMask) const override;
+	//~ End UComputeDataInterface Interface
+
+private:
+	static TCHAR const* TemplateFilePath;
+};
+
+/** Compute Framework Data Provider for reading groom. */
+UCLASS(BlueprintType, editinlinenew, Category = ComputeFramework)
+class UOptimusGroomGuideDataProvider : public UComputeDataProvider
+{
+	GENERATED_BODY()
+
+public:
+	UPROPERTY(EditAnywhere, BlueprintReadWrite, Category = Binding)
+	TObjectPtr<UGroomComponent> Groom = nullptr;
+
+	//~ Begin UComputeDataProvider Interface
+	FComputeDataProviderRenderProxy* GetRenderProxy() override;
+	//~ End UComputeDataProvider Interface
+};
+
+class FOptimusGroomGuideDataProviderProxy : public FComputeDataProviderRenderProxy
+{
+public:
+	FOptimusGroomGuideDataProviderProxy(UGroomComponent* InGroomComponent);
+
+	//~ Begin FComputeDataProviderRenderProxy Interface
+	bool IsValid(FValidationData const& InValidationData) const override;
+	void AllocateResources(FRDGBuilder& GraphBuilder) override;
+	void GatherDispatchData(FDispatchData const& InDispatchData) override;
+	//~ End FComputeDataProviderRenderProxy Interface
+
+private:
+	using FParameters = FGroomGuideDataInterfaceParameters;
+
+	UGroomComponent* GroomComponent = nullptr;
+	TArray<FHairStrandsInstanceResourceParameters> Resources;
+	FRDGBufferSRVRef FallbackSRV;
+
+};
diff --git a/Engine/Plugins/Runtime/HairStrands/Source/HairStrandsDeformer/Private/DeformerDataInterfaceGroomWrite.cpp b/Engine/Plugins/Runtime/HairStrands/Source/HairStrandsDeformer/Private/DeformerDataInterfaceGroomWrite.cpp
new file mode 100644
--- /dev/null
+++ b/Engine/Plugins/Runtime/HairStrands/Source/HairStrandsDeformer/Private/DeformerDataInterfaceGroomWrite.cpp
@@ -0,0 +1,291 @@
+// Copyright Epic Games, Inc. All Rights Reserved.
+
+#include "DeformerDataInterfaceGroomWrite.h"
+
+#include "ComputeFramework/ShaderParamTypeDefinition.h"
+#include "OptimusDataDomain.h"
+#include "RenderGraphBuilder.h"
+#include "Rendering/SkeletalMeshLODRenderData.h"
+#include "Rendering/SkeletalMeshRenderData.h"
+#include "ShaderParameterMetadataBuilder.h"
+#include "SkeletalMeshDeformerHelpers.h"
+#include "SkeletalRenderPublic.h"
+#include "GroomComponent.h"
+#include "GroomInstance.h"
+#include "DeformerGroomComponentSource.h"
+#include "RenderGraphUtils.h"
+#include "SystemTextures.h"
+#include "HairStrandsInterpolation.h"
+
+FString UOptimusGroomWriteDataInterface::GetDisplayName() const
+{
+	return TEXT("Write Groom");
+}
+
+FName UOptimusGroomWriteDataInterface::GetCategory() const
+{
+	return CategoryName::OutputDataInterfaces;
+}
+
+TArray<FOptimusCDIPinDefinition> UOptimusGroomWriteDataInterface::GetPinDefinitions() const
+{
+	FName ControlPoint(UOptimusGroomComponentSource::Domains::ControlPoint);
+	FName Curve(UOptimusGroomComponentSource::Domains::Curve);
+
+	TArray<FOptimusCDIPinDefinition> Defs;
+	Defs.Add({ "Position",         "WritePosition",         	ControlPoint, "ReadNumControlPoints" });
+	Defs.Add({ "Radius",           "WriteRadius",           	ControlPoint, "ReadNumControlPoints" });
+	Defs.Add({ "PositionAndRadius","WritePositionAndRadius",	ControlPoint, "ReadNumControlPoints" });
+
+	Defs.Add({ "RootUV",			"WriteRootUV",   			Curve, 		  "ReadNumCurves" });
+	Defs.Add({ "Seed",				"WriteSeed",   				Curve, 		  "ReadNumCurves" });
+	Defs.Add({ "ClumpId",			"WriteClumpId",   			Curve, 		  "ReadNumCurves" });
+	Defs.Add({ "Color",				"WriteColor",   			ControlPoint, "ReadNumControlPoints" });
+	Defs.Add({ "Roughness",			"WriteRoughness",   		ControlPoint, "ReadNumControlPoints" });
+	Defs.Add({ "AO",				"WriteAO",   				ControlPoint, "ReadNumControlPoints" });
+
+	return Defs;
+}
+
+TSubclassOf<UActorComponent> UOptimusGroomWriteDataInterface::GetRequiredComponentClass() const
+{
+	return UGroomComponent::StaticClass();
+}
+
+void UOptimusGroomWriteDataInterface::GetSupportedInputs(TArray<FShaderFunctionDefinition>& OutFunctions) const
+{
+	OutFunctions.AddDefaulted_GetRef()
+		.SetName(TEXT("ReadNumControlPoints"))
+		.AddReturnType(EShaderFundamentalType::Uint);
+
+	OutFunctions.AddDefaulted_GetRef()
+		.SetName(TEXT("ReadNumCurves"))
+		.AddReturnType(EShaderFundamentalType::Uint);
+}
+
+void UOptimusGroomWriteDataInterface::GetSupportedOutputs(TArray<FShaderFunctionDefinition>& OutFunctions) const
+{
+	OutFunctions.AddDefaulted_GetRef()
+		.SetName(TEXT("WritePosition"))
+		.AddParam(EShaderFundamentalType::Uint)
+		.AddParam(EShaderFundamentalType::Float, 3);
+
+	OutFunctions.AddDefaulted_GetRef()
+		.SetName(TEXT("WriteRadius"))
+		.AddParam(EShaderFundamentalType::Uint)
+		.AddParam(EShaderFundamentalType::Float);
+
+	OutFunctions.AddDefaulted_GetRef()
+		.SetName(TEXT("WritePositionAndRadius"))
+		.AddParam(EShaderFundamentalType::Uint)
+		.AddParam(EShaderFundamentalType::Float, 4);
+
+	OutFunctions.AddDefaulted_GetRef()
+		.SetName(TEXT("WriteRootUV"))
+		.AddParam(EShaderFundamentalType::Uint)
+		.AddParam(EShaderFundamentalType::Float, 2);
+
+	OutFunctions.AddDefaulted_GetRef()
+		.SetName(TEXT("WriteSeed"))
+		.AddParam(EShaderFundamentalType::Uint)
+		.AddParam(EShaderFundamentalType::Float);
+
+	OutFunctions.AddDefaulted_GetRef()
+		.SetName(TEXT("WriteClumpId"))
+		.AddParam(EShaderFundamentalType::Uint)
+		.AddParam(EShaderFundamentalType::Float);
+
+	OutFunctions.AddDefaulted_GetRef()
+		.SetName(TEXT("WriteColor"))
+		.AddParam(EShaderFundamentalType::Uint)
+		.AddParam(EShaderFundamentalType::Float, 3);
+
+	OutFunctions.AddDefaulted_GetRef()
+		.SetName(TEXT("WriteRoughness"))
+		.AddParam(EShaderFundamentalType::Uint)
+		.AddParam(EShaderFundamentalType::Float);
+
+	OutFunctions.AddDefaulted_GetRef()
+		.SetName(TEXT("WriteAO"))
+		.AddParam(EShaderFundamentalType::Uint)
+		.AddParam(EShaderFundamentalType::Float);
+}
+
+BEGIN_SHADER_PARAMETER_STRUCT(FGroomWriteDataInterfaceParameters, )
+	SHADER_PARAMETER_STRUCT_INCLUDE(FHairStrandsInstanceCommonParameters, Common)
+	SHADER_PARAMETER(uint32, OutputStreamStart)
+	SHADER_PARAMETER_RDG_BUFFER_SRV(StructuredBuffer, PositionOffsetBufferSRV)
+	SHADER_PARAMETER_RDG_BUFFER_SRV(ByteAddressBuffer, PositionBufferSRV)
+	SHADER_PARAMETER_RDG_BUFFER_UAV(RWByteAddressBuffer, PositionBufferUAV)
+	SHADER_PARAMETER_RDG_BUFFER_UAV(RWByteAddressBuffer, CurveAttributeBufferUAV)
+	SHADER_PARAMETER_RDG_BUFFER_UAV(RWByteAddressBuffer, PointAttributeBufferUAV)
+END_SHADER_PARAMETER_STRUCT()
+
+void UOptimusGroomWriteDataInterface::GetShaderParameters(TCHAR const* UID, FShaderParametersMetadataBuilder& InOutBuilder, FShaderParametersMetadataAllocations& InOutAllocations) const
+{
+	InOutBuilder.AddNestedStruct<FGroomWriteDataInterfaceParameters>(UID);
+}
+
+TCHAR const* UOptimusGroomWriteDataInterface::TemplateFilePath = TEXT("/Plugin/Runtime/HairStrands/Private/DeformerDataInterfaceGroomWrite.ush");
+
+TCHAR const* UOptimusGroomWriteDataInterface::GetShaderVirtualPath() const
+{
+	return TemplateFilePath;
+}
+
+void UOptimusGroomWriteDataInterface::GetShaderHash(FString& InOutKey) const
+{
+	GetShaderFileHash(TemplateFilePath, EShaderPlatform::SP_PCD3D_SM5).AppendString(InOutKey);
+}
+
+void UOptimusGroomWriteDataInterface::GetHLSL(FString& OutHLSL, FString const& InDataInterfaceName) const
+{
+	TMap<FString, FStringFormatArg> TemplateArgs =
+	{
+		{ TEXT("DataInterfaceName"), InDataInterfaceName },
+	};
+
+	FString TemplateFile;
+	LoadShaderSourceFile(TemplateFilePath, EShaderPlatform::SP_PCD3D_SM5, &TemplateFile, nullptr);
+	OutHLSL += FString::Format(*TemplateFile, TemplateArgs);
+}
+
+UComputeDataProvider* UOptimusGroomWriteDataInterface::CreateDataProvider(TObjectPtr<UObject> InBinding, uint64 InInputMask, uint64 InOutputMask) const
+{
+	UOptimusGroomWriteDataProvider* Provider = NewObject<UOptimusGroomWriteDataProvider>();
+	Provider->GroomComponent = Cast<UGroomComponent>(InBinding);
+	Provider->OutputMask = InOutputMask;
+	return Provider;
+}
+
+
+FComputeDataProviderRenderProxy* UOptimusGroomWriteDataProvider::GetRenderProxy()
+{
+	return new FOptimusGroomWriteDataProviderProxy(GroomComponent, OutputMask);
+}
+
+
+FOptimusGroomWriteDataProviderProxy::FOptimusGroomWriteDataProviderProxy(UGroomComponent* InGroomComponent, uint64 InOutputMask)
+{
+	OutputMask = InOutputMask;
+	const uint32 InstanceCount = InGroomComponent ? InGroomComponent->GetGroupCount() : 0;
+	for (uint32 Index = 0; Index < InstanceCount; ++Index)
+	{
+		Instances.Add(InGroomComponent->GetGroupInstance(Index));
+	}
+}
+
+bool FOptimusGroomWriteDataProviderProxy::IsValid(FValidationData const& InValidationData) const
+{
+	if (InValidationData.ParameterStructSize != sizeof(FParameters))
+	{
+		return false;
+	}
+	if (InValidationData.NumInvocations != Instances.Num())
+	{
+		return false;
+	}
+
+	return true;
+}
+
+void FOptimusGroomWriteDataProviderProxy::AllocateResources(FRDGBuilder& GraphBuilder)
+{
+	for (FHairGroupInstance* GroomInstance : Instances)
+	{
+		if (GroomInstance)
+		{
+			// Allocate required buffers
+			const int32 NumControlPoints = GroomInstance->Strands.GetData().GetNumPoints();
+			const int32 NumCurves = GroomInstance->Strands.GetData().GetNumCurves();
+
+			FResources& R = Resources.AddDefaulted_GetRef();
+
+			FRDGBufferRef PositionBuffer_fallback = GraphBuilder.CreateBuffer(FRDGBufferDesc::CreateBufferDesc(8, 1), TEXT("Groom.DeformedPositionBuffer"), ERDGBufferFlags::None);
+			R.PositionBufferSRV_fallback = GraphBuilder.CreateSRV(GraphBuilder.RegisterExternalBuffer(GWhiteVertexBufferWithRDG->Buffer));
+			R.PositionBufferUAV_fallback = GraphBuilder.CreateUAV(PositionBuffer_fallback, ERDGUnorderedAccessViewFlags::SkipBarrier);
+
+			FRDGBufferRef AttributeBuffer_fallback = GraphBuilder.CreateBuffer(FRDGBufferDesc::CreateByteAddressDesc(16), TEXT("Groom.DeformedAttributeBuffer"), ERDGBufferFlags::None);
+			R.AttributeBufferUAV_fallback = GraphBuilder.CreateUAV(AttributeBuffer_fallback, PF_Unknown, ERDGUnorderedAccessViewFlags::SkipBarrier);
+
+			// Positions / Radius
+			if (OutputMask & 0x7)
+			{
+				R.PositionOffsetBufferSRV 	= Register(GraphBuilder, GroomInstance->Strands.RestResource->PositionOffsetBuffer, ERDGImportedBufferFlags::CreateSRV).SRV;
+				R.PositionBufferSRV 		= Register(GraphBuilder, GroomInstance->Strands.RestResource->PositionBuffer, ERDGImportedBufferFlags::CreateSRV).SRV;
+				R.PositionBufferUAV 		= Register(GraphBuilder, GroomInstance->Strands.DeformedResource->GetDeformerBuffer(GraphBuilder), ERDGImportedBufferFlags::CreateUAV).UAV;
+			}
+			else
+			{
+				FRDGBufferRef DummyUAVBuffer = GraphBuilder.CreateBuffer(FRDGBufferDesc::CreateBufferDesc(8u, 1), TEXT("Hair.Deformer.DummyBuffer"));
+				
+				R.PositionOffsetBufferSRV 	= GraphBuilder.CreateSRV(GSystemTextures.GetDefaultStructuredBuffer(GraphBuilder, 16u), PF_A32B32G32R32F); //GraphBuilder.CreateSRV(GraphBuilder.RegisterExternalBuffer(GWhiteVertexBufferWithRDG->Buffer), PF_A32B32G32R32F);
+				R.PositionBufferSRV 		= GraphBuilder.CreateSRV(GraphBuilder.RegisterExternalBuffer(GWhiteVertexBufferWithRDG->Buffer));
+				R.PositionBufferUAV 		= GraphBuilder.CreateUAV(DummyUAVBuffer);
+			}
+			
+			// Curve Attributes
+			if (OutputMask & 0x38)
+			{
+				FRDGExternalBuffer CurveAttributeBufferExt = GroomInstance->Strands.DeformedResource->GetDeformerCurveAttributeBuffer(GraphBuilder);
+				if (CurveAttributeBufferExt.Buffer)
+				{
+					// Always copy the attributes from the rest asset, so that if deformer write different attribute at different tick, it all remains consistent with the source data
+					FRDGImportedBuffer DstCurveAttributeBuffer = Register(GraphBuilder, CurveAttributeBufferExt, ERDGImportedBufferFlags::CreateUAV);
+					FRDGImportedBuffer SrcCurveAttributeBuffer = Register(GraphBuilder, GroomInstance->Strands.RestResource->CurveAttributeBuffer, ERDGImportedBufferFlags::None);
+					if (SrcCurveAttributeBuffer.Buffer)
+					{
+						AddCopyBufferPass(GraphBuilder, DstCurveAttributeBuffer.Buffer, SrcCurveAttributeBuffer.Buffer);
+					}
+					R.CurveAttributeBufferUAV = DstCurveAttributeBuffer.UAV;
+				}
+			}
+			if (R.CurveAttributeBufferUAV == nullptr)
+			{
+				R.CurveAttributeBufferUAV = R.AttributeBufferUAV_fallback;
+			}
+
+			// Point Attributes
+			if (OutputMask & 0x1c0)
+			{
+				FRDGExternalBuffer PointAttributeBufferExt = GroomInstance->Strands.DeformedResource->GetDeformerPointAttributeBuffer(GraphBuilder);
+				if (PointAttributeBufferExt.Buffer)
+				{
+					FRDGImportedBuffer PointAttributeBuffer = Register(GraphBuilder, PointAttributeBufferExt, ERDGImportedBufferFlags::CreateUAV);
+					AddCopyBufferPass(GraphBuilder, PointAttributeBuffer.Buffer, Register(GraphBuilder, GroomInstance->Strands.RestResource->PointAttributeBuffer, ERDGImportedBufferFlags::None).Buffer);
+					R.PointAttributeBufferUAV = PointAttributeBuffer.UAV;
+				}
+			}
+			if (R.PointAttributeBufferUAV == nullptr)
+			{
+				R.PointAttributeBufferUAV = R.AttributeBufferUAV_fallback;
+				//R.PointAttributeBufferUAV = GraphBuilder.CreateUAV(GraphBuilder.RegisterExternalBuffer(GWhiteVertexBufferWithRDG->Buffer), PF_R32_UINT);
+			}
+		}	
+	}
+}
+
+void FOptimusGroomWriteDataProviderProxy::GatherDispatchData(FDispatchData const& InDispatchData)
+{
+	const TStridedView<FParameters> ParameterArray = MakeStridedParameterView<FParameters>(InDispatchData);
+	for (int32 InvocationIndex = 0; InvocationIndex < ParameterArray.Num(); ++InvocationIndex)
+	{
+		if (FHairGroupInstance* GroomInstance = Instances[InvocationIndex])
+		{
+			FResources& Resource = Resources[InvocationIndex];
+
+			const bool bValid = Resource.PositionBufferUAV != nullptr && Resource.PositionBufferSRV != nullptr && Resource.PositionOffsetBufferSRV != nullptr;
+
+			const FHairGroupPublicData::FVertexFactoryInput VFInput = ComputeHairStrandsVertexInputData(GroomInstance, EGroomViewMode::None);		
+
+			FParameters& Parameters = ParameterArray[InvocationIndex];
+			Parameters.Common 					= VFInput.Strands.Common;
+			Parameters.OutputStreamStart 		= 0;
+			Parameters.PositionOffsetBufferSRV 	= bValid ? Resource.PositionOffsetBufferSRV : Resource.PositionBufferSRV_fallback;
+			Parameters.PositionBufferSRV 		= bValid ? Resource.PositionBufferSRV : Resource.PositionBufferSRV_fallback;
+			Parameters.PositionBufferUAV 		= bValid ? Resource.PositionBufferUAV : Resource.PositionBufferUAV_fallback;
+			Parameters.PointAttributeBufferUAV 	= bValid ? Resource.PointAttributeBufferUAV : Resource.AttributeBufferUAV_fallback;
+			Parameters.CurveAttributeBufferUAV 	= bValid ? Resource.CurveAttributeBufferUAV : Resource.AttributeBufferUAV_fallback;
+		}
+	}
+}
\ No newline at end of file
diff --git a/Engine/Plugins/Runtime/HairStrands/Source/HairStrandsDeformer/Private/DeformerDataInterfaceGroomWrite.h b/Engine/Plugins/Runtime/HairStrands/Source/HairStrandsDeformer/Private/DeformerDataInterfaceGroomWrite.h
new file mode 100644
--- /dev/null
+++ b/Engine/Plugins/Runtime/HairStrands/Source/HairStrandsDeformer/Private/DeformerDataInterfaceGroomWrite.h
@@ -0,0 +1,92 @@
+// Copyright Epic Games, Inc. All Rights Reserved.
+
+#pragma once
+
+#include "OptimusComputeDataInterface.h"
+#include "ComputeFramework/ComputeDataProvider.h"
+#include "DeformerDataInterfaceGroomWrite.generated.h"
+
+class FGroomWriteDataInterfaceParameters;
+struct FHairGroupInstance;
+class FRDGBuffer;
+class FRDGBufferSRV;
+class FRDGBufferUAV;
+class UGroomComponent;
+
+/** Compute Framework Data Interface for writing skinned mesh. */
+UCLASS(Category = ComputeFramework)
+class UOptimusGroomWriteDataInterface : public UOptimusComputeDataInterface
+{
+	GENERATED_BODY()
+
+public:
+	//~ Begin UOptimusComputeDataInterface Interface
+	FString GetDisplayName() const override;
+	FName GetCategory() const override;
+	TArray<FOptimusCDIPinDefinition> GetPinDefinitions() const override;
+	TSubclassOf<UActorComponent> GetRequiredComponentClass() const override;
+	//~ End UOptimusComputeDataInterface Interface
+	
+	//~ Begin UComputeDataInterface Interface
+	TCHAR const* GetClassName() const override { return TEXT("GroomWrite"); }
+	void GetSupportedInputs(TArray<FShaderFunctionDefinition>& OutFunctions) const override;
+	void GetSupportedOutputs(TArray<FShaderFunctionDefinition>& OutFunctions) const override;
+	void GetShaderParameters(TCHAR const* UID, FShaderParametersMetadataBuilder& InOutBuilder, FShaderParametersMetadataAllocations& InOutAllocations) const override;
+	TCHAR const* GetShaderVirtualPath() const override;
+	void GetShaderHash(FString& InOutKey) const override;
+	void GetHLSL(FString& OutHLSL, FString const& InDataInterfaceName) const override;
+	UComputeDataProvider* CreateDataProvider(TObjectPtr<UObject> InBinding, uint64 InInputMask, uint64 InOutputMask) const override;
+	//~ End UComputeDataInterface Interface
+
+private:
+	static TCHAR const* TemplateFilePath;
+};
+
+/** Compute Framework Data Provider for writing skinned mesh. */
+UCLASS(BlueprintType, editinlinenew, Category = ComputeFramework)
+class UOptimusGroomWriteDataProvider : public UComputeDataProvider
+{
+	GENERATED_BODY()
+
+public:
+	UPROPERTY(EditAnywhere, BlueprintReadWrite, Category = Binding)
+	TObjectPtr<UGroomComponent> GroomComponent = nullptr;
+
+	uint64 OutputMask = 0;
+
+	//~ Begin UComputeDataProvider Interface
+	FComputeDataProviderRenderProxy* GetRenderProxy() override;
+	//~ End UComputeDataProvider Interface
+};
+
+class FOptimusGroomWriteDataProviderProxy : public FComputeDataProviderRenderProxy
+{
+public:
+	FOptimusGroomWriteDataProviderProxy(UGroomComponent* InGroomComponent, uint64 InOutputMask);
+
+	//~ Begin FComputeDataProviderRenderProxy Interface
+	bool IsValid(FValidationData const& InValidationData) const override;
+	void AllocateResources(FRDGBuilder& GraphBuilder) override;
+	void GatherDispatchData(FDispatchData const& InDispatchData) override;
+	//~ End FComputeDataProviderRenderProxy Interface
+
+private:
+	using FParameters = FGroomWriteDataInterfaceParameters;
+
+	TArray<FHairGroupInstance*> Instances;
+	uint64 OutputMask = 0;
+
+	struct FResources
+	{
+		FRDGBufferSRV* PositionOffsetBufferSRV = nullptr;
+		FRDGBufferSRV* PositionBufferSRV = nullptr;
+		FRDGBufferUAV* PositionBufferUAV = nullptr;
+		FRDGBufferSRV* PositionBufferSRV_fallback = nullptr;
+		FRDGBufferUAV* PositionBufferUAV_fallback = nullptr;
+
+		FRDGBufferUAV* PointAttributeBufferUAV = nullptr;
+		FRDGBufferUAV* CurveAttributeBufferUAV = nullptr;
+		FRDGBufferUAV* AttributeBufferUAV_fallback = nullptr;
+	};
+	TArray<FResources> Resources;
+};
\ No newline at end of file
diff --git a/Engine/Plugins/Runtime/HairStrands/Source/HairStrandsDeformer/Private/DeformerGroomComponentSource.cpp b/Engine/Plugins/Runtime/HairStrands/Source/HairStrandsDeformer/Private/DeformerGroomComponentSource.cpp
new file mode 100644
--- /dev/null
+++ b/Engine/Plugins/Runtime/HairStrands/Source/HairStrandsDeformer/Private/DeformerGroomComponentSource.cpp
@@ -0,0 +1,99 @@
+// Copyright Epic Games, Inc. All Rights Reserved.
+
+#include "DeformerGroomComponentSource.h"
+
+#include "GroomComponent.h"
+
+#define LOCTEXT_NAMESPACE "DeformersGroomComponentSource"
+
+FName UOptimusGroomComponentSource::Domains::ControlPoint("ControlPoint");
+FName UOptimusGroomComponentSource::Domains::Curve("Curve");
+
+FText UOptimusGroomComponentSource::GetDisplayName() const
+{
+	return LOCTEXT("GroomComponent", "Groom Component");
+}
+
+TSubclassOf<UActorComponent> UOptimusGroomComponentSource::GetComponentClass() const
+{
+	return UGroomComponent::StaticClass();
+}
+
+bool UOptimusGroomComponentSource::IsUsableAsPrimarySource() const
+{
+	return GetComponentClass()->IsChildOf<UGroomComponent>();
+}
+
+TArray<FName> UOptimusGroomComponentSource::GetExecutionDomains() const
+{
+	return { Domains::ControlPoint, Domains::Curve };
+}
+
+int32 UOptimusGroomComponentSource::GetLodIndex(const UActorComponent* InComponent) const
+{
+	return 0;
+}
+
+uint32 UOptimusGroomComponentSource::GetDefaultNumInvocations(
+	const UActorComponent* InComponent,
+	int32 InLod
+	) const
+{
+	const UGroomComponent* GroomComponent = Cast<UGroomComponent>(InComponent);
+	if (!GroomComponent)
+	{
+		return 0;
+	}
+
+	const uint32 GroupCount = GroomComponent->GetGroupCount();
+
+	return GroupCount;
+}
+
+bool UOptimusGroomComponentSource::GetComponentElementCountsForExecutionDomain(
+	FName InDomainName,
+	const UActorComponent* InComponent,
+	int32 InLodIndex,
+	TArray<int32>& OutInvocationElementCounts
+) const
+{
+	const UGroomComponent* GroomComponent = Cast<UGroomComponent>(InComponent);
+	if (!GroomComponent)
+	{
+		return false;
+	}
+
+	// Ensure all groups are valid
+	const uint32 GroupCount = GroomComponent->GetGroupCount();
+	for (uint32 GroupIt=0;GroupIt<GroupCount;++GroupIt)
+	{
+		if (!GroomComponent->GetGroupInstance(GroupIt))
+		{
+			return 0;
+		}
+	}
+
+	OutInvocationElementCounts.Reset();
+
+	if (InDomainName == Domains::ControlPoint || InDomainName == Domains::Curve)
+	{
+		const int32 NumInvocations = GroupCount;
+		OutInvocationElementCounts.Reserve(NumInvocations);
+		for (int32 InvocationIndex = 0; InvocationIndex < NumInvocations; ++InvocationIndex)
+		{
+			if (const FHairGroupInstance* Instance = GroomComponent->GetGroupInstance(InvocationIndex))
+			{
+				const int32 NumControlPoints = Instance->Strands.GetData().GetNumPoints();
+				const int32 NumCurves = Instance->Strands.GetData().GetNumCurves();
+				const int32 NumThreads = InDomainName == Domains::ControlPoint ? NumControlPoints : NumCurves;
+				OutInvocationElementCounts.Add(NumThreads);
+			}
+		}
+
+		return true;
+	}
+
+	// Unknown execution domain.
+	return false;
+}
+#undef LOCTEXT_NAMESPACE
diff --git a/Engine/Plugins/Runtime/HairStrands/Source/HairStrandsDeformer/Private/DeformerGroomComponentSource.h b/Engine/Plugins/Runtime/HairStrands/Source/HairStrandsDeformer/Private/DeformerGroomComponentSource.h
new file mode 100644
--- /dev/null
+++ b/Engine/Plugins/Runtime/HairStrands/Source/HairStrandsDeformer/Private/DeformerGroomComponentSource.h
@@ -0,0 +1,29 @@
+// Copyright Epic Games, Inc. All Rights Reserved.
+
+#pragma once
+
+#include "OptimusComponentSource.h"
+#include "DeformerGroomComponentSource.generated.h"
+
+UCLASS()
+class UOptimusGroomComponentSource : public UOptimusComponentSource
+{
+	GENERATED_BODY()
+public:
+
+	struct Domains
+	{
+		static FName ControlPoint;
+		static FName Curve;
+	};
+
+	// UOptimusComponentSource implementations
+	FText GetDisplayName() const override;
+	FName GetBindingName() const override { return FName("Groom"); }
+	TSubclassOf<UActorComponent> GetComponentClass() const override;
+	TArray<FName> GetExecutionDomains() const override;
+	int32 GetLodIndex(const UActorComponent* InComponent) const override;
+	uint32 GetDefaultNumInvocations(const UActorComponent* InComponent, int32 InLod) const;
+	bool GetComponentElementCountsForExecutionDomain(FName InDomainName, const UActorComponent* InComponent, int32 InLodIndex, TArray<int32>& OutInvocationElementCounts) const override;
+	bool IsUsableAsPrimarySource() const override;
+};
diff --git a/Engine/Plugins/Runtime/HairStrands/Source/HairStrandsDeformer/Private/HairStrandsDeformerModule.cpp b/Engine/Plugins/Runtime/HairStrands/Source/HairStrandsDeformer/Private/HairStrandsDeformerModule.cpp
new file mode 100644
--- /dev/null
+++ b/Engine/Plugins/Runtime/HairStrands/Source/HairStrandsDeformer/Private/HairStrandsDeformerModule.cpp
@@ -0,0 +1,5 @@
+// Copyright Epic Games, Inc. All Rights Reserved.
+
+#include "Modules/ModuleManager.h"
+
+IMPLEMENT_MODULE(FDefaultModuleImpl, HairStrandsDeformer);
