Add hair position compression, which reduce position data by 30%.

This compressions is based on relative error, and is automatically disabled is error is too high. Current the compression error is hardcoded to 1% but this could be exposed later per asset.
Compression can be globally disabled with r.HairStrands.CompressedPosition=0

diff --git a/Engine/Plugins/Runtime/HairStrands/Source/HairStrandsCore/Private/GroomAsset.cpp b/Engine/Plugins/Runtime/HairStrands/Source/HairStrandsCore/Private/GroomAsset.cpp
--- a/Engine/Plugins/Runtime/HairStrands/Source/HairStrandsCore/Private/GroomAsset.cpp
+++ b/Engine/Plugins/Runtime/HairStrands/Source/HairStrandsCore/Private/GroomAsset.cpp
@@ -85,6 +85,14 @@ static TAutoConsoleVariable<int32> GHairStrandsWarningLogVerbosity(
 static int32 GHairStrandsDDCLogEnable = 0;
 static FAutoConsoleVariableRef CVarHairStrandsDDCLogEnable(TEXT("r.HairStrands.DDCLog"), GHairStrandsDDCLogEnable, TEXT("Enable DDC logging for groom assets and groom binding assets"));
 
+static int32 GHairStrandsSupportCompressedPosition = 1;
+static FAutoConsoleVariableRef CVarHairStrandsSupportCompressedPosition(TEXT("r.HairStrands.CompressedPosition"), GHairStrandsSupportCompressedPosition, TEXT("Optional compessed position"), ECVF_ReadOnly);
+
+bool DoesHairStrandsSupportCompressedPosition()
+{
+	return GHairStrandsSupportCompressedPosition > 0;
+}
+
 bool IsHairStrandsDDCLogEnable()
 {
 	return GHairStrandsDDCLogEnable > 0;
@@ -528,8 +536,8 @@ static bool BuildHairGroup_Strands(
 		FHairStrandsDatas GuidesData;
 		FGroomBuilder::BuildData(HairGroup, InHairGroupsInterpolation[GroupIndex], OutHairGroupsInfo[GroupIndex], StrandsData, GuidesData);
 
-		FGroomBuilder::BuildBulkData(HairGroup.Info, GuidesData, OutHairGroupsData[GroupIndex].Guides.BulkData);
-		FGroomBuilder::BuildBulkData(HairGroup.Info, StrandsData, OutHairGroupsData[GroupIndex].Strands.BulkData);
+		FGroomBuilder::BuildBulkData(HairGroup.Info, GuidesData, OutHairGroupsData[GroupIndex].Guides.BulkData, false /*bAllowCompression*/);
+		FGroomBuilder::BuildBulkData(HairGroup.Info, StrandsData, OutHairGroupsData[GroupIndex].Strands.BulkData, true /*bAllowCompression*/);
 
 		// If there is no simulation or no global interpolation on that group there is no need for builder the interpolation data
 		if (bNeedInterpolationData)
@@ -2063,9 +2071,12 @@ namespace GroomDerivedDataCacheUtils
 
 	void SerializeHairInterpolationSettingsForDDC(FArchive& Ar, uint32 GroupIndex, FHairGroupsInterpolation& InterpolationSettings, FHairGroupsLOD& LODSettings, bool bRequireInterpolationData)
 	{
+		bool bSupportCompressedPosition = DoesHairStrandsSupportCompressedPosition();
+
 		// Note: this serializer is only used to build the groom DDC key, no versioning is required
 		Ar << GroupIndex;
 		Ar << bRequireInterpolationData;
+		Ar << bSupportCompressedPosition;
 
 		InterpolationSettings.BuildDDCKey(Ar);
 		LODSettings.BuildDDCKey(Ar);
@@ -2888,7 +2899,7 @@ bool UGroomAsset::BuildHairGroup_Cards(uint32 GroupIndex)
 
 				// 1. Build data & bulk data for the cards guides
 				FGroomBuilder::BuildData(LODGuidesData);
-				FGroomBuilder::BuildBulkData(DummyGroupInfo, LODGuidesData, LOD.GuideBulkData);
+				FGroomBuilder::BuildBulkData(DummyGroupInfo, LODGuidesData, LOD.GuideBulkData, false /*bAllowCompression*/);
 
 				// 2. (Re)Build data for the sim guides (since there are transient/no-cached)
 				const FHairDescriptionGroups& LocalHairDescriptionGroups = GetHairDescriptionGroups();
diff --git a/Engine/Plugins/Runtime/HairStrands/Source/HairStrandsCore/Private/GroomBuilder.cpp b/Engine/Plugins/Runtime/HairStrands/Source/HairStrandsCore/Private/GroomBuilder.cpp
--- a/Engine/Plugins/Runtime/HairStrands/Source/HairStrandsCore/Private/GroomBuilder.cpp
+++ b/Engine/Plugins/Runtime/HairStrands/Source/HairStrandsCore/Private/GroomBuilder.cpp
@@ -22,6 +22,9 @@ DEFINE_LOG_CATEGORY_STATIC(LogGroomBuilder, Log, All);
 
 #define LOCTEXT_NAMESPACE "GroomBuilder"
 
+/////////////////////////////////////////////////////////////////////////////////////////
+// CVars
+
 // For debug purpose
 static float GHairInterpolationMetric_Distance = 1;
 static float GHairInterpolationMetric_Angle = 0;
@@ -38,9 +41,17 @@ static FAutoConsoleVariableRef CVarHairClusterBuilder_MaxVoxelResolution(TEXT("r
 static int32 GHairGroupIndexBuilder_MaxVoxelResolution = 64;
 static FAutoConsoleVariableRef CVarHairGroupIndexBuilder_MaxVoxelResolution(TEXT("r.HairStrands.HairGroupBuilder.MaxVoxelResolution"), GHairGroupIndexBuilder_MaxVoxelResolution, TEXT("Max voxel resolution used when voxelizing hair strands to transfer group index grom strands to cards. This avoids too long building time (default:64).  "));
 
+/////////////////////////////////////////////////////////////////////////////////////////
+// Forward declaration
+
+bool DoesHairStrandsSupportCompressedPosition();
+
+/////////////////////////////////////////////////////////////////////////////////////////
+// Helpers
+
 FString FGroomBuilder::GetVersion()
 {
-	return TEXT("v12c");
+	return TEXT("v13_Dummy10");
 }
 
 // For debug purpose
@@ -54,6 +65,9 @@ void ReportSize(const TCHAR* InName, const TArray<T>& In)
 #endif
 }
 
+/////////////////////////////////////////////////////////////////////////////////////////
+// Build functions
+
 namespace FHairStrandsDecimation
 {
 	void Decimate(
@@ -168,7 +182,7 @@ namespace HairStrandsBuilder
 	}
 
 	/** Build the bulk/packed datas for gpu rendering/simulation */
-	void BuildBulkData(const FHairStrandsDatas& HairStrands, const TArray<uint8>& RandomSeeds, FHairStrandsBulkData& OutBulkData)
+	void BuildBulkData(const FHairStrandsDatas& HairStrands, const TArray<uint8>& RandomSeeds, FHairStrandsBulkData& OutBulkData, bool bAllowTranscoding)
 	{
 		TRACE_CPUPROFILER_EVENT_SCOPE(HairStrandsBuilder::BuildBulkData);
 
@@ -229,6 +243,7 @@ namespace HairStrandsBuilder
 		const float MaxRadius = GetHairStrandsMaxRadius(HairStrands);
 
 		static_assert(sizeof(FPackedHairVertex) == sizeof(FPackedHairVertex::BulkType));
+		static_assert(sizeof(FPackedHairVertex) == FPackedHairPositionStrideInBytes);
 		static_assert(sizeof(FPackedHairVertex) == 8u);
 
 		uint32 MinPointPerCurve = HAIR_MAX_NUM_POINT_PER_CURVE;
@@ -401,6 +416,138 @@ namespace HairStrandsBuilder
 				Curve.PointOffset = FMath::Min(uint32(IndexOffset), HAIR_MAX_NUM_POINT_PER_GROUP);
 			}
 		}
+		
+		// Transcoded positions (optional)
+		const FVector3f TranscodingPositionOffset = -HairStrands.BoundingBox.GetExtent();
+		const FVector3f TranscodingPositionScale  =  HairStrands.BoundingBox.GetExtent() * 2.f;
+		const uint32 TranscodedPositionChunkElementCount = 3u;
+		const uint32 TranscodedPositionChunkStride = sizeof(FTranscodedHairPositions);
+		TArray<FTranscodedHairPositions> OutTranscodedPositions;
+		if (DoesHairStrandsSupportCompressedPosition() && bAllowTranscoding)
+		{
+			bool bHasTranscodedPositions = false;
+			static_assert(sizeof(FTranscodedHairPositions) == 16u);
+			static_assert(sizeof(FTranscodedHairPositions) == FCompressedHairPositionsStrideInBytes);
+
+			auto To10bitsPosition = [TranscodingPositionOffset, TranscodingPositionScale](const FHairStrandsPositionFormat::Type& In, bool& bIsResconstructionErrorLowerThanThreshold)
+			{
+				const FVector3f PF = FVector3f(In.X, In.Y, In.Z);
+				const FVector3f NP = (PF - TranscodingPositionOffset) / TranscodingPositionScale;
+
+				const float bPrecisionError = 0.001f;
+				check(NP.X >= 0-bPrecisionError && NP.X <= 1+bPrecisionError); 
+				check(NP.Y >= 0-bPrecisionError && NP.Y <= 1+bPrecisionError); 
+				check(NP.Z >= 0-bPrecisionError && NP.Z <= 1+bPrecisionError); 
+
+				FTranscodedHairPositions::FPosition Out;
+				Out.X 	= NP.X * 1023.f;
+				Out.Y 	= NP.Y * 1023.f;
+				Out.Z 	= NP.Z * 1023.f;
+				Out.Type= In.Type;
+
+				// Relative error
+				{
+					const FVector3f ReconstructedNP = FVector3f(Out.X / 1023.f, Out.Y / 1023.f, Out.Z / 1023.f);
+					const float fMaxError = (ReconstructedNP - NP).GetAbsMax();
+
+					const float ErrorThreshold = 0.01f; // 1%
+					const bool bIsValid = fMaxError <= ErrorThreshold;
+					if (!bIsValid) { bIsResconstructionErrorLowerThanThreshold = false; }
+				}
+
+				// Absolute error
+				#if 0
+				{
+					const FVector3f ReconstructedPF = FVector3f(Out.X / 1023.f, Out.Y / 1023.f, Out.Z / 1023.f) * TranscodingPositionScale + TranscodingPositionOffset;
+					const float fMaxError = (ReconstructedPF - PF).GetAbsMax();
+
+					const float ErrorThreshold = 0.1f; // 1mm
+					const bool bIsValid = fMaxError <= ErrorThreshold;
+					if (!bIsValid) { bIsResconstructionErrorLowerThanThreshold = false; }
+				}
+				#endif
+
+				return Out;
+			};
+
+			const FHairStrandsPositionFormat::Type DummInPoint = {0,0,0,0,0,0};
+			const FTranscodedHairPositions::FPosition DummyOutPoint = {0,0,0,0};
+			const uint32 TranscodedPositionChunkCount = FMath::DivideAndRoundUp(uint32(OutPackedPositions.Num()), TranscodedPositionChunkElementCount);
+			OutTranscodedPositions.SetNum(TranscodedPositionChunkCount);
+			for (uint32 ChunkIt=0;ChunkIt<TranscodedPositionChunkCount;++ChunkIt)
+			{
+				static_assert(TranscodedPositionChunkElementCount == 3u);
+				const uint32 I0 = ChunkIt * TranscodedPositionChunkElementCount + 0;
+				const uint32 I1 = ChunkIt * TranscodedPositionChunkElementCount + 1;
+				const uint32 I2 = ChunkIt * TranscodedPositionChunkElementCount + 2;
+	
+				const FHairStrandsPositionFormat::Type P0 = OutPackedPositions.IsValidIndex(I0) ? OutPackedPositions[I0] : DummInPoint;
+				const FHairStrandsPositionFormat::Type P1 = OutPackedPositions.IsValidIndex(I1) ? OutPackedPositions[I1] : DummInPoint;
+				const FHairStrandsPositionFormat::Type P2 = OutPackedPositions.IsValidIndex(I2) ? OutPackedPositions[I2] : DummInPoint;
+	
+				bool bCanUseCompressedPosition = true;
+				FTranscodedHairPositions Out;
+				Out.CP0 = OutPackedPositions.IsValidIndex(I0) ? To10bitsPosition(P0, bCanUseCompressedPosition) : DummyOutPoint;
+				Out.CP1 = OutPackedPositions.IsValidIndex(I1) ? To10bitsPosition(P1, bCanUseCompressedPosition) : DummyOutPoint;
+				Out.CP2 = OutPackedPositions.IsValidIndex(I2) ? To10bitsPosition(P2, bCanUseCompressedPosition) : DummyOutPoint;
+
+				if (!bCanUseCompressedPosition)
+				{
+					OutTranscodedPositions.SetNum(0);
+					break;
+				}
+
+				const float UCoord0 = P0.UCoord / 255.f;
+				const float UCoord1 = P1.UCoord / 255.f;
+				const float UCoord2 = P2.UCoord / 255.f;
+
+				const bool bInnerCurve = P0.Type == 0 && P1.Type == 0 && P2.Type == 0;
+				if (bInnerCurve)
+				{
+					// Radius
+					Out.Attribute0.Radius0 = P0.Radius;
+					Out.Attribute0.Radius1 = P1.Radius;
+					Out.Attribute0.Radius2 = P2.Radius;	
+
+					// UCoord
+					Out.Attribute0.UCoord0 = UCoord0 * 63.f; // 6 bits
+					Out.Attribute0.UCoord2 = UCoord2 * 63.f; // 6 bits
+					Out.Attribute0.Interp  = float(UCoord1 - UCoord0) / FMath::Max(0.001f, UCoord2 - UCoord0) * 3.f; // 2 bits
+				}
+				else
+				{
+					// Radius
+					Out.Attribute1.Radius0 = P0.Radius;
+					Out.Attribute1.Radius1 = P1.Radius;
+					Out.Attribute1.Radius2 = P2.Radius;
+
+					// UCoord
+					bool bFirst = true;
+					switch (P0.Type)
+					{
+						case HAIR_CONTROLPOINT_START : /* = 0*/ break;
+						case HAIR_CONTROLPOINT_INSIDE: bFirst ? Out.Attribute1.UCoord0 = UCoord0 * 127u: Out.Attribute1.UCoord1 = UCoord0 * 127u; bFirst = false; break;
+						case HAIR_CONTROLPOINT_END   : /* = 1*/ break;
+					}
+
+					switch (P1.Type)
+					{
+						case HAIR_CONTROLPOINT_START : /* = 0*/ break;
+						case HAIR_CONTROLPOINT_INSIDE: bFirst ? Out.Attribute1.UCoord0 = UCoord1 * 127u: Out.Attribute1.UCoord1 = UCoord1 * 127u; bFirst = false; break;
+						case HAIR_CONTROLPOINT_END   : /* = 1*/ break;
+					}
+
+					switch (P2.Type)
+					{
+						case HAIR_CONTROLPOINT_START : /* = 0*/ break;
+						case HAIR_CONTROLPOINT_INSIDE: bFirst ? Out.Attribute1.UCoord0 = UCoord2 * 127u: Out.Attribute1.UCoord1 = UCoord2 * 127u; bFirst = false; break;
+						case HAIR_CONTROLPOINT_END   : /* = 1*/ break;
+					}
+				}
+
+				OutTranscodedPositions[ChunkIt] = Out;
+			}
+		}
 
 		OutBulkData.Header.BoundingBox.Min = (FVector)HairStrands.BoundingBox.Min;
 		OutBulkData.Header.BoundingBox.Max = (FVector)HairStrands.BoundingBox.Max;
@@ -415,6 +562,17 @@ namespace HairStrandsBuilder
 		OutBulkData.Header.Flags = FHairStrandsBulkData::DataFlags_HasData;
 		OutBulkData.Header.ImportedAttributes = HairStrands.GetAttributes();
 		OutBulkData.Header.ImportedAttributeFlags = HairStrands.GetAttributeFlags();
+		OutBulkData.Header.Transcoding.PositionOffset = FVector3f::ZeroVector;
+		OutBulkData.Header.Transcoding.PositionScale = FVector3f::ZeroVector;
+
+		// Transcoding
+		const bool bHasTranscodedPosition = OutTranscodedPositions.Num() > 0;
+		if (bHasTranscodedPosition)
+		{
+			OutBulkData.Header.Flags |= FHairStrandsBulkData::DataFlags_HasTranscodedPosition;
+			OutBulkData.Header.Transcoding.PositionOffset = TranscodingPositionOffset;
+			OutBulkData.Header.Transcoding.PositionScale = TranscodingPositionScale;
+		}
 
 		const uint32 UintToByte = 4;
 
@@ -578,8 +736,17 @@ namespace HairStrandsBuilder
 			}
 		}
 
-		CopyToBulkData<FHairStrandsPositionFormat>(OutBulkData.Data.Positions, OutPackedPositions);
-		ReportSize(TEXT("Positions"), OutPackedPositions);
+		if (bHasTranscodedPosition)
+		{
+			CopyToBulkData<FHairStrandsTranscodedPositionFormat>(OutBulkData.Data.TranscodedPositions, OutTranscodedPositions);
+			ReportSize(TEXT("Positions"), OutTranscodedPositions);
+		}
+		else
+		{
+			CopyToBulkData<FHairStrandsPositionFormat>(OutBulkData.Data.Positions, OutPackedPositions);
+			ReportSize(TEXT("Positions"), OutPackedPositions);
+		}
+
 		CopyToBulkData<FHairStrandsCurveFormat>(OutBulkData.Data.Curves, OutPackedCurves);
 		ReportSize(TEXT("PackedCurves"), OutPackedCurves);
 		CopyToBulkData<FHairStrandsPointToCurveFormat>(OutBulkData.Data.PointToCurve, OutPointToCurve);
@@ -604,12 +771,15 @@ namespace HairStrandsBuilder
 
 		OutBulkData.Header.Strides.CurveAttributeChunkElementCount = CurveAttributeChunkElementCount;
 		OutBulkData.Header.Strides.PointAttributeChunkElementCount = PointAttributeChunkElementCount;
+
+		OutBulkData.Header.Strides.TranscodedPositionChunkElementCount = TranscodedPositionChunkElementCount;
+		OutBulkData.Header.Strides.TranscodedPositionChunkStride = sizeof(FTranscodedHairPositions);
 	}
 
-	void BuildBulkData(const FHairStrandsDatas& HairStrands, FHairStrandsBulkData& OutBulkData)
+	void BuildBulkData(const FHairStrandsDatas& HairStrands, FHairStrandsBulkData& OutBulkData, bool bAllowTranscoding)
 	{
 		TArray<uint8> RandomSeeds;
-		BuildBulkData(HairStrands, RandomSeeds, OutBulkData);
+		BuildBulkData(HairStrands, RandomSeeds, OutBulkData, bAllowTranscoding);
 	}
 
 } // namespace HairStrandsBuilder
@@ -2206,7 +2376,8 @@ void FGroomBuilder::BuildData(
 void FGroomBuilder::BuildBulkData(
 	const FHairGroupInfo& InInfo,
 	const FHairStrandsDatas& InData,
-	FHairStrandsBulkData& OutBulkData)
+	FHairStrandsBulkData& OutBulkData, 
+	bool bAllowTranscoding)
 {
 	OutBulkData.Reset();
 
@@ -2219,7 +2390,7 @@ void FGroomBuilder::BuildBulkData(
 	{
 		CurveSeeds[Index] = Random.RandHelper(255);
 	}
-	HairStrandsBuilder::BuildBulkData(InData, CurveSeeds, OutBulkData);
+	HairStrandsBuilder::BuildBulkData(InData, CurveSeeds, OutBulkData, bAllowTranscoding);
 }
 
 void FGroomBuilder::BuildInterplationData(
diff --git a/Engine/Plugins/Runtime/HairStrands/Source/HairStrandsCore/Private/GroomRBFDeformer.cpp b/Engine/Plugins/Runtime/HairStrands/Source/HairStrandsCore/Private/GroomRBFDeformer.cpp
--- a/Engine/Plugins/Runtime/HairStrands/Source/HairStrandsCore/Private/GroomRBFDeformer.cpp
+++ b/Engine/Plugins/Runtime/HairStrands/Source/HairStrandsCore/Private/GroomRBFDeformer.cpp
@@ -514,8 +514,8 @@ static void ApplyDeformationToGroom(const TArray<FRBFDeformedPositions>& Deforme
 				FHairStrandsDatas GuidesData;
 				FGroomBuilder::BuildData(HairGroup, GroomAsset->GetHairGroupsInterpolation()[GroupIndex], HairGroupsInfo, StrandsData, GuidesData, ALLOW_CURVE_SHUFFLING);
 
-				FGroomBuilder::BuildBulkData(HairGroup.Info, GuidesData,  HairGroupsData.Guides.BulkData);
-				FGroomBuilder::BuildBulkData(HairGroup.Info, StrandsData, HairGroupsData.Strands.BulkData);
+				FGroomBuilder::BuildBulkData(HairGroup.Info, GuidesData,  HairGroupsData.Guides.BulkData, false /*bAllowCompression*/);
+				FGroomBuilder::BuildBulkData(HairGroup.Info, StrandsData, HairGroupsData.Strands.BulkData, true /*bAllowCompression*/);
 
 				FHairStrandsInterpolationDatas StrandsInterpolationData;
 				FGroomBuilder::BuildInterplationData(HairGroup.Info, StrandsData, GuidesData, InterpolationSettings, StrandsInterpolationData);
diff --git a/Engine/Plugins/Runtime/HairStrands/Source/HairStrandsCore/Private/GroomResources.cpp b/Engine/Plugins/Runtime/HairStrands/Source/HairStrandsCore/Private/GroomResources.cpp
--- a/Engine/Plugins/Runtime/HairStrands/Source/HairStrandsCore/Private/GroomResources.cpp
+++ b/Engine/Plugins/Runtime/HairStrands/Source/HairStrandsCore/Private/GroomResources.cpp
@@ -24,12 +24,17 @@
 #include "UObject/ReleaseObjectVersion.h"
 #include "NiagaraSystem.h"
 #include "Async/ParallelFor.h"
+#include "Shader.h"
+#include "GlobalShader.h"
+#include "ShaderParameters.h"
+#include "ShaderParameterStruct.h"
 #include "RenderGraph.h"
 #include "RenderGraphUtils.h"
 #include "RenderTargetPool.h"
 #include "RHIUtilities.h"
 #include "GroomBindingBuilder.h"
 #include "GroomRBFDeformer.h"
+#include "UnifiedBuffer.h"
 #if WITH_EDITORONLY_DATA
 #include "Containers/StringView.h"
 #include "DerivedDataCache.h"
@@ -60,9 +65,9 @@ bool ValidateHairBulkData()
 	return GHairStrandsBulkData_Validation > 0;
 }
 
-static void ConvertToExternalBufferWithViews(FRDGBuilder& GraphBuilder, FRDGBufferRef& InBuffer, FRDGExternalBuffer& OutBuffer, EPixelFormat Format = PF_Unknown)
+static void CreateExternalBufferViews(FRDGBuilder& GraphBuilder, FRDGBufferRef& InBuffer, FRDGExternalBuffer& OutBuffer, EPixelFormat Format = PF_Unknown)
 {
-	OutBuffer.Buffer = GraphBuilder.ConvertToExternalBuffer(InBuffer);
+	check(OutBuffer.Buffer);
 	if (EnumHasAnyFlags(InBuffer->Desc.Usage, BUF_ShaderResource))
 	{
 		OutBuffer.SRV = OutBuffer.Buffer->GetOrCreateSRV(GraphBuilder.RHICmdList, FRDGBufferSRVDesc(InBuffer, Format));
@@ -74,6 +79,12 @@ static void ConvertToExternalBufferWithViews(FRDGBuilder& GraphBuilder, FRDGBuff
 	OutBuffer.Format = Format;
 }
 
+static void ConvertToExternalBufferWithViews(FRDGBuilder& GraphBuilder, FRDGBufferRef& InBuffer, FRDGExternalBuffer& OutBuffer, EPixelFormat Format = PF_Unknown)
+{
+	OutBuffer.Buffer = GraphBuilder.ConvertToExternalBuffer(InBuffer);
+	CreateExternalBufferViews(GraphBuilder, InBuffer, OutBuffer, Format);
+}
+
 ////////////////////////////////////////////////////////////////////////////////////
 
 EHairResourceLoadingType GetHairResourceLoadingType(EHairGeometryType InGeometryType, int32 InLODIndex)
@@ -247,16 +258,107 @@ FRDGBufferRef InternalCreateStructuredBufferRDG_FromBulkData(FRDGBuilder& GraphB
 	return InternalCreateBufferRDG_FromBulkData(GraphBuilder, InBulkData, Out, PF_Unknown, Desc, DebugName, OwnerName);
 }
 
-FRDGBufferRef InternalCreateByteAddressBufferRDG_FromBulkData(FRDGBuilder& GraphBuilder, FByteBulkData& InBulkData, FRDGExternalBuffer& Out, const TCHAR* DebugName, const FName& OwnerName, EHairResourceUsageType UsageType)
-{
-	const FRDGBufferDesc Desc = ApplyUsage(FRDGBufferDesc::CreateByteAddressDesc(InBulkData.GetBulkDataSize()), UsageType);
-	return InternalCreateBufferRDG_FromBulkData(GraphBuilder, InBulkData, Out, PF_Unknown, Desc, DebugName, OwnerName);
-}
+///////////////////////////////////////////////////////////////////////////////////////////////////
+// HairBulkData loading
 
 
 ///////////////////////////////////////////////////////////////////////////////////////////////////
-// HairBulkData loading
+class FHairTranscodeCS : public FGlobalShader
+{
+	DECLARE_GLOBAL_SHADER(FHairTranscodeCS);
+	SHADER_USE_PARAMETER_STRUCT(FHairTranscodeCS, FGlobalShader);
+
+	class FTranscoding : SHADER_PERMUTATION_BOOL("PERMUTATION_TRANSCODING");
+	using FPermutationDomain = TShaderPermutationDomain<FTranscoding>;
+
+	BEGIN_SHADER_PARAMETER_STRUCT(FParameters, )
+		SHADER_PARAMETER(uint32, OffsetInBytes)
+		SHADER_PARAMETER(uint32, DataSizeInBytes)
+		SHADER_PARAMETER(uint32, TotalSizeInBytes)
+
+		SHADER_PARAMETER(uint32, UncompressedOffsetInBytes)
+		SHADER_PARAMETER(uint32, UncompressedTotalSizeInBytes)
+
+		SHADER_PARAMETER(FVector3f, PositionScale)
+		SHADER_PARAMETER(FVector3f, PositionOffset)
+
+		SHADER_PARAMETER_RDG_BUFFER_SRV(ByteAddressBuffer, InBuffer)
+		SHADER_PARAMETER_RDG_BUFFER_UAV(RWByteAddressBuffer, OutBuffer)
+	END_SHADER_PARAMETER_STRUCT()
 
+public:
+	static uint32 GetGroupSize() { return 1024u; }
+	static bool ShouldCompilePermutation(const FGlobalShaderPermutationParameters& Parameters) { return IsHairStrandsSupported(EHairStrandsShaderType::All, Parameters.Platform); }
+	static void ModifyCompilationEnvironment(const FGlobalShaderPermutationParameters& Parameters, FShaderCompilerEnvironment& OutEnvironment)
+	{
+		FGlobalShader::ModifyCompilationEnvironment(Parameters, OutEnvironment);
+		OutEnvironment.SetDefine(TEXT("SHADER_TRANSCODE"), 1);
+		OutEnvironment.SetDefine(TEXT("GROUP_SIZE"), GetGroupSize());
+	}
+};
+
+IMPLEMENT_GLOBAL_SHADER(FHairTranscodeCS, "/Engine/Private/HairStrands/HairStrandsTranscode.usf", "MainCS", SF_Compute);
+
+struct FHairTranscoding
+{
+	bool bEnabled = false;
+	FVector3f PositionOffset = FVector3f::ZeroVector;
+	FVector3f PositionScale = FVector3f::ZeroVector;
+
+	uint32 UncompressedTotalSizeInBytes = 0;
+	uint32 UncompressedOffsetInBytes = 0;
+};
+
+static void AddTranscodePass(
+	FRDGBuilder& GraphBuilder,
+	uint32 InOffsetInBytes,
+	uint32 InNewDataSizeInSize,
+	uint32 InTotalSizeInBytes,
+	FRDGBufferRef In,
+	FRDGBufferRef Out,
+	const FHairTranscoding& InTranscoding)
+{
+	const FGlobalShaderMap* ShaderMap = GetGlobalShaderMap(GMaxRHIFeatureLevel);
+
+	FHairTranscodeCS::FParameters* Parameters = GraphBuilder.AllocParameters<FHairTranscodeCS::FParameters>();
+	Parameters->TotalSizeInBytes = InTotalSizeInBytes;
+	Parameters->DataSizeInBytes = InNewDataSizeInSize;
+	Parameters->OffsetInBytes = InOffsetInBytes;
+	Parameters->InBuffer = GraphBuilder.CreateSRV(In);
+	Parameters->OutBuffer = GraphBuilder.CreateUAV(Out);
+
+	// Transcoding parameters
+	Parameters->UncompressedTotalSizeInBytes = InTranscoding.UncompressedTotalSizeInBytes;
+	Parameters->UncompressedOffsetInBytes = InTranscoding.UncompressedOffsetInBytes;
+	Parameters->PositionScale = InTranscoding.PositionScale;
+	Parameters->PositionOffset = InTranscoding.PositionOffset;
+
+	// Data needs to be 4-bytes aligned
+	check((Parameters->DataSizeInBytes & 0x3) == 0);
+	check((Parameters->OffsetInBytes & 0x3) == 0);
+
+	const uint32 Byte4Count  = Parameters->DataSizeInBytes >> 2u;
+	const uint32 Byte16Count = Parameters->DataSizeInBytes >> 4u;
+	const FIntVector DispatchGroupCount = InTranscoding.bEnabled ? 
+		FComputeShaderUtils::GetGroupCount(Byte16Count, FHairTranscodeCS::GetGroupSize()): 
+		FComputeShaderUtils::GetGroupCount(Byte4Count,  FHairTranscodeCS::GetGroupSize());
+
+	const FShaderParametersMetadata* ParametersMetadata = FHairTranscodeCS::FParameters::FTypeInfo::GetStructMetadata();
+
+	FHairTranscodeCS::FPermutationDomain PermutationVector;
+	PermutationVector.Set<FHairTranscodeCS::FTranscoding>(InTranscoding.bEnabled ? 1u : 0u);
+
+	TShaderMapRef<FHairTranscodeCS> ComputeShader(ShaderMap, PermutationVector);
+	ClearUnusedGraphResources(ComputeShader, ParametersMetadata, Parameters);
+	FComputeShaderUtils::AddPass(
+		GraphBuilder,
+		RDG_EVENT_NAME("HairStrands::Transcode(Transcoding:%d)", InTranscoding.bEnabled ? 1u: 0u),
+		ComputeShader,
+		Parameters,
+		DispatchGroupCount);
+}
+
+// Used for Buffer and StructuredBuffer
 static FRDGBufferRef InternalCreateBufferRDG_FromHairBulkData(FRDGBuilder& GraphBuilder, FHairBulkContainer& InData, const FRDGBufferRef& In, const FRDGBufferDesc& BufferDesc, const FRDGBufferDesc& UploadDesc, const TCHAR* DebugName, const FName& OwnerName)
 {
 	check(InData.ChunkRequest);
@@ -275,6 +377,7 @@ static FRDGBufferRef InternalCreateBufferRDG_FromHairBulkData(FRDGBuilder& Graph
 		FRDGBufferDesc NewBufferDesc = BufferDesc;
 		NewBufferDesc.Usage |= EBufferUsageFlags::UnorderedAccess;
 		FRDGBufferRef Out = GraphBuilder.CreateBuffer(NewBufferDesc, DebugName, ERDGBufferFlags::MultiFrame);
+		Out->SetOwnerName(OwnerName);
 
 		// 2. Copy existing data from the old buffer to the new buffer
 		check (Out->Desc.GetSize() <= In->Desc.GetSize());
@@ -286,6 +389,7 @@ static FRDGBufferRef InternalCreateBufferRDG_FromHairBulkData(FRDGBuilder& Graph
 		check(BufferDesc.GetSize() >= UploadDesc.GetSize());
 
 		FRDGBufferRef Out = GraphBuilder.CreateBuffer(BufferDesc, DebugName, ERDGBufferFlags::MultiFrame);
+		Out->SetOwnerName(OwnerName);
 		GraphBuilder.QueueBufferUpload(Out, InChunk.GetData(), InChunk.Size, ERDGInitialDataFlags::None);
 		InChunk.Release();
 		return Out;
@@ -300,7 +404,8 @@ static FRDGBufferRef InternalCreateBufferRDG_FromHairBulkData(FRDGBuilder& Graph
 			FRDGBufferDesc NewBufferDesc = BufferDesc;
 			NewBufferDesc.Usage |= EBufferUsageFlags::UnorderedAccess;
 			Out = GraphBuilder.CreateBuffer(NewBufferDesc, DebugName, ERDGBufferFlags::MultiFrame);
-	
+			Out->SetOwnerName(OwnerName);
+
 			// 1.2 Copy existing data from the old buffer to the new buffer
 			AddCopyBufferPass(GraphBuilder, Out, 0, In, 0, In->Desc.GetSize());
 		}
@@ -367,25 +472,87 @@ FRDGBufferRef InternalCreateStructuredBufferRDG_FromHairBulkData(FRDGBuilder& Gr
 	return nullptr;
 }
 
-FRDGBufferRef InternalCreateByteAddressBufferRDG_FromHairBulkData(FRDGBuilder& GraphBuilder, FHairBulkContainer& InChunk, FRDGExternalBuffer& Out, const TCHAR* DebugName, const FName& OwnerName, EHairResourceUsageType UsageType)
+FRDGBufferRef InternalCreateByteAddressBufferRDG_FromHairBulkData(FRDGBuilder& GraphBuilder, FHairBulkContainer& InHairBulkContainer, FRDGExternalBuffer& OutExternal, const TCHAR* DebugName, const FName& OwnerName, EHairResourceUsageType UsageType, FHairTranscoding InTranscoding = FHairTranscoding())
 {
-	// Fallback for non-streamble resources (e.g. guides)
-	if (InChunk.ChunkRequest == nullptr)
+	FHairStreamingRequest::FChunk* Chunk = InHairBulkContainer.ChunkRequest;
+	FByteBulkData* BulkData = &InHairBulkContainer.Data;
+
+	// Fallback for non-streamble resources (e.g. guides, cards guides, ...)
+	if (Chunk == nullptr)
 	{
-		return InternalCreateByteAddressBufferRDG_FromBulkData(GraphBuilder, InChunk.Data, Out, DebugName, OwnerName, UsageType);
-	}
+		const FRDGBufferDesc Desc = ApplyUsage(FRDGBufferDesc::CreateByteAddressDesc(BulkData->GetBulkDataSize()), UsageType);
+		//return InternalCreateBufferRDG_FromBulkData(GraphBuilder, LocalBulkData, Out, PF_Unknown, Desc, DebugName, OwnerName);
+		InternalSetBulkDataFlags(*BulkData);
 
-	const FRDGBufferRef In = Out.Buffer ? Register(GraphBuilder, Out, ERDGImportedBufferFlags::None).Buffer : nullptr;
-	const FRDGBufferDesc BufferDesc = ApplyUsage(FRDGBufferDesc::CreateByteAddressDesc(InChunk.ChunkRequest->TotalSize), UsageType);
-	const FRDGBufferDesc UploadDesc = ApplyUsage(FRDGBufferDesc::CreateByteAddressDesc(InChunk.ChunkRequest->Size), UsageType);
-	if (BufferDesc.GetSize() == 0) 	{ Out.Buffer = nullptr; return nullptr; }
-	if (UploadDesc.GetSize()==0) 	{ return nullptr; }
-	if (FRDGBufferRef Buffer = InternalCreateBufferRDG_FromHairBulkData(GraphBuilder, InChunk, In, BufferDesc, UploadDesc, DebugName, OwnerName))
+		const uint32 DataSizeInBytes = Desc.GetSize();
+		check(BulkData->GetBulkDataSize() >= DataSizeInBytes);
+		if (DataSizeInBytes == 0)
+		{
+			OutExternal.Buffer = nullptr;
+			return nullptr;
+		}
+	
+		const uint8* Data = (const uint8*)BulkData->Lock(LOCK_READ_ONLY);
+		FRDGBufferRef OutBuffer = GraphBuilder.CreateBuffer(Desc, DebugName, ERDGBufferFlags::MultiFrame);
+		OutBuffer->SetOwnerName(OwnerName);
+		if (Data && DataSizeInBytes)
+		{
+			#if !WITH_EDITORONLY_DATA
+			if (ReleaseAfterUse())
+			{
+				GraphBuilder.QueueBufferUpload(OutBuffer, Data, DataSizeInBytes, [BulkData](const void* Ptr) { BulkData->Unlock(); });
+			}
+			else
+			#endif
+			{
+				GraphBuilder.QueueBufferUpload(OutBuffer, Data, DataSizeInBytes, ERDGInitialDataFlags::None);  // Copy data internally
+				BulkData->Unlock();
+			}
+		}
+		ConvertToExternalBufferWithViews(GraphBuilder, OutBuffer, OutExternal, PF_Unknown);
+		return OutBuffer;
+	}
+	else // Streamable resources
 	{
-		ConvertToExternalBufferWithViews(GraphBuilder, Buffer, Out);
-		return Buffer;
+		// Size of the final buffer, after decompression if any
+		const uint32 BufferTotalSizeInBytes = InTranscoding.bEnabled ? InTranscoding.UncompressedTotalSizeInBytes : Chunk->TotalSize;
+		const FRDGBufferRef In = OutExternal.Buffer ? Register(GraphBuilder, OutExternal, ERDGImportedBufferFlags::None).Buffer : nullptr;
+		const FRDGBufferDesc BufferDesc = ApplyUsage(FRDGBufferDesc::CreateByteAddressDesc(BufferTotalSizeInBytes), UsageType);
+		const FRDGBufferDesc UploadDesc = ApplyUsage(FRDGBufferDesc::CreateByteAddressDesc(Chunk->Size), UsageType);
+
+		if (BufferDesc.GetSize()==0) { OutExternal.Buffer = nullptr; return nullptr; }
+		if (UploadDesc.GetSize()==0) { return nullptr; }
+	
+		const bool bDeallocate  = In != nullptr && Chunk->Status == FHairStreamingRequest::FChunk::EStatus::Unloading;
+		const bool bUploadData	= Chunk->Size > 0;
+	
+		// If no shrink, no create, no append, no action needed
+		if (!bDeallocate && !bUploadData)
+		{
+			return nullptr;
+		}
+	
+		// 1. Resize buffer to fit the new required size (shrink/increase)
+		// Ensure the buffer is 16 bytes aligned for resources copy during resizing
+		const uint32 TotalSizeInBytes = FMath::DivideAndRoundUp(BufferDesc.NumElements * BufferDesc.BytesPerElement, 16u) * 16u;
+		FRDGBufferRef OutBuffer = ResizeByteAddressBufferIfNeeded(GraphBuilder, OutExternal.Buffer, TotalSizeInBytes, DebugName);
+		OutBuffer->SetOwnerName(OwnerName);
+		CreateExternalBufferViews(GraphBuilder, OutBuffer, OutExternal, PF_Unknown);
+	
+		// 2. Append data
+		if (!bDeallocate && bUploadData)
+		{
+			// 2.1 Upload new data
+			FRDGBufferRef UploadBuffer = GraphBuilder.CreateBuffer(UploadDesc, DebugName, ERDGBufferFlags::MultiFrame);
+			UploadBuffer->SetOwnerName(OwnerName);
+			GraphBuilder.QueueBufferUpload(UploadBuffer, Chunk->GetData(), Chunk->Size, ERDGInitialDataFlags::None); // Copy data internally
+			Chunk->Release();
+	
+			// 2.2 Append new data to the new/existing buffer
+			AddTranscodePass(GraphBuilder, Chunk->Offset, Chunk->Size, Chunk->TotalSize, UploadBuffer, OutBuffer, InTranscoding);
+		}
+		return OutBuffer;
 	}
-	return nullptr;
 }
 
 ///////////////////////////////////////////////////////////////////////////////////////////////////
@@ -1027,7 +1194,37 @@ void FHairStrandsRestResource::InternalAllocate(FRDGBuilder& GraphBuilder)
 
 	// 1. Lock data, which force the loading data from files (on non-editor build/cooked data). These data are then uploaded to the GPU
 	// 2. A local copy is done by the buffer uploader. This copy is discarded once the uploading is done.
-	InternalCreateByteAddressBufferRDG_FromHairBulkData(GraphBuilder, BulkData.Data.Positions, PositionBuffer, ToHairResourceDebugName(HAIRSTRANDS_RESOUCE_NAME(CurveType, Hair.StrandsRest_PositionBuffer), ResourceName), OwnerName, EHairResourceUsageType::Static);
+	const bool bTranscodedPosition = !!(BulkData.Header.Flags & FHairStrandsBulkData::EDataFlags::DataFlags_HasTranscodedPosition);
+	if (bTranscodedPosition)
+	{
+		FHairTranscoding PositionTranscoding;
+		PositionTranscoding.bEnabled						= true;
+		PositionTranscoding.PositionOffset 					= BulkData.Header.Transcoding.PositionOffset;
+		PositionTranscoding.PositionScale 					= BulkData.Header.Transcoding.PositionScale;
+		PositionTranscoding.UncompressedOffsetInBytes 		= 0;
+		PositionTranscoding.UncompressedTotalSizeInBytes 	= 0;
+
+		if (BulkData.Data.TranscodedPositions.ChunkRequest)
+		{
+			// Compressed data stored into chunk data
+			const uint32 CompressedDataOffsetBytes = BulkData.Data.TranscodedPositions.ChunkRequest->Offset;
+			const uint32 CompressedDataSizeInBytes = BulkData.Data.TranscodedPositions.ChunkRequest->TotalSize;
+
+			// Number of compressed chunks
+			const uint32 CompressedChunkCount_Offset = (CompressedDataOffsetBytes / FCompressedHairPositionsStrideInBytes);
+			const uint32 CompressedChunkCount_Total  = (CompressedDataSizeInBytes / FCompressedHairPositionsStrideInBytes);
+
+			// Uncompressed size/offset, i.e., size and offset of the final data (faster transcoding)
+			PositionTranscoding.UncompressedOffsetInBytes 		= CompressedChunkCount_Offset * HAIR_POINT_COUNT_PER_COMPRESSED_POSITION_CHUNK * FPackedHairPositionStrideInBytes;
+			PositionTranscoding.UncompressedTotalSizeInBytes 	= CompressedChunkCount_Total  * HAIR_POINT_COUNT_PER_COMPRESSED_POSITION_CHUNK * FPackedHairPositionStrideInBytes;
+		}
+
+		InternalCreateByteAddressBufferRDG_FromHairBulkData(GraphBuilder, BulkData.Data.TranscodedPositions, PositionBuffer, ToHairResourceDebugName(HAIRSTRANDS_RESOUCE_NAME(CurveType, Hair.StrandsRest_PositionBuffer), ResourceName), OwnerName, EHairResourceUsageType::Static, PositionTranscoding);
+	}
+	else
+	{
+		InternalCreateByteAddressBufferRDG_FromHairBulkData(GraphBuilder, BulkData.Data.Positions, PositionBuffer, ToHairResourceDebugName(HAIRSTRANDS_RESOUCE_NAME(CurveType, Hair.StrandsRest_PositionBuffer), ResourceName), OwnerName, EHairResourceUsageType::Static);
+	}
 	InternalCreateByteAddressBufferRDG_FromHairBulkData(GraphBuilder, BulkData.Data.CurveAttributes, CurveAttributeBuffer, ToHairResourceDebugName(HAIRSTRANDS_RESOUCE_NAME(CurveType, Hair.StrandsRest_CurveAttributeBuffer), ResourceName), OwnerName, EHairResourceUsageType::Static);
 	if (!!(BulkData.Header.Flags & FHairStrandsBulkData::DataFlags_HasPointAttribute))
 	{
diff --git a/Engine/Plugins/Runtime/HairStrands/Source/HairStrandsCore/Private/HairStrandsDatas.cpp b/Engine/Plugins/Runtime/HairStrands/Source/HairStrandsCore/Private/HairStrandsDatas.cpp
--- a/Engine/Plugins/Runtime/HairStrands/Source/HairStrandsCore/Private/HairStrandsDatas.cpp
+++ b/Engine/Plugins/Runtime/HairStrands/Source/HairStrandsCore/Private/HairStrandsDatas.cpp
@@ -552,14 +552,18 @@ void FHairStrandsBulkData::SerializeHeader(FArchive& Ar, UObject* Owner)
 
 	Ar << Header.Strides.PositionStride;
 	Ar << Header.Strides.CurveStride;
-
 	Ar << Header.Strides.PointToCurveChunkStride;
-	Ar << Header.Strides.PointToCurveChunkElementCount;
-
 	Ar << Header.Strides.CurveAttributeChunkStride;
 	Ar << Header.Strides.PointAttributeChunkStride;
+	Ar << Header.Strides.TranscodedPositionChunkStride;
+
+	Ar << Header.Strides.PointToCurveChunkElementCount;
 	Ar << Header.Strides.CurveAttributeChunkElementCount;
 	Ar << Header.Strides.PointAttributeChunkElementCount;
+	Ar << Header.Strides.TranscodedPositionChunkElementCount;
+
+	Ar << Header.Transcoding.PositionOffset;
+	Ar << Header.Transcoding.PositionScale;
 }
 
 void FHairStrandsBulkData::GetResourceVersion(FArchive& Ar) const
@@ -592,17 +596,25 @@ void FHairStrandsBulkData::GetResources(FHairStrandsBulkCommon::FQuery& Out)
 	const uint32 PointAttributeSize = GetPointAttributeSizeInBytes(PointCount);
 	const uint32 CurveAttributeSize = GetCurveAttributeSizeInBytes(CurveCount);
 	const uint32 PointToCurveIndexSize = GetPointToCurveSizeInBytes(PointCount);
+	const uint32 TranscodedPositionSize = GetTranscodedPositionSizeInBytes(PointCount);
 
 	if (!!(Header.Flags & DataFlags_HasData))
 	{
-		Out.Add(Data.Positions, 			TEXT("_Positions"), 		Data.Positions.LoadedSize, 		PointCount * Header.Strides.PositionStride);
-		Out.Add(Data.CurveAttributes, 		TEXT("_CurveAttributes"), 	Data.CurveAttributes.LoadedSize,CurveAttributeSize);
+		if (!!(Header.Flags & DataFlags_HasTranscodedPosition))
+		{
+			Out.Add(Data.TranscodedPositions,TEXT("_TranscodedPositions"), 	Data.TranscodedPositions.LoadedSize,TranscodedPositionSize);
+		}
+		else
+		{
+			Out.Add(Data.Positions, 		TEXT("_Positions"), 			Data.Positions.LoadedSize, 			PointCount * Header.Strides.PositionStride);
+		}
+		Out.Add(Data.CurveAttributes, 		TEXT("_CurveAttributes"), 		Data.CurveAttributes.LoadedSize,	CurveAttributeSize);
 		if (Header.Flags & DataFlags_HasPointAttribute)
 		{
-			Out.Add(Data.PointAttributes, 	TEXT("_PointAttributes"), 	Data.PointAttributes.LoadedSize,PointAttributeSize);
+			Out.Add(Data.PointAttributes, 	TEXT("_PointAttributes"), 		Data.PointAttributes.LoadedSize,	PointAttributeSize);
 		}
-		Out.Add(Data.PointToCurve, 			TEXT("_PointToCurve"), 		Data.PointToCurve.LoadedSize,	PointToCurveIndexSize);
-		Out.Add(Data.Curves, 				TEXT("_Curves"), 			Data.Curves.LoadedSize, 		CurveCount * Header.Strides.CurveStride);
+		Out.Add(Data.PointToCurve, 			TEXT("_PointToCurve"), 			Data.PointToCurve.LoadedSize,		PointToCurveIndexSize);
+		Out.Add(Data.Curves, 				TEXT("_Curves"), 				Data.Curves.LoadedSize, 			CurveCount * Header.Strides.CurveStride);
 	}
 }
 
@@ -610,6 +622,7 @@ uint32 FHairStrandsBulkData::GetSize() const
 {
 	uint32 Out = 0;
 	Out += Data.Positions.GetBulkDataSize();
+	Out += Data.TranscodedPositions.GetBulkDataSize();
 	Out += Data.CurveAttributes.GetBulkDataSize();
 	Out += Data.PointAttributes.GetBulkDataSize();
 	Out += Data.PointToCurve.GetBulkDataSize();
@@ -635,6 +648,7 @@ void FHairStrandsBulkData::Reset()
 	}
 	// Deallocate memory if needed
 	Data.Positions.RemoveBulkData();
+	Data.TranscodedPositions.RemoveBulkData();
 	Data.CurveAttributes.RemoveBulkData();
 	Data.PointAttributes.RemoveBulkData();
 	Data.PointToCurve.RemoveBulkData();
@@ -642,6 +656,7 @@ void FHairStrandsBulkData::Reset()
 
 	// Reset the bulk byte buffer to ensure the (serialize) data size is reset to 0
 	Data.Positions 		= FHairBulkContainer();
+	Data.TranscodedPositions= FHairBulkContainer();
 	Data.CurveAttributes= FHairBulkContainer();
 	Data.PointAttributes= FHairBulkContainer();
 	Data.PointToCurve	= FHairBulkContainer();
@@ -651,6 +666,7 @@ void FHairStrandsBulkData::Reset()
 void FHairStrandsBulkData::ResetLoadedSize()
 {
 	Data.Positions.LoadedSize		= 0;
+	Data.TranscodedPositions.LoadedSize		= 0;
 	Data.CurveAttributes.LoadedSize	= 0;
 	Data.PointAttributes.LoadedSize	= 0;
 	Data.PointToCurve.LoadedSize	= 0;
diff --git a/Engine/Plugins/Runtime/HairStrands/Source/HairStrandsCore/Public/GroomBuilder.h b/Engine/Plugins/Runtime/HairStrands/Source/HairStrandsCore/Public/GroomBuilder.h
--- a/Engine/Plugins/Runtime/HairStrands/Source/HairStrandsCore/Public/GroomBuilder.h
+++ b/Engine/Plugins/Runtime/HairStrands/Source/HairStrandsCore/Public/GroomBuilder.h
@@ -87,7 +87,8 @@ struct HAIRSTRANDSCORE_API FGroomBuilder
 	static void BuildBulkData(
 		const FHairGroupInfo& InInfo,
 		const FHairStrandsDatas& InData,
-		FHairStrandsBulkData& OutBulkData);
+		FHairStrandsBulkData& OutBulkData,
+		bool bAllowCompression);
 
 	// 4. Build interplation data based on the hairStrands data
 	static void BuildInterplationData(
diff --git a/Engine/Plugins/Runtime/HairStrands/Source/HairStrandsCore/Public/HairStrandsDatas.h b/Engine/Plugins/Runtime/HairStrands/Source/HairStrandsCore/Public/HairStrandsDatas.h
--- a/Engine/Plugins/Runtime/HairStrands/Source/HairStrandsCore/Public/HairStrandsDatas.h
+++ b/Engine/Plugins/Runtime/HairStrands/Source/HairStrandsCore/Public/HairStrandsDatas.h
@@ -45,6 +45,49 @@ struct FPackedHairVertex
 	uint8 Type : 2;
 };
 
+struct FTranscodedHairPositions
+{
+	typedef FUintVector4 BulkType;
+
+	struct FPosition
+	{
+		uint32 X    : 10;
+		uint32 Y    : 10;
+		uint32 Z    : 10;
+		uint32 Type : 2;
+	};
+
+	struct FAttributePacking0
+	{
+		uint32 Radius0 : 6;
+		uint32 Radius1 : 6;
+		uint32 Radius2 : 6;
+
+		uint32 UCoord0 : 6;
+		uint32 UCoord2 : 6;
+		uint32 Interp  : 2;
+	};
+
+	struct FAttributePacking1
+	{
+		uint32 Radius0 : 6;
+		uint32 Radius1 : 6;
+		uint32 Radius2 : 6;
+
+		uint32 UCoord0 : 7;
+		uint32 UCoord1 : 7;
+	};
+
+	FPosition CP0;
+	FPosition CP1;
+	FPosition CP2;
+	union
+	{
+		FAttributePacking0 Attribute0;
+		FAttributePacking1 Attribute1;
+	};
+};
+
 struct FPackedHairAttribute0Vertex
 {
 	typedef uint16 BulkType;
@@ -70,6 +113,16 @@ struct FVector4_16
 }; 
 FArchive& operator<<(FArchive& Ar, FVector4_16& Vertex);
 
+struct FHairStrandsTranscodedPositionFormat
+{
+	typedef FTranscodedHairPositions Type;
+	typedef FTranscodedHairPositions::BulkType BulkType;
+	static const uint32 ComponentCount = 1;
+	static const uint32 SizeInByte = sizeof(Type);
+	static const EVertexElementType VertexElementType = VET_UShort4;
+	static const EPixelFormat Format = PF_Unknown;
+};
+
 struct FHairStrandsPositionFormat
 {
 	typedef FPackedHairVertex Type;
@@ -554,8 +607,9 @@ struct HAIRSTRANDSCORE_API FHairStrandsBulkData : FHairStrandsBulkCommon
 {
 	enum EDataFlags
 	{
-		DataFlags_HasData = 1,				// Contains valid data. Otherwise: Position, Attributes, ... are all empty
-		DataFlags_HasPointAttribute = 2,	// Contains point attribute data.
+		DataFlags_HasData = 0x1u,				// Contains valid data. Otherwise: Position, Attributes, ... are all empty
+		DataFlags_HasPointAttribute = 0x2,		// Contains point attribute data.
+		DataFlags_HasTranscodedPosition = 0x4u,	// Contains transcoded position
 	};
 
 	virtual void SerializeHeader(FArchive& Ar, UObject* Owner) override;
@@ -575,9 +629,10 @@ struct HAIRSTRANDSCORE_API FHairStrandsBulkData : FHairStrandsBulkCommon
 	const FBox& GetBounds() const { return Header.BoundingBox; }
 	uint32 GetSize() const;
 
-	uint32 GetCurveAttributeSizeInBytes(uint32 InCurveCount=HAIR_MAX_NUM_CURVE_PER_GROUP) const	{ return InCurveCount > 0 ? FMath::DivideAndRoundUp(FMath::Min(Header.CurveCount, InCurveCount), Header.Strides.CurveAttributeChunkElementCount) * Header.Strides.CurveAttributeChunkStride : 0; }
-	uint32 GetPointAttributeSizeInBytes(uint32 InPointCount=HAIR_MAX_NUM_POINT_PER_GROUP) const	{ return InPointCount > 0 ? FMath::DivideAndRoundUp(FMath::Min(Header.PointCount, InPointCount), Header.Strides.PointAttributeChunkElementCount) * Header.Strides.PointAttributeChunkStride : 0; }
-	uint32 GetPointToCurveSizeInBytes(uint32 InPointCount=HAIR_MAX_NUM_POINT_PER_GROUP) const	{ return InPointCount > 0 ? FMath::DivideAndRoundUp(FMath::Min(Header.PointCount, InPointCount), Header.Strides.PointToCurveChunkElementCount)   * Header.Strides.PointToCurveChunkStride   : 0; }
+	uint32 GetCurveAttributeSizeInBytes(uint32 InCurveCount=HAIR_MAX_NUM_CURVE_PER_GROUP) const			{ return InCurveCount > 0 ? FMath::DivideAndRoundUp(FMath::Min(Header.CurveCount, InCurveCount), Header.Strides.CurveAttributeChunkElementCount) * Header.Strides.CurveAttributeChunkStride : 0; }
+	uint32 GetPointAttributeSizeInBytes(uint32 InPointCount=HAIR_MAX_NUM_POINT_PER_GROUP) const			{ return InPointCount > 0 ? FMath::DivideAndRoundUp(FMath::Min(Header.PointCount, InPointCount), Header.Strides.PointAttributeChunkElementCount) * Header.Strides.PointAttributeChunkStride : 0; }
+	uint32 GetPointToCurveSizeInBytes(uint32 InPointCount=HAIR_MAX_NUM_POINT_PER_GROUP) const			{ return InPointCount > 0 ? FMath::DivideAndRoundUp(FMath::Min(Header.PointCount, InPointCount), Header.Strides.PointToCurveChunkElementCount)   * Header.Strides.PointToCurveChunkStride   : 0; }
+	uint32 GetTranscodedPositionSizeInBytes(uint32 InPointCount = HAIR_MAX_NUM_POINT_PER_GROUP) const 	{ return InPointCount > 0 ? FMath::DivideAndRoundUp(FMath::Min(Header.PointCount, InPointCount), Header.Strides.TranscodedPositionChunkElementCount) * Header.Strides.TranscodedPositionChunkStride : 0; }
 
 	struct FHeader
 	{
@@ -600,6 +655,13 @@ struct HAIRSTRANDSCORE_API FHairStrandsBulkData : FHairStrandsBulkCommon
 		// Map 'curve' count to 'point' count (used for CLOD)
 		TArray<uint32> CurveToPointCount;
 
+		// Data transcoding parameters
+		struct FTranscoding
+		{
+			FVector3f PositionOffset = FVector3f::OneVector;
+			FVector3f PositionScale = FVector3f::OneVector;
+		} Transcoding;
+
 		// Data strides
 		struct FStrides
 		{
@@ -608,21 +670,24 @@ struct HAIRSTRANDSCORE_API FHairStrandsBulkData : FHairStrandsBulkCommon
 			uint32 PointToCurveChunkStride = 0;
 			uint32 CurveAttributeChunkStride = 0;
 			uint32 PointAttributeChunkStride = 0;
+			uint32 TranscodedPositionChunkStride = 0;
 
 			// Number of element per chunk block
 			uint32 PointToCurveChunkElementCount = 0;
 			uint32 CurveAttributeChunkElementCount = 0;
 			uint32 PointAttributeChunkElementCount = 0;
+			uint32 TranscodedPositionChunkElementCount = 0;
 		} Strides;
 	} Header;
 
 	struct FData
 	{
-		FHairBulkContainer Positions;		// Size = PointCount
-		FHairBulkContainer CurveAttributes;	// Size = y*CurveCount (depends on the per-curve stored attributes)
-		FHairBulkContainer PointAttributes;	// Size = x*PointCount (depends on the per-point stored attributes)
-		FHairBulkContainer PointToCurve; 	// Size = PointCount
-		FHairBulkContainer Curves;			// Size = CurveCount
+		FHairBulkContainer Positions;			// Size = PointCount
+		FHairBulkContainer TranscodedPositions;	// Size = PointCount / TranscodedPositionChunkElementCount (=3)
+		FHairBulkContainer CurveAttributes;		// Size = y*CurveCount (depends on the per-curve stored attributes)
+		FHairBulkContainer PointAttributes;		// Size = x*PointCount (depends on the per-point stored attributes)
+		FHairBulkContainer PointToCurve; 		// Size = PointCount
+		FHairBulkContainer Curves;				// Size = CurveCount
 	} Data;
 };
 
diff --git a/Engine/Shaders/Private/HairStrands/HairStrandsPack.ush b/Engine/Shaders/Private/HairStrands/HairStrandsPack.ush
--- a/Engine/Shaders/Private/HairStrands/HairStrandsPack.ush
+++ b/Engine/Shaders/Private/HairStrands/HairStrandsPack.ush
@@ -9,16 +9,6 @@
 ////////////////////////////////////////////////////////////////////////////////////////////////////////////////////////////////
 // Hair control points
 
-// Hair control point layout
-// [                  uint2                   ]
-// [      x     ][              y             ]
-// [   32bits   ][           32bits           ]
-// [  16 ][  16 ][ 16  ][   8  ][   6  ][  2  ]
-// [ 0-15][16-32][ 0-15][ 16-23][ 24-29][30-31]
-// [Pos.x][Pos.y][Pos.z][CoordU][Radius][ Type]
-#define FPackedHairPosition uint2
-#define FPackedHairPositionStrideInBytes 8u
-
 FPackedHairPosition ReadPackedHairPosition(ByteAddressBuffer InBuffer, uint InPointIndex)
 {
 	return InBuffer.Load2(InPointIndex * FPackedHairPositionStrideInBytes);
diff --git a/Engine/Shaders/Private/HairStrands/HairStrandsTranscode.usf b/Engine/Shaders/Private/HairStrands/HairStrandsTranscode.usf
new file mode 100644
--- /dev/null
+++ b/Engine/Shaders/Private/HairStrands/HairStrandsTranscode.usf
@@ -0,0 +1,130 @@
+// Copyright Epic Games, Inc. All Rights Reserved.
+
+#pragma once
+
+#include "/Engine/Private/Common.ush"
+#include "/Engine/Shared/HairStrandsDefinitions.h"
+#include "../PackUnpack.ush"
+#include "HairStrandsPack.ush"
+
+#if SHADER_TRANSCODE
+
+uint OffsetInBytes;
+uint DataSizeInBytes;
+uint TotalSizeInBytes;
+
+uint UncompressedOffsetInBytes;
+uint UncompressedTotalSizeInBytes;
+
+float3 PositionScale;
+float3 PositionOffset;
+
+ByteAddressBuffer InBuffer;
+RWByteAddressBuffer OutBuffer;
+
+#if GROUP_SIZE != 1024
+#error Group size is not == 1024u. Please update this shader.
+#endif
+
+float UnpackR7(uint In) { return In / 127.f; }
+float UnpackR2(uint In) { return In / 3.f; }
+
+FHairControlPoint UnpackTranscoodedPosition(uint In)
+{
+	FHairControlPoint Out;
+	Out.Position.x = UnpackUnorm10(BitFieldExtractU32(In, 10,  0));
+	Out.Position.y = UnpackUnorm10(BitFieldExtractU32(In, 10, 10));
+	Out.Position.z = UnpackUnorm10(BitFieldExtractU32(In, 10, 20));
+	Out.Type       = BitFieldExtractU32(In, 2, 30);
+
+	Out.Position   = Out.Position * PositionScale + PositionOffset;
+	return Out;
+}
+
+void UnpackUCoord(inout FHairControlPoint InCP, uint DataW0, uint DataW1, inout uint bFirst)
+{
+	switch(InCP.Type)
+	{
+		case HAIR_CONTROLPOINT_START: 	InCP.UCoord = 0u; break;
+		case HAIR_CONTROLPOINT_END:		InCP.UCoord = 1u; break;
+		case HAIR_CONTROLPOINT_INSIDE:	InCP.UCoord = bFirst ? DataW0 : DataW1; bFirst = false; break;
+	}
+}
+
+[numthreads(GROUP_SIZE, 1, 1)]
+void MainCS(uint ThreadIndex : SV_DispatchThreadID)
+#if PERMUTATION_TRANSCODING
+{
+	// 3 CPs per 128bits/16bytes
+	const uint InAddressInBytes  = ThreadIndex * FCompressedHairPositionsStrideInBytes;
+
+	if (InAddressInBytes + FCompressedHairPositionsStrideInBytes <= DataSizeInBytes)
+	{
+		const uint4 Data = InBuffer.Load4(InAddressInBytes);
+	
+		// Sanity check
+		#if HAIR_POINT_COUNT_PER_COMPRESSED_POSITION_CHUNK != 3
+		#error Update hair point decompression code. Assume 3 CPs per compressed chunk
+		#endif
+
+		// Position - 3x32bits
+		// Each position - 10|10|10|2 -> 3x16bits
+		FHairControlPoint CP0 = UnpackTranscoodedPosition(Data.x);
+		FHairControlPoint CP1 = UnpackTranscoodedPosition(Data.y);
+		FHairControlPoint CP2 = UnpackTranscoodedPosition(Data.z);
+
+		// Radius - 18bits
+		{
+			CP0.WorldRadius = UnpackR6(BitFieldExtractU32(Data.w, 6, 0));
+			CP1.WorldRadius = UnpackR6(BitFieldExtractU32(Data.w, 6, 6));
+			CP2.WorldRadius = UnpackR6(BitFieldExtractU32(Data.w, 6, 12));
+		}
+
+		// CoordU - 14bits
+		if (CP0.Type==HAIR_CONTROLPOINT_INSIDE && 
+			CP1.Type==HAIR_CONTROLPOINT_INSIDE && 
+			CP2.Type==HAIR_CONTROLPOINT_INSIDE)
+		{
+			// The 3 CPs are consecutive and we compress them by storing begin/end points + lerping the middle point
+			CP0.UCoord = UnpackR6(BitFieldExtractU32(Data.w, 6, 18));
+			CP2.UCoord = UnpackR6(BitFieldExtractU32(Data.w, 6, 24));
+
+			const float S = UnpackR2(BitFieldExtractU32(Data.w, 2, 30));
+			CP1.UCoord = lerp(CP0.UCoord, CP2.UCoord, S);
+		}
+		else
+		{
+			// In this configuration 1 or 2 CPs are begin or end of a curve, meaning that their CoordU value is either 0 or 1
+			// We use the 14bit to store 2x7bits radius information
+			const uint DataW0 = UnpackR7(BitFieldExtractU32(Data.w, 7, 18));
+			const uint DataW1 = UnpackR7(BitFieldExtractU32(Data.w, 7, 25));
+
+			bool bFirst = true;
+			UnpackUCoord(CP0, DataW0, DataW1, bFirst);
+			UnpackUCoord(CP1, DataW0, DataW1, bFirst);
+			UnpackUCoord(CP2, DataW0, DataW1, bFirst);
+		}
+
+		const uint OutAddress0 = UncompressedOffsetInBytes + (ThreadIndex*3 + 0) * FPackedHairPositionStrideInBytes;
+		const uint OutAddress1 = UncompressedOffsetInBytes + (ThreadIndex*3 + 1) * FPackedHairPositionStrideInBytes;
+		const uint OutAddress2 = UncompressedOffsetInBytes + (ThreadIndex*3 + 2) * FPackedHairPositionStrideInBytes;
+
+		if (OutAddress0+FPackedHairPositionStrideInBytes <= UncompressedTotalSizeInBytes) { OutBuffer.Store2(OutAddress0, PackHairControlPoint(CP0, 0 /*PositionOffset*/, 1 /*Radius*/)); }
+		if (OutAddress1+FPackedHairPositionStrideInBytes <= UncompressedTotalSizeInBytes) { OutBuffer.Store2(OutAddress1, PackHairControlPoint(CP1, 0 /*PositionOffset*/, 1 /*Radius*/)); }
+		if (OutAddress2+FPackedHairPositionStrideInBytes <= UncompressedTotalSizeInBytes) { OutBuffer.Store2(OutAddress2, PackHairControlPoint(CP2, 0 /*PositionOffset*/, 1 /*Radius*/)); }
+	}
+}
+#else //PERMUTATION_TRANSCODING
+{
+	// Simple copy	
+	const uint InAddressInBytes  = ThreadIndex * 4u;
+	const uint OutAddressInBytes = OffsetInBytes + ThreadIndex * 4u;
+
+	if (InAddressInBytes < DataSizeInBytes && OutAddressInBytes < TotalSizeInBytes)
+	{
+		const uint InData = InBuffer.Load(InAddressInBytes);
+		OutBuffer.Store(OutAddressInBytes,InData);
+	}
+}
+#endif // PERMUTATION_TRANSCODING
+#endif // SHADER_TRANSCODE
\ No newline at end of file
diff --git a/Engine/Shaders/Shared/HairStrandsDefinitions.h b/Engine/Shaders/Shared/HairStrandsDefinitions.h
--- a/Engine/Shaders/Shared/HairStrandsDefinitions.h
+++ b/Engine/Shaders/Shared/HairStrandsDefinitions.h
@@ -119,3 +119,25 @@ FORCEINLINE bool HasHairFlags(uint32 In, uint32 Flags)
 {
 	return (In & Flags) != 0;
 }
+
+///////////////////////////////////////////////////////////////////////////////////////////////////
+// Types
+
+// Hair control point layout
+// [                  uint2                   ]
+// [      x     ][              y             ]
+// [   32bits   ][           32bits           ]
+// [  16 ][  16 ][ 16  ][   8  ][   6  ][  2  ]
+// [ 0-15][16-32][ 0-15][ 16-23][ 24-29][30-31]
+// [Pos.x][Pos.y][Pos.z][CoordU][Radius][ Type]
+#define FPackedHairPosition uint2
+#define FPackedHairPositionStrideInBytes 8u
+
+// Compressed Hair control point layout - Packed 4 points per 16-bytes
+// [                            uint4                                 ]
+// [Position0+Type][Position1+Type][Position2+Type][ Radius012|CoordU ]
+// [   32bits     ][   32bits     ][   32bits     ][      32bits      ]
+// [10][10][10][ 2][10][10][10][ 2][10][10][10][ 2][ 6][ 6][ 6][ 7][ 7]
+#define FCompressedHairPositions uint4
+#define FCompressedHairPositionsStrideInBytes 16u
+#define HAIR_POINT_COUNT_PER_COMPRESSED_POSITION_CHUNK 3
\ No newline at end of file
