Fix CIS comiplation issue on ARM platform.

diff --git a/Engine/Plugins/Runtime/HairStrands/Source/HairStrandsCore/Private/GroomBuilder.cpp b/Engine/Plugins/Runtime/HairStrands/Source/HairStrandsCore/Private/GroomBuilder.cpp
--- a/Engine/Plugins/Runtime/HairStrands/Source/HairStrandsCore/Private/GroomBuilder.cpp
+++ b/Engine/Plugins/Runtime/HairStrands/Source/HairStrandsCore/Private/GroomBuilder.cpp
@@ -526,21 +526,21 @@ namespace HairStrandsBuilder
 					switch (P0.Type)
 					{
 						case HAIR_CONTROLPOINT_START : /* = 0*/ break;
-						case HAIR_CONTROLPOINT_INSIDE: bFirst ? Out.Attribute1.UCoord0 = UCoord0 * 127u: Out.Attribute1.UCoord1 = UCoord0 * 127u; bFirst = false; break;
+						case HAIR_CONTROLPOINT_INSIDE: if (bFirst) { Out.Attribute1.UCoord0 = UCoord0 * 127u; } else { Out.Attribute1.UCoord1 = UCoord0 * 127u; } bFirst = false; break;
 						case HAIR_CONTROLPOINT_END   : /* = 1*/ break;
 					}
 
 					switch (P1.Type)
 					{
 						case HAIR_CONTROLPOINT_START : /* = 0*/ break;
-						case HAIR_CONTROLPOINT_INSIDE: bFirst ? Out.Attribute1.UCoord0 = UCoord1 * 127u: Out.Attribute1.UCoord1 = UCoord1 * 127u; bFirst = false; break;
+						case HAIR_CONTROLPOINT_INSIDE: if (bFirst) { Out.Attribute1.UCoord0 = UCoord1 * 127u; } else { Out.Attribute1.UCoord1 = UCoord1 * 127u; } bFirst = false; break;
 						case HAIR_CONTROLPOINT_END   : /* = 1*/ break;
 					}
 
 					switch (P2.Type)
 					{
 						case HAIR_CONTROLPOINT_START : /* = 0*/ break;
-						case HAIR_CONTROLPOINT_INSIDE: bFirst ? Out.Attribute1.UCoord0 = UCoord2 * 127u: Out.Attribute1.UCoord1 = UCoord2 * 127u; bFirst = false; break;
+						case HAIR_CONTROLPOINT_INSIDE: if (bFirst) { Out.Attribute1.UCoord0 = UCoord2 * 127u; } else { Out.Attribute1.UCoord1 = UCoord2 * 127u; } bFirst = false; break;
 						case HAIR_CONTROLPOINT_END   : /* = 1*/ break;
 					}
 				}
