* Add Point/Curve trimming report into groom editor.
* Remove warning report during building

This information is now only displayed at import time and within groom editor.

diff --git a/Engine/Plugins/Runtime/HairStrands/Source/HairStrandsCore/Private/GroomBuilder.cpp b/Engine/Plugins/Runtime/HairStrands/Source/HairStrandsCore/Private/GroomBuilder.cpp
--- a/Engine/Plugins/Runtime/HairStrands/Source/HairStrandsCore/Private/GroomBuilder.cpp
+++ b/Engine/Plugins/Runtime/HairStrands/Source/HairStrandsCore/Private/GroomBuilder.cpp
@@ -51,7 +51,7 @@ bool DoesHairStrandsSupportCompressedPosition();
 
 FString FGroomBuilder::GetVersion()
 {
-	return TEXT("v13_Dummy10");
+	return TEXT("v14");
 }
 
 // For debug purpose
@@ -182,7 +182,7 @@ namespace HairStrandsBuilder
 	}
 
 	/** Build the bulk/packed datas for gpu rendering/simulation */
-	void BuildBulkData(const FHairStrandsDatas& HairStrands, const TArray<uint8>& RandomSeeds, FHairStrandsBulkData& OutBulkData, bool bAllowTranscoding)
+	void BuildBulkData(const FHairStrandsDatas& HairStrands, const TArray<uint8>& RandomSeeds, FHairStrandsBulkData& OutBulkData, bool bAllowTranscoding, uint32 InGroupFlags)
 	{
 		TRACE_CPUPROFILER_EVENT_SCOPE(HairStrandsBuilder::BuildBulkData);
 
@@ -565,6 +565,17 @@ namespace HairStrandsBuilder
 		OutBulkData.Header.Transcoding.PositionOffset = FVector3f::ZeroVector;
 		OutBulkData.Header.Transcoding.PositionScale = FVector3f::ZeroVector;
 
+		// Transfer group trimmed pt./curve info from group to header. This allows to display this information 
+		// in the groom editor by reading the bulk data header
+		if (InGroupFlags & uint32(EHairGroupInfoFlags::HasTrimmedCurve))
+		{
+			OutBulkData.Header.Flags |= FHairStrandsBulkData::DataFlags_HasTrimmedCurve;
+		}
+		if (InGroupFlags & uint32(EHairGroupInfoFlags::HasTrimmedPoint))
+		{
+			OutBulkData.Header.Flags |= FHairStrandsBulkData::DataFlags_HasTrimmedPoint;
+		}
+	
 		// Transcoding
 		const bool bHasTranscodedPosition = OutTranscodedPositions.Num() > 0;
 		if (bHasTranscodedPosition)
@@ -775,13 +786,6 @@ namespace HairStrandsBuilder
 		OutBulkData.Header.Strides.TranscodedPositionChunkElementCount = TranscodedPositionChunkElementCount;
 		OutBulkData.Header.Strides.TranscodedPositionChunkStride = sizeof(FTranscodedHairPositions);
 	}
-
-	void BuildBulkData(const FHairStrandsDatas& HairStrands, FHairStrandsBulkData& OutBulkData, bool bAllowTranscoding)
-	{
-		TArray<uint8> RandomSeeds;
-		BuildBulkData(HairStrands, RandomSeeds, OutBulkData, bAllowTranscoding);
-	}
-
 } // namespace HairStrandsBuilder
 
 namespace HairInterpolationBuilder
@@ -2072,7 +2076,10 @@ bool FGroomBuilder::BuildHairDescriptionGroups(const FHairDescription& HairDescr
 		// If the current curve has more control point than allows, issue a warning
 		if (CurveNumVertices > HAIR_MAX_NUM_POINT_PER_CURVE && !bPointCountWarningIssued)
 		{
-			UE_LOG(LogGroomBuilder, Warning, TEXT("[Groom] Groom contains strands with more than %d control points. Control points beyond that limit will be trimmed"), uint32(HAIR_MAX_NUM_POINT_PER_CURVE));
+			// Do not display this during build time. 
+			// * This is reported at import time
+			// * This is visible in Groom editor
+			//UE_LOG(LogGroomBuilder, Warning, TEXT("[Groom] Groom contains strands with more than %d control points. Control points beyond that limit will be trimmed"), uint32(HAIR_MAX_NUM_POINT_PER_CURVE));
 			Group.Info.Flags |= uint32(EHairGroupInfoFlags::HasTrimmedPoint);
 		}
 
@@ -2081,7 +2088,10 @@ bool FGroomBuilder::BuildHairDescriptionGroups(const FHairDescription& HairDescr
 		{
 			if (!bCurveCountWarningIssued)
 			{
-				UE_LOG(LogGroomBuilder, Warning, TEXT("[Groom] Group has more than %d curves per group. Curve beyong that limit won't be part of the groom"), uint32(HAIR_MAX_NUM_CURVE_PER_GROUP));
+				// Do not display this during build time. 
+				// * This is reported at import time
+				// * This is visible in Groom editor
+				//UE_LOG(LogGroomBuilder, Warning, TEXT("[Groom] Group has more than %d curves per group. Curve beyong that limit won't be part of the groom"), uint32(HAIR_MAX_NUM_CURVE_PER_GROUP));
 				bCurveCountWarningIssued = true;
 			}
 			Group.Info.Flags |= uint32(EHairGroupInfoFlags::HasTrimmedCurve);
@@ -2390,7 +2400,7 @@ void FGroomBuilder::BuildBulkData(
 	{
 		CurveSeeds[Index] = Random.RandHelper(255);
 	}
-	HairStrandsBuilder::BuildBulkData(InData, CurveSeeds, OutBulkData, bAllowTranscoding);
+	HairStrandsBuilder::BuildBulkData(InData, CurveSeeds, OutBulkData, bAllowTranscoding, InInfo.Flags);
 }
 
 void FGroomBuilder::BuildInterplationData(
diff --git a/Engine/Plugins/Runtime/HairStrands/Source/HairStrandsCore/Public/HairStrandsDatas.h b/Engine/Plugins/Runtime/HairStrands/Source/HairStrandsCore/Public/HairStrandsDatas.h
--- a/Engine/Plugins/Runtime/HairStrands/Source/HairStrandsCore/Public/HairStrandsDatas.h
+++ b/Engine/Plugins/Runtime/HairStrands/Source/HairStrandsCore/Public/HairStrandsDatas.h
@@ -610,6 +610,8 @@ struct HAIRSTRANDSCORE_API FHairStrandsBulkData : FHairStrandsBulkCommon
 		DataFlags_HasData = 0x1u,				// Contains valid data. Otherwise: Position, Attributes, ... are all empty
 		DataFlags_HasPointAttribute = 0x2,		// Contains point attribute data.
 		DataFlags_HasTranscodedPosition = 0x4u,	// Contains transcoded position
+		DataFlags_HasTrimmedCurve = 0x8u,		// Source data had more curves per group than the limit and some curves were trimmed
+		DataFlags_HasTrimmedPoint = 0x10u,		// Source data had more points per curve than the limit and some points were trimmed
 	};
 
 	virtual void SerializeHeader(FArchive& Ar, UObject* Owner) override;
diff --git a/Engine/Plugins/Runtime/HairStrands/Source/HairStrandsEditor/Private/GroomAssetDetails.cpp b/Engine/Plugins/Runtime/HairStrands/Source/HairStrandsEditor/Private/GroomAssetDetails.cpp
--- a/Engine/Plugins/Runtime/HairStrands/Source/HairStrandsEditor/Private/GroomAssetDetails.cpp
+++ b/Engine/Plugins/Runtime/HairStrands/Source/HairStrandsEditor/Private/GroomAssetDetails.cpp
@@ -99,6 +99,105 @@ TSharedRef<SUniformGridPanel> MakeHairStrandsAttributeInfoGrid(const FSlateFontI
 	return Grid;
 }
 
+///////////////////////////////////////////////////////////////////////////////////////////////////////////
+// Array panel for hair strands curve infos
+TSharedRef<SUniformGridPanel> MakeHairStrandsCurveInfoGrid(const FSlateFontInfo& DetailFontInfo, const FHairStrandsBulkData::FHeader& InHeader)
+{
+	TSharedRef<SUniformGridPanel> Grid = SNew(SUniformGridPanel).SlotPadding(2.0f);
+	
+	// Min. point per curve
+	Grid->AddSlot(0, 1) // x, y
+	.HAlign(HAlign_Left)
+	[
+		SNew(STextBlock)
+		.Font(DetailFontInfo)
+		.Text(LOCTEXT("HairInfo_CPPerCurve", "Pt/Curve"))
+	];
+	Grid->AddSlot(1, 1) // x, y
+	.HAlign(HAlign_Right)
+	[
+		SNew(STextBlock)
+		.Font(DetailFontInfo)
+		.Text(LOCTEXT("HairInfo_MinCPPerCurveText", "Min"))
+	];
+	Grid->AddSlot(2, 1) // x, y
+	.HAlign(HAlign_Right)
+	[
+		SNew(STextBlock)
+		.Font(DetailFontInfo)
+		.Text(LOCTEXT("HairInfo_MaxCPPerCurveText", "Max"))
+	];
+	Grid->AddSlot(3, 1) // x, y
+	.HAlign(HAlign_Right)
+	[
+		SNew(STextBlock)
+		.Font(DetailFontInfo)
+		.Text(LOCTEXT("HairInfo_AvgCPPerCurveText", "Avg"))
+	];
+
+	// Max. point per curve
+	Grid->AddSlot(0, 2) // x, y
+	.HAlign(HAlign_Left)
+	[
+		SNew(STextBlock)
+		.Font(DetailFontInfo)
+		.Text(LOCTEXT("HairInfo_CPPerCurveEmpty", ""))
+	];
+	Grid->AddSlot(1, 2) // x, y
+	.HAlign(HAlign_Right)
+	[
+		SNew(STextBlock)
+		.Font(DetailFontInfo)
+		.Text(FText::AsNumber(InHeader.MinPointPerCurve))
+	];
+	Grid->AddSlot(2, 2) // x, y
+	.HAlign(HAlign_Right)
+	[
+		SNew(STextBlock)
+		.Font(DetailFontInfo)
+		.Text(FText::AsNumber(InHeader.MaxPointPerCurve))
+	];
+	Grid->AddSlot(3, 2) // x, y
+	.HAlign(HAlign_Right)
+	[
+		SNew(STextBlock)
+		.Font(DetailFontInfo)
+		.Text(FText::AsNumber(InHeader.AvgPointPerCurve))
+	];
+
+	return Grid;
+}
+
+void AddHairStrandsCurveWarning(IDetailChildrenBuilder& ChildrenBuilder, const FSlateFontInfo& DetailFontInfo, const FHairStrandsBulkData::FHeader& InHeader)
+{
+	// Warning if group has trimmed curves or trimmed points
+	const FLinearColor ErrorColor = FLinearColor(FColor::Red);
+	if (InHeader.Flags & uint32(FHairStrandsBulkData::DataFlags_HasTrimmedCurve))
+	{
+		ChildrenBuilder.AddCustomRow(LOCTEXT("HairStrandsCurveWarning_Curve", "HairStrandsCurveWarning"))
+		.ValueContent()
+		.HAlign(HAlign_Fill)
+		[
+			SNew(STextBlock)
+			.Font(DetailFontInfo)
+			.ColorAndOpacity(ErrorColor)
+			.Text(LOCTEXT("HairInfo_TrimmedCurve", "Group has > 4M curves"))
+		];
+	}
+	if (InHeader.Flags & uint32(FHairStrandsBulkData::DataFlags_HasTrimmedPoint))
+	{
+		ChildrenBuilder.AddCustomRow(LOCTEXT("HairStrandsCurveWarning_Point", "HairStrandsPointWarning"))
+		.ValueContent()
+		.HAlign(HAlign_Fill)
+		[
+			SNew(STextBlock)
+			.Font(DetailFontInfo)
+			.ColorAndOpacity(ErrorColor)
+			.Text(LOCTEXT("HairInfo_TrimmedPoint", "Group has curve with >255 points."))
+		];
+	}
+}
+
 ///////////////////////////////////////////////////////////////////////////////////////////////////////////
 // Array panel for hair strands infos
 TSharedRef<SUniformGridPanel> MakeHairStrandsLODInfoGrid(const FSlateFontInfo& DetailFontInfo, const FHairLODInfo& LODInfo)
@@ -161,7 +260,7 @@ TSharedRef<SUniformGridPanel> MakeHairStrandsInfoGrid(const FSlateFontInfo& Deta
 	[
 		SNew(STextBlock)
 		.Font(DetailFontInfo)
-		.Text(LOCTEXT("HairInfo_Vertices", "Vertices"))
+		.Text(LOCTEXT("HairInfo_Points", "Points"))
 	];
 
 	// Strands
@@ -1424,6 +1523,13 @@ void FGroomRenderingDetails::OnGenerateElementForHairGroup(TSharedRef<IPropertyH
 		[
 			MakeHairStrandsInfoGrid(DetailFontInfo, GroomAsset->GetHairGroupsInfo()[GroupIndex], GroomAsset->GetHairGroupsPlatformData()[GroupIndex].Strands.BulkData.Header.MaxRadius)
 		];
+		ChildrenBuilder.AddCustomRow(LOCTEXT("HairStrandsCurveInfo_Array", "HairStrandsCurveInfo"))
+		.ValueContent()
+		.HAlign(HAlign_Fill)
+		[
+			MakeHairStrandsCurveInfoGrid(DetailFontInfo, GroomAsset->GetHairGroupsPlatformData()[GroupIndex].Strands.BulkData.Header)
+		];
+		AddHairStrandsCurveWarning(ChildrenBuilder, DetailFontInfo, GroomAsset->GetHairGroupsPlatformData()[GroupIndex].Strands.BulkData.Header);
 		ChildrenBuilder.AddCustomRow(LOCTEXT("HairStrandsAttributeInfo_Array", "HairStrandsAttributeInfo"))
 		.ValueContent()
 		.HAlign(HAlign_Fill)
