Remove incorrect check.

diff --git a/Engine/Plugins/Runtime/HairStrands/Source/HairStrandsCore/Private/GroomManager.cpp b/Engine/Plugins/Runtime/HairStrands/Source/HairStrandsCore/Private/GroomManager.cpp
--- a/Engine/Plugins/Runtime/HairStrands/Source/HairStrandsCore/Private/GroomManager.cpp
+++ b/Engine/Plugins/Runtime/HairStrands/Source/HairStrandsCore/Private/GroomManager.cpp
@@ -265,7 +265,6 @@ struct FHairGeometryCache
 		uint32 Index = HashTable.First(Hash);
 		while (HashTable.IsValid(Index) && Datas[Index].Key != Key)
 		{
-			check(Datas[Index].Hash == Hash); // Sanitycheck
 			Index = HashTable.Next(Index);
 		}
 
