Groom: Fixed possible crash when using groom cache in Blueprint.

diff --git a/Engine/Plugins/Runtime/HairStrands/Source/HairStrandsCore/Private/GroomCacheStreamingManager.cpp b/Engine/Plugins/Runtime/HairStrands/Source/HairStrandsCore/Private/GroomCacheStreamingManager.cpp
--- a/Engine/Plugins/Runtime/HairStrands/Source/HairStrandsCore/Private/GroomCacheStreamingManager.cpp
+++ b/Engine/Plugins/Runtime/HairStrands/Source/HairStrandsCore/Private/GroomCacheStreamingManager.cpp
@@ -227,10 +227,15 @@ void FGroomCacheStreamingManager::UnregisterGroomCache(UGroomCache* GroomCache,
 {
 	FScopeLock Lock(&CriticalSection);
 
-	TSet<UGroomComponent*>& Users = GroomCacheUsers.FindChecked(GroomCache);
-	Users.Remove(GroomComponent);
+	TSet<UGroomComponent*>* Users = GroomCacheUsers.Find(GroomCache);
+	if (!Users)
+	{
+		return;
+	}
+
+	Users->Remove(GroomComponent);
 
-	if (Users.Num() == 0)
+	if (Users->Num() == 0)
 	{
 		FGroomCacheStreamingData** StreamingData = StreamingGroomCaches.Find(GroomCache);
 		if (StreamingData)
