Fix typed buffer SRV creation issue with unknown format in groom deformer.

diff --git a/Engine/Plugins/Runtime/HairStrands/Source/HairStrandsDeformer/Private/DeformerDataInterfaceGroom.cpp b/Engine/Plugins/Runtime/HairStrands/Source/HairStrandsDeformer/Private/DeformerDataInterfaceGroom.cpp
--- a/Engine/Plugins/Runtime/HairStrands/Source/HairStrandsDeformer/Private/DeformerDataInterfaceGroom.cpp
+++ b/Engine/Plugins/Runtime/HairStrands/Source/HairStrandsDeformer/Private/DeformerDataInterfaceGroom.cpp
@@ -14,6 +14,7 @@
 #include "HairStrandsDefinitions.h"
 #include "DeformerGroomComponentSource.h"
 #include "HairStrandsInterpolation.h"
+#include "SystemTextures.h"
 
 FString UOptimusGroomDataInterface::GetDisplayName() const
 {
@@ -209,7 +210,9 @@ bool FOptimusGroomDataProviderProxy::IsValid(FValidationData const& InValidation
 void FOptimusGroomDataProviderProxy::AllocateResources(FRDGBuilder& GraphBuilder)
 {
 	Resources.Empty();
-	FallbackSRV = nullptr;
+	FallbackByteAddressSRV = nullptr;
+	FallbackStructuredSRV = nullptr;
+	FallbackVertexSRV = nullptr;
 	const uint32 InstanceCount = GroomComponent ? GroomComponent->GetGroupCount() : 0;
 	for (uint32 Index =0; Index <InstanceCount;++Index)
 	{
@@ -231,10 +234,9 @@ void FOptimusGroomDataProviderProxy::AllocateResources(FRDGBuilder& GraphBuilder
 				R.PointInterpolationBuffer = Instance->Strands.InterpolationResource ? RegisterAsSRV(GraphBuilder, Instance->Strands.InterpolationResource->PointInterpolationBuffer) : nullptr;
 			}
 
-			if (!FallbackSRV)
-			{
-				FallbackSRV = GraphBuilder.CreateSRV(GraphBuilder.RegisterExternalBuffer(GWhiteVertexBufferWithRDG->Buffer), PF_R16G16B16A16_UINT); // This should be a ByteAddressBuffer
-			}
+			if (!FallbackByteAddressSRV){ FallbackByteAddressSRV = GraphBuilder.CreateSRV(GSystemTextures.GetDefaultByteAddressBuffer(GraphBuilder, 16u)); }
+			if (!FallbackStructuredSRV) { FallbackStructuredSRV  = GraphBuilder.CreateSRV(GSystemTextures.GetDefaultStructuredBuffer(GraphBuilder, 16u)); }
+			if (!FallbackVertexSRV) 	{ FallbackVertexSRV 	 = GraphBuilder.CreateSRV(GSystemTextures.GetDefaultBuffer(GraphBuilder, 4u, 0u), PF_R32_UINT); }
 		}
 	}
 }
@@ -252,14 +254,14 @@ void FOptimusGroomDataProviderProxy::GatherDispatchData(FDispatchData const& InD
 		FParameters& Parameters = ParameterArray[InvocationIndex];
 		Parameters.Common.PointCount = 0;
 		Parameters.Common.CurveCount = 0;
-		Parameters.Resources.PositionBuffer			= FallbackSRV;
-		Parameters.Resources.PositionOffsetBuffer 	= FallbackSRV;
-		Parameters.Resources.CurveAttributeBuffer	= FallbackSRV;
-		Parameters.Resources.PointAttributeBuffer	= FallbackSRV;
-		Parameters.Resources.PointToCurveBuffer		= FallbackSRV;
-		Parameters.Resources.CurveBuffer			= FallbackSRV;
-		Parameters.Interpolation.CurveInterpolationBuffer = FallbackSRV;
-		Parameters.Interpolation.PointInterpolationBuffer = FallbackSRV;
+		Parameters.Resources.PositionBuffer			= FallbackByteAddressSRV;
+		Parameters.Resources.PositionOffsetBuffer 	= FallbackStructuredSRV;
+		Parameters.Resources.CurveAttributeBuffer	= FallbackByteAddressSRV;
+		Parameters.Resources.PointAttributeBuffer	= FallbackByteAddressSRV;
+		Parameters.Resources.PointToCurveBuffer		= FallbackVertexSRV;
+		Parameters.Resources.CurveBuffer			= FallbackVertexSRV;
+		Parameters.Interpolation.CurveInterpolationBuffer = FallbackByteAddressSRV;
+		Parameters.Interpolation.PointInterpolationBuffer = FallbackByteAddressSRV;
 
 		if (FHairGroupInstance* Instance = GroomComponent->GetGroupInstance(InvocationIndex))
 		{
diff --git a/Engine/Plugins/Runtime/HairStrands/Source/HairStrandsDeformer/Private/DeformerDataInterfaceGroom.h b/Engine/Plugins/Runtime/HairStrands/Source/HairStrandsDeformer/Private/DeformerDataInterfaceGroom.h
--- a/Engine/Plugins/Runtime/HairStrands/Source/HairStrandsDeformer/Private/DeformerDataInterfaceGroom.h
+++ b/Engine/Plugins/Runtime/HairStrands/Source/HairStrandsDeformer/Private/DeformerDataInterfaceGroom.h
@@ -70,6 +70,7 @@ private:
 	UGroomComponent* GroomComponent = nullptr; // Should it be HairInstance instead?
 	TArray<FHairStrandsInstanceResourceParameters> Resources;
 	TArray<FHairStrandsInstanceInterpolationParameters> Interpolations;
-	FRDGBufferSRVRef FallbackSRV;
-
+	FRDGBufferSRVRef FallbackByteAddressSRV;
+	FRDGBufferSRVRef FallbackStructuredSRV;
+	FRDGBufferSRVRef FallbackVertexSRV;
 };
diff --git a/Engine/Plugins/Runtime/HairStrands/Source/HairStrandsDeformer/Private/DeformerDataInterfaceGroomGuide.cpp b/Engine/Plugins/Runtime/HairStrands/Source/HairStrandsDeformer/Private/DeformerDataInterfaceGroomGuide.cpp
--- a/Engine/Plugins/Runtime/HairStrands/Source/HairStrandsDeformer/Private/DeformerDataInterfaceGroomGuide.cpp
+++ b/Engine/Plugins/Runtime/HairStrands/Source/HairStrandsDeformer/Private/DeformerDataInterfaceGroomGuide.cpp
@@ -14,6 +14,7 @@
 #include "HairStrandsDefinitions.h"
 #include "DeformerGroomComponentSource.h"
 #include "HairStrandsInterpolation.h"
+#include "SystemTextures.h"
 
 FString UOptimusGroomGuideDataInterface::GetDisplayName() const
 {
@@ -148,7 +149,9 @@ bool FOptimusGroomGuideDataProviderProxy::IsValid(FValidationData const& InValid
 void FOptimusGroomGuideDataProviderProxy::AllocateResources(FRDGBuilder& GraphBuilder)
 {
 	Resources.Empty();
-	FallbackSRV = nullptr;
+	FallbackByteAddressSRV = nullptr;
+	FallbackStructuredSRV = nullptr;
+	FallbackVertexSRV = nullptr;
 	const uint32 InstanceCount = GroomComponent ? GroomComponent->GetGroupCount() : 0;
 	for (uint32 Index =0; Index <InstanceCount;++Index)
 	{
@@ -163,10 +166,9 @@ void FOptimusGroomGuideDataProviderProxy::AllocateResources(FRDGBuilder& GraphBu
 				R.PointToCurveBuffer	= RegisterAsSRV(GraphBuilder, Instance->Guides.RestResource->PointToCurveBuffer);
 			}
 
-			if (!FallbackSRV)
-			{
-				FallbackSRV = GraphBuilder.CreateSRV(GraphBuilder.RegisterExternalBuffer(GWhiteVertexBufferWithRDG->Buffer), PF_R16G16B16A16_UINT); // This should be a byte address buffer
-			}
+			if (!FallbackByteAddressSRV){ FallbackByteAddressSRV = GraphBuilder.CreateSRV(GSystemTextures.GetDefaultByteAddressBuffer(GraphBuilder, 16u)); }
+			if (!FallbackStructuredSRV) { FallbackStructuredSRV  = GraphBuilder.CreateSRV(GSystemTextures.GetDefaultStructuredBuffer(GraphBuilder, 16u)); }
+			if (!FallbackVertexSRV) 	{ FallbackVertexSRV 	 = GraphBuilder.CreateSRV(GSystemTextures.GetDefaultBuffer(GraphBuilder, 4u, 0u), PF_R32_UINT); }
 		}
 	}
 }
@@ -184,10 +186,10 @@ void FOptimusGroomGuideDataProviderProxy::GatherDispatchData(FDispatchData const
 		FParameters& Parameters = ParameterArray[InvocationIndex];
 		Parameters.Common.PointCount 				= 0;
 		Parameters.Common.CurveCount 				= 0;
-		Parameters.Resources.PositionBuffer			= FallbackSRV;
-		Parameters.Resources.PositionOffsetBuffer 	= FallbackSRV;
-		Parameters.Resources.PointToCurveBuffer		= FallbackSRV;
-		Parameters.Resources.CurveBuffer			= FallbackSRV;
+		Parameters.Resources.PositionBuffer			= FallbackByteAddressSRV;
+		Parameters.Resources.PositionOffsetBuffer 	= FallbackStructuredSRV;
+		Parameters.Resources.PointToCurveBuffer		= FallbackVertexSRV;
+		Parameters.Resources.CurveBuffer			= FallbackVertexSRV;
 
 		if (FHairGroupInstance* Instance = GroomComponent->GetGroupInstance(InvocationIndex))
 		{
diff --git a/Engine/Plugins/Runtime/HairStrands/Source/HairStrandsDeformer/Private/DeformerDataInterfaceGroomGuide.h b/Engine/Plugins/Runtime/HairStrands/Source/HairStrandsDeformer/Private/DeformerDataInterfaceGroomGuide.h
--- a/Engine/Plugins/Runtime/HairStrands/Source/HairStrandsDeformer/Private/DeformerDataInterfaceGroomGuide.h
+++ b/Engine/Plugins/Runtime/HairStrands/Source/HairStrandsDeformer/Private/DeformerDataInterfaceGroomGuide.h
@@ -69,6 +69,7 @@ private:
 
 	UGroomComponent* GroomComponent = nullptr;
 	TArray<FHairStrandsInstanceResourceParameters> Resources;
-	FRDGBufferSRVRef FallbackSRV;
-
+	FRDGBufferSRVRef FallbackByteAddressSRV;
+	FRDGBufferSRVRef FallbackStructuredSRV;
+	FRDGBufferSRVRef FallbackVertexSRV;
 };
diff --git a/Engine/Plugins/Runtime/HairStrands/Source/HairStrandsDeformer/Private/DeformerDataInterfaceGroomWrite.cpp b/Engine/Plugins/Runtime/HairStrands/Source/HairStrandsDeformer/Private/DeformerDataInterfaceGroomWrite.cpp
--- a/Engine/Plugins/Runtime/HairStrands/Source/HairStrandsDeformer/Private/DeformerDataInterfaceGroomWrite.cpp
+++ b/Engine/Plugins/Runtime/HairStrands/Source/HairStrandsDeformer/Private/DeformerDataInterfaceGroomWrite.cpp
@@ -201,12 +201,9 @@ void FOptimusGroomWriteDataProviderProxy::AllocateResources(FRDGBuilder& GraphBu
 
 			FResources& R = Resources.AddDefaulted_GetRef();
 
-			FRDGBufferRef PositionBuffer_fallback = GraphBuilder.CreateBuffer(FRDGBufferDesc::CreateBufferDesc(8, 1), TEXT("Groom.DeformedPositionBuffer"), ERDGBufferFlags::None);
-			R.PositionBufferSRV_fallback = GraphBuilder.CreateSRV(GraphBuilder.RegisterExternalBuffer(GWhiteVertexBufferWithRDG->Buffer));
-			R.PositionBufferUAV_fallback = GraphBuilder.CreateUAV(PositionBuffer_fallback, ERDGUnorderedAccessViewFlags::SkipBarrier);
-
-			FRDGBufferRef AttributeBuffer_fallback = GraphBuilder.CreateBuffer(FRDGBufferDesc::CreateByteAddressDesc(16), TEXT("Groom.DeformedAttributeBuffer"), ERDGBufferFlags::None);
-			R.AttributeBufferUAV_fallback = GraphBuilder.CreateUAV(AttributeBuffer_fallback, PF_Unknown, ERDGUnorderedAccessViewFlags::SkipBarrier);
+			R.PositionBufferSRV_fallback = GraphBuilder.CreateSRV(GSystemTextures.GetDefaultByteAddressBuffer(GraphBuilder, 16u));
+			R.PositionBufferUAV_fallback = GraphBuilder.CreateUAV(GraphBuilder.CreateBuffer(FRDGBufferDesc::CreateByteAddressDesc(16u), TEXT("Groom.Deformer.FallbackDeformedPositionBuffer")), ERDGUnorderedAccessViewFlags::SkipBarrier);
+			R.AttributeBufferUAV_fallback = GraphBuilder.CreateUAV(GraphBuilder.CreateBuffer(FRDGBufferDesc::CreateByteAddressDesc(16u), TEXT("Groom.Deformer.FallbackDeformedAttributeBuffer")), ERDGUnorderedAccessViewFlags::SkipBarrier);
 
 			// Positions / Radius
 			if (OutputMask & 0x7)
@@ -217,11 +214,9 @@ void FOptimusGroomWriteDataProviderProxy::AllocateResources(FRDGBuilder& GraphBu
 			}
 			else
 			{
-				FRDGBufferRef DummyUAVBuffer = GraphBuilder.CreateBuffer(FRDGBufferDesc::CreateBufferDesc(8u, 1), TEXT("Hair.Deformer.DummyBuffer"));
-				
-				R.PositionOffsetBufferSRV 	= GraphBuilder.CreateSRV(GSystemTextures.GetDefaultStructuredBuffer(GraphBuilder, 16u), PF_A32B32G32R32F); //GraphBuilder.CreateSRV(GraphBuilder.RegisterExternalBuffer(GWhiteVertexBufferWithRDG->Buffer), PF_A32B32G32R32F);
-				R.PositionBufferSRV 		= GraphBuilder.CreateSRV(GraphBuilder.RegisterExternalBuffer(GWhiteVertexBufferWithRDG->Buffer));
-				R.PositionBufferUAV 		= GraphBuilder.CreateUAV(DummyUAVBuffer);
+				R.PositionOffsetBufferSRV 	= GraphBuilder.CreateSRV(GSystemTextures.GetDefaultStructuredBuffer(GraphBuilder, 16u), PF_A32B32G32R32F);
+				R.PositionBufferSRV 		= GraphBuilder.CreateSRV(GSystemTextures.GetDefaultByteAddressBuffer(GraphBuilder, 16u));
+				R.PositionBufferUAV 		= GraphBuilder.CreateUAV(GraphBuilder.CreateBuffer(FRDGBufferDesc::CreateByteAddressDesc(16u), TEXT("Hair.Deformer.DummyBuffer")));
 			}
 			
 			// Curve Attributes
