Change Curve and PointToCurve buffers from typed buffer  to byteaddress buffer for scalarizing loads.

diff --git a/Engine/Plugins/Runtime/HairStrands/Shaders/Private/DeformerDataInterfaceGroom.ush b/Engine/Plugins/Runtime/HairStrands/Shaders/Private/DeformerDataInterfaceGroom.ush
--- a/Engine/Plugins/Runtime/HairStrands/Shaders/Private/DeformerDataInterfaceGroom.ush
+++ b/Engine/Plugins/Runtime/HairStrands/Shaders/Private/DeformerDataInterfaceGroom.ush
@@ -43,7 +43,7 @@ FHairCurve ReadCurveData(uint CurveIndex)
 		return (FHairCurve)0;
 	}
 
-	return UnpackHairCurve({DataInterfaceName}_CurveBuffer[CurveIndex]);
+	return ReadHairCurve({DataInterfaceName}_CurveBuffer, CurveIndex);
 }
 
 ///////////////////////////////////////////////////////////////////////////////////////////////////
@@ -122,7 +122,7 @@ float ReadCurveNumPoint_{DataInterfaceName}(uint CurveIndex)
 uint ReadGuideIndex_{DataInterfaceName}(uint ControlPointIndex)
 {
 	// Find the rendering curve index
-	const uint CurveIndex = UnpackHairPointToCurveIndex({DataInterfaceName}_PointToCurveBuffer, ControlPointIndex);
+	const uint CurveIndex = ReadHairPointToCurveIndex({DataInterfaceName}_PointToCurveBuffer, ControlPointIndex);
 
 	// Fint the curve and point interpolation data
 	const FGuideCurveData CurveData = UnpackGuideCurveData({DataInterfaceName}_CurveInterpolationBuffer, CurveIndex, {DataInterfaceName}_bSingleGuide);
diff --git a/Engine/Plugins/Runtime/HairStrands/Shaders/Private/DeformerDataInterfaceGroomGuide.ush b/Engine/Plugins/Runtime/HairStrands/Shaders/Private/DeformerDataInterfaceGroomGuide.ush
--- a/Engine/Plugins/Runtime/HairStrands/Shaders/Private/DeformerDataInterfaceGroomGuide.ush
+++ b/Engine/Plugins/Runtime/HairStrands/Shaders/Private/DeformerDataInterfaceGroomGuide.ush
@@ -38,7 +38,7 @@ FHairCurve ReadCurveData_{DataInterfaceName}(uint CurveIndex)
 		return (FHairCurve)0;
 	}
 
-	return UnpackHairCurve({DataInterfaceName}_CurveBuffer[CurveIndex]);
+	return ReadHairCurve({DataInterfaceName}_CurveBuffer, CurveIndex);
 }
 
 ///////////////////////////////////////////////////////////////////////////////////////////////////
diff --git a/Engine/Plugins/Runtime/HairStrands/Shaders/Private/NiagaraDataInterfaceHairStrandsTemplate.ush b/Engine/Plugins/Runtime/HairStrands/Shaders/Private/NiagaraDataInterfaceHairStrandsTemplate.ush
--- a/Engine/Plugins/Runtime/HairStrands/Shaders/Private/NiagaraDataInterfaceHairStrandsTemplate.ush
+++ b/Engine/Plugins/Runtime/HairStrands/Shaders/Private/NiagaraDataInterfaceHairStrandsTemplate.ush
@@ -7,7 +7,7 @@ NiagaraDataInterfaceHairStrands.ush
 
 RWByteAddressBuffer	{ParameterName}_DeformedPositionBuffer;
 ByteAddressBuffer	{ParameterName}_RestPositionBuffer;
-Buffer<uint>		{ParameterName}_CurvesOffsetsBuffer;
+ByteAddressBuffer	{ParameterName}_CurvesOffsetsBuffer;
 Buffer<float4>		{ParameterName}_RestTrianglePositionBuffer;
 Buffer<float4>		{ParameterName}_DeformedTrianglePositionBuffer;
 Buffer<uint>		{ParameterName}_RootBarycentricCoordinatesBuffer;
@@ -115,7 +115,7 @@ void BuildBoundingBox_{ParameterName}(in float3 NodePosition, out bool FunctionS
  */
 FHairCurve UnpackCurve_{ParameterName}(in uint InCurveIndex)
 {
-	return UnpackHairCurve({ParameterName}_CurvesOffsetsBuffer[InCurveIndex]);
+	return ReadHairCurve({ParameterName}_CurvesOffsetsBuffer, InCurveIndex);
 }
 
 /* -----------------------------------------------------------------
diff --git a/Engine/Plugins/Runtime/HairStrands/Shaders/Private/NiagaraInterpolateGroomGuides.usf b/Engine/Plugins/Runtime/HairStrands/Shaders/Private/NiagaraInterpolateGroomGuides.usf
--- a/Engine/Plugins/Runtime/HairStrands/Shaders/Private/NiagaraInterpolateGroomGuides.usf
+++ b/Engine/Plugins/Runtime/HairStrands/Shaders/Private/NiagaraInterpolateGroomGuides.usf
@@ -23,10 +23,10 @@ Buffer<uint>		RootToUniqueTriangleIndexBuffer;
 
 #endif
 
-Buffer<uint>		PointToCurveIndexBuffer;
+ByteAddressBuffer	PointToCurveIndexBuffer;
 RWByteAddressBuffer	DeformedPositionBuffer;
 ByteAddressBuffer	RestPositionBuffer;
-Buffer<uint>		CurvesOffsetsBuffer;
+ByteAddressBuffer	CurvesOffsetsBuffer;
 StructuredBuffer<float4> DeformedPositionOffset;
 float3				RestPositionOffset;
 float4x4			WorldToLocal;
@@ -71,9 +71,9 @@ void MainCS(uint3 DispatchThreadId : SV_DispatchThreadID)
 {
 	if (DispatchThreadId.x < uint(NumPoints))
 	{
-		const int StrandIndex = UnpackHairPointToCurveIndex(PointToCurveIndexBuffer, DispatchThreadId.x);
+		const int StrandIndex = ReadHairPointToCurveIndex(PointToCurveIndexBuffer, DispatchThreadId.x);
 
-		const FHairCurve Curve = UnpackHairCurve(CurvesOffsetsBuffer[StrandIndex]);
+		const FHairCurve Curve = ReadHairCurve(CurvesOffsetsBuffer, StrandIndex);
 		const int PointBegin = Curve.PointIndex;
 		const int PointCount = Curve.PointCount;
 		const int PointIndex = DispatchThreadId.x - PointBegin;
diff --git a/Engine/Plugins/Runtime/HairStrands/Source/HairStrandsCore/Private/GroomManager.cpp b/Engine/Plugins/Runtime/HairStrands/Source/HairStrandsCore/Private/GroomManager.cpp
--- a/Engine/Plugins/Runtime/HairStrands/Source/HairStrandsCore/Private/GroomManager.cpp
+++ b/Engine/Plugins/Runtime/HairStrands/Source/HairStrandsCore/Private/GroomManager.cpp
@@ -1403,7 +1403,6 @@ static void RunHairStrandsInterpolation_Strands(
 					bHasSkinning && bValidGuide ? InstanceData.Instance->Guides.DeformedRootResource : nullptr,
 					RegisterAsSRV(GraphBuilder, InstanceData.Instance->Strands.RestResource->PositionBuffer),
 					RegisterAsSRV(GraphBuilder, InstanceData.Instance->Strands.RestResource->CurveBuffer),
-					bValidGuide ? RegisterAsSRV(GraphBuilder, InstanceData.Instance->Guides.RestResource->CurveBuffer) : nullptr,
 					bUseSingleGuide,
 					bValidGuide ? RegisterAsSRV(GraphBuilder, InstanceData.Instance->Strands.InterpolationResource->CurveInterpolationBuffer) : nullptr,
 					bValidGuide ? RegisterAsSRV(GraphBuilder, InstanceData.Instance->Strands.InterpolationResource->PointInterpolationBuffer) : nullptr,
@@ -1961,7 +1960,6 @@ static void RunHairStrandsInterpolation_Cards(
 					InstanceData.bHasSkinning && InstanceData.bValidGuide ? InstanceData.Instance->Guides.DeformedRootResource : nullptr,
 					RegisterAsSRV(GraphBuilder, LOD.Guides.RestResource->PositionBuffer),
 					RegisterAsSRV(GraphBuilder, LOD.Guides.RestResource->CurveBuffer),
-					InstanceData.bValidGuide ? RegisterAsSRV(GraphBuilder, InstanceData.Instance->Guides.RestResource->CurveBuffer) : nullptr,
 					bUseSingleGuide,
 					InstanceData.bValidGuide ? RegisterAsSRV(GraphBuilder, LOD.Guides.InterpolationResource->CurveInterpolationBuffer) : nullptr,
 					InstanceData.bValidGuide ? RegisterAsSRV(GraphBuilder, LOD.Guides.InterpolationResource->PointInterpolationBuffer) : nullptr,
diff --git a/Engine/Plugins/Runtime/HairStrands/Source/HairStrandsCore/Private/GroomResources.cpp b/Engine/Plugins/Runtime/HairStrands/Source/HairStrandsCore/Private/GroomResources.cpp
--- a/Engine/Plugins/Runtime/HairStrands/Source/HairStrandsCore/Private/GroomResources.cpp
+++ b/Engine/Plugins/Runtime/HairStrands/Source/HairStrandsCore/Private/GroomResources.cpp
@@ -1237,10 +1237,9 @@ void FHairStrandsRestResource::InternalAllocate(FRDGBuilder& GraphBuilder)
 		InternalCreateByteAddressBufferRDG(GraphBuilder, DummyAttribute, EPixelFormat::PF_R32_UINT, PointAttributeBuffer, ToHairResourceDebugName(HAIRSTRANDS_RESOUCE_NAME(CurveType, Hair.StrandsRest_PointAttributeBuffer), ResourceName), OwnerName, EHairResourceUsageType::Static);
 		GraphBuilder.UseExternalAccessMode(Register(GraphBuilder, PointAttributeBuffer, ERDGImportedBufferFlags::CreateSRV).Buffer, ERHIAccess::SRVMask);
 	}
-	InternalCreateVertexBufferRDG_FromHairBulkData<FHairStrandsCurveFormat>(GraphBuilder, BulkData.Data.Curves, CurveCount, CurveBuffer, ToHairResourceDebugName(HAIRSTRANDS_RESOUCE_NAME(CurveType, Hair.StrandsRest_CurveBuffer), ResourceName), OwnerName, EHairResourceUsageType::Static);
+	InternalCreateByteAddressBufferRDG_FromHairBulkData(GraphBuilder, BulkData.Data.Curves, CurveBuffer, ToHairResourceDebugName(HAIRSTRANDS_RESOUCE_NAME(CurveType, Hair.StrandsRest_CurveBuffer), ResourceName), OwnerName, EHairResourceUsageType::Static);
 
-	const uint32 PointToCurveIndexCount = FMath::DivideAndRoundUp(PointCount, BulkData.Header.Strides.PointToCurveChunkElementCount);
-	InternalCreateVertexBufferRDG_FromHairBulkData<FHairStrandsPointToCurveFormat>(GraphBuilder, BulkData.Data.PointToCurve, PointToCurveIndexCount, PointToCurveBuffer, ToHairResourceDebugName(HAIRSTRANDS_RESOUCE_NAME(CurveType, Hair.StrandsRest_PointToCurveBuffer), ResourceName), OwnerName, EHairResourceUsageType::Static);
+	InternalCreateByteAddressBufferRDG_FromHairBulkData(GraphBuilder, BulkData.Data.PointToCurve, PointToCurveBuffer, ToHairResourceDebugName(HAIRSTRANDS_RESOUCE_NAME(CurveType, Hair.StrandsRest_PointToCurveBuffer), ResourceName), OwnerName, EHairResourceUsageType::Static);
 
 	// If the buffer is not null, it means it was been lazily allocated once. In such a case, The tangent buffer is recreated the new required size
 	const bool bStaticTangentNeeded = TangentBuffer.Buffer != nullptr;
@@ -1497,7 +1496,7 @@ void FHairStrandsCullingResource::InternalAllocate(FRDGBuilder& GraphBuilder, ui
 		const uint32 PointCount = FMath::Min(InPointCount, MaxPointCount);
 		const uint32 CurveCount = FMath::Min(InCurveCount, MaxCurveCount);
 
-		InternalCreateVertexBufferRDG<FHairStrandsUintFormat>(GraphBuilder, CurveCount, Resources.CulledCurveBuffer, TEXT("Hair.Cluster_CulledCurveBuffer"), GetOwnerName(), EHairResourceUsageType::Dynamic);
+		InternalCreateByteAddressBufferRDG(GraphBuilder, CurveCount * FHairStrandsUintFormat::SizeInByte, Resources.CulledCurveBuffer, TEXT("Hair.Cluster_CulledCurveBuffer"), GetOwnerName(), EHairResourceUsageType::Dynamic);
 		InternalCreateVertexBufferRDG<FHairStrandsUintFormat>(GraphBuilder, PointCount, Resources.CulledVertexIdBuffer, TEXT("Hair.Cluster_CulledVertexIdBuffer"), GetOwnerName(), EHairResourceUsageType::Dynamic);
 		InternalCreateVertexBufferRDG<FHairStrandsFloatFormat>(GraphBuilder, PointCount, Resources.CulledVertexRadiusScaleBuffer, TEXT("Hair.Cluster_CulledVertexRadiusScaleBuffer"), GetOwnerName(), EHairResourceUsageType::Dynamic);
 		AddClearUAVFloatPass(GraphBuilder, RegisterAsUAV(GraphBuilder, Resources.CulledVertexRadiusScaleBuffer), 0.f);
diff --git a/Engine/Plugins/Runtime/HairStrands/Source/HairStrandsCore/Private/HairStrandsClusterCulling.cpp b/Engine/Plugins/Runtime/HairStrands/Source/HairStrandsCore/Private/HairStrandsClusterCulling.cpp
--- a/Engine/Plugins/Runtime/HairStrands/Source/HairStrandsCore/Private/HairStrandsClusterCulling.cpp
+++ b/Engine/Plugins/Runtime/HairStrands/Source/HairStrandsCore/Private/HairStrandsClusterCulling.cpp
@@ -29,7 +29,7 @@ class FHairClusterCullCS: public FGlobalShader
 		SHADER_PARAMETER(uint32, DispatchCountX)
 		SHADER_PARAMETER(FVector4f, ClusterInfoParameters)
 
-		SHADER_PARAMETER_RDG_BUFFER_SRV(Buffer<uint>, CurveBuffer)
+		SHADER_PARAMETER_RDG_BUFFER_SRV(ByteAddressBuffer, CurveBuffer)
 		SHADER_PARAMETER_RDG_BUFFER_SRV(Buffer<uint>, CurveToClusterIdBuffer)
 		SHADER_PARAMETER_RDG_BUFFER_SRV(Buffer<uint>, PointLODBuffer)
 		SHADER_PARAMETER_RDG_BUFFER_SRV(StructuredBuffer<FPackedHairClusterInfo>, ClusterInfoBuffer)
@@ -37,7 +37,7 @@ class FHairClusterCullCS: public FGlobalShader
 		SHADER_PARAMETER_RDG_BUFFER_UAV(RWStructuredBuffer<uint>, OutPointCounter)
 		SHADER_PARAMETER_RDG_BUFFER_UAV(RWBuffer<uint>, OutIndexBuffer)
 		SHADER_PARAMETER_RDG_BUFFER_UAV(RWBuffer<float>, OutRadiusScaleBuffer)
-		SHADER_PARAMETER_RDG_BUFFER_UAV(RWBuffer<uint>, OutCulledCurveBuffer)
+		SHADER_PARAMETER_RDG_BUFFER_UAV(RWByteAddressBuffer, OutCulledCurveBuffer)
 
 		SHADER_PARAMETER_STRUCT_REF(FViewUniformShaderParameters, ViewUniformBuffer)
 		SHADER_PARAMETER_STRUCT_INCLUDE(ShaderPrint::FShaderParameters, ShaderPrintUniformBuffer)
diff --git a/Engine/Plugins/Runtime/HairStrands/Source/HairStrandsCore/Private/HairStrandsInterpolation.cpp b/Engine/Plugins/Runtime/HairStrands/Source/HairStrandsCore/Private/HairStrandsInterpolation.cpp
--- a/Engine/Plugins/Runtime/HairStrands/Source/HairStrandsCore/Private/HairStrandsInterpolation.cpp
+++ b/Engine/Plugins/Runtime/HairStrands/Source/HairStrandsCore/Private/HairStrandsInterpolation.cpp
@@ -302,12 +302,12 @@ class FDeformGuideCS : public FGlobalShader
 		SHADER_PARAMETER_RDG_BUFFER_SRV(Buffer<float4>, SimDeformedPositionBuffer)
 
 		SHADER_PARAMETER_RDG_BUFFER_SRV(Buffer<uint>, SimRootBarycentricBuffer)
-		SHADER_PARAMETER_RDG_BUFFER_SRV(Buffer<uint>, SimPointToCurveBuffer)
+		SHADER_PARAMETER_RDG_BUFFER_SRV(ByteAddressBuffer, SimPointToCurveBuffer)
 		SHADER_PARAMETER_RDG_BUFFER_SRV(Buffer<uint>, SimRootToUniqueTriangleIndexBuffer)
 
 		SHADER_PARAMETER_RDG_BUFFER_SRV(StructuredBuffer, SimDeformedOffsetBuffer)
-		SHADER_PARAMETER_RDG_BUFFER_SRV(Buffer, SimRestPosePositionBuffer)
-		SHADER_PARAMETER_RDG_BUFFER_UAV(RWBuffer, OutSimDeformedPositionBuffer)
+		SHADER_PARAMETER_RDG_BUFFER_SRV(ByteAddressBuffer, SimRestPosePositionBuffer)
+		SHADER_PARAMETER_RDG_BUFFER_UAV(RWByteAddressBuffer, OutSimDeformedPositionBuffer)
 
 		SHADER_PARAMETER(uint32, SampleCount)
 		SHADER_PARAMETER_RDG_BUFFER_SRV(StructuredBuffer, RestSamplePositionsBuffer)
@@ -467,8 +467,7 @@ class FHairInterpolationCS : public FGlobalShader
 		SHADER_PARAMETER(FMatrix44f, LocalToWorldMatrix)
 		SHADER_PARAMETER(uint32,  DispatchCountX)
 	
-		SHADER_PARAMETER_RDG_BUFFER_SRV(Buffer, RenCurveBuffer)
-		SHADER_PARAMETER_RDG_BUFFER_SRV(Buffer, SimCurveBuffer)
+		SHADER_PARAMETER_RDG_BUFFER_SRV(ByteAddressBuffer, RenCurveBuffer)
 
 		SHADER_PARAMETER_RDG_BUFFER_SRV(ByteAddressBuffer, RenRestPosePositionBuffer)
 		SHADER_PARAMETER_RDG_BUFFER_SRV(ByteAddressBuffer, RenDeformerPositionBuffer)
@@ -548,7 +547,6 @@ void AddHairStrandsInterpolationPass(
 	const FHairStrandsDeformedRootResource* SimDeformedRootResources,
 	const FRDGBufferSRVRef& RenRestPosePositionBuffer,
 	const FRDGBufferSRVRef& RenCurveBuffer,
-	const FRDGBufferSRVRef& SimCurveBuffer,
 	const bool bUseSingleGuide,
 	const FRDGBufferSRVRef& CurveInterpolationBuffer,
 	const FRDGBufferSRVRef& PointInterpolationBuffer,
@@ -578,7 +576,6 @@ void AddHairStrandsInterpolationPass(
 	Parameters->HairLengthScale = HairLengthScale;
 
 	Parameters->RenCurveBuffer = RenCurveBuffer;
-	Parameters->SimCurveBuffer = SimCurveBuffer;
 
 	if (ShaderPrintData)
 	{
@@ -683,7 +680,7 @@ class FHairPatchAttributeCS : public FGlobalShader
 		SHADER_PARAMETER(uint32, CurveCount)
 		SHADER_PARAMETER(uint32, bUseSingleGuide)
 		SHADER_PARAMETER_STRUCT_INCLUDE(FHairStrandsInstanceAttributeParameters, RenAttributes)
-		SHADER_PARAMETER_RDG_BUFFER_SRV(Buffer, RenCurveBuffer)
+		SHADER_PARAMETER_RDG_BUFFER_SRV(ByteAddressBuffer, RenCurveBuffer)
 		SHADER_PARAMETER_RDG_BUFFER_SRV(ByteAddressBuffer, CurveInterpolationBuffer)
 		SHADER_PARAMETER_RDG_BUFFER_SRV(ByteAddressBuffer, PointInterpolationBuffer)
 		SHADER_PARAMETER_RDG_BUFFER_SRV(Buffer, RenCurveToClusterIdBuffer)
@@ -777,9 +774,9 @@ class FHairClusterAABBCS : public FGlobalShader
 		SHADER_PARAMETER(FVector3f, CPUBoundMin)
 		SHADER_PARAMETER(FVector3f, CPUBoundMax)
 		SHADER_PARAMETER(FMatrix44f, LocalToTranslatedWorldMatrix)
-		SHADER_PARAMETER_RDG_BUFFER_SRV(Buffer<uint>, RenCurveBuffer)
+		SHADER_PARAMETER_RDG_BUFFER_SRV(ByteAddressBuffer, RenCurveBuffer)
 		SHADER_PARAMETER_RDG_BUFFER_SRV(Buffer<uint>, RenPointLODBuffer)
-		SHADER_PARAMETER_RDG_BUFFER_SRV(Buffer, RenderDeformedPositionBuffer)
+		SHADER_PARAMETER_RDG_BUFFER_SRV(ByteAddressBuffer, RenderDeformedPositionBuffer)
 		SHADER_PARAMETER_RDG_BUFFER_SRV(StructuredBuffer, RenderDeformedOffsetBuffer)
 		SHADER_PARAMETER_RDG_BUFFER_UAV(RWBuffer, OutClusterAABBBuffer)
 		SHADER_PARAMETER_RDG_BUFFER_UAV(RWBuffer, OutGroupAABBBuffer)
@@ -896,7 +893,7 @@ class FHairCardsDeformationCS : public FGlobalShader
 
 		SHADER_PARAMETER_RDG_BUFFER_SRV(Buffer<uint>, GuideRootToUniqueTriangleIndexBuffer)
 		SHADER_PARAMETER_RDG_BUFFER_SRV(Buffer<uint>, GuideRootBarycentricBuffer)
-		SHADER_PARAMETER_RDG_BUFFER_SRV(Buffer<uint>, GuideVertexToRootIndexBuffer)
+		SHADER_PARAMETER_RDG_BUFFER_SRV(ByteAddressBuffer, GuideVertexToRootIndexBuffer)
 
 		SHADER_PARAMETER_RDG_BUFFER_UAV(RWBuffer, CardsDeformedPositionBuffer)
 		SHADER_PARAMETER_RDG_BUFFER_UAV(RWBuffer, CardsDeformedTangentBuffer)
diff --git a/Engine/Plugins/Runtime/HairStrands/Source/HairStrandsCore/Private/HairStrandsVertexFactory.cpp b/Engine/Plugins/Runtime/HairStrands/Source/HairStrandsCore/Private/HairStrandsVertexFactory.cpp
--- a/Engine/Plugins/Runtime/HairStrands/Source/HairStrandsCore/Private/HairStrandsVertexFactory.cpp
+++ b/Engine/Plugins/Runtime/HairStrands/Source/HairStrandsCore/Private/HairStrandsVertexFactory.cpp
@@ -107,15 +107,15 @@ FHairStrandsVertexFactoryUniformShaderParameters FHairGroupInstance::GetHairStan
 	Out.PrevResources.PreviousPositionOffsetBuffer	= VFInput.Strands.PrevPositionOffsetBufferRHISRV;
 
 	// swap in some default data for those buffers that are not valid yet
-	if (!Out.Resources.PositionBuffer) 						{ Out.Resources.PositionBuffer = GDummyCulledDispatchVertexIdsBuffer.SRVFloat; }
+	if (!Out.Resources.PositionBuffer) 						{ Out.Resources.PositionBuffer = GDummyCulledDispatchVertexIdsBuffer.SRVByteAddress; }
 	if (!Out.Resources.PositionOffsetBuffer) 				{ Out.Resources.PositionOffsetBuffer = GDummyCulledDispatchVertexIdsBuffer.SRVFloat; }
 	if (!Out.Resources.CurveAttributeBuffer) 				{ Out.Resources.CurveAttributeBuffer = GDummyCulledDispatchVertexIdsBuffer.SRVByteAddress; }
 	if (!Out.Resources.PointAttributeBuffer) 				{ Out.Resources.PointAttributeBuffer = GDummyCulledDispatchVertexIdsBuffer.SRVByteAddress; }
 	if (!Out.Resources.CurveBuffer) 						{ Out.Resources.CurveBuffer = GDummyCulledDispatchVertexIdsBuffer.SRVByteAddress; }
-	if (!Out.Resources.PointToCurveBuffer) 					{ Out.Resources.PointToCurveBuffer = GDummyCulledDispatchVertexIdsBuffer.SRVUint; }
+	if (!Out.Resources.PointToCurveBuffer) 					{ Out.Resources.PointToCurveBuffer = GDummyCulledDispatchVertexIdsBuffer.SRVByteAddress; }
 	if (!Out.Resources.TangentBuffer) 						{ Out.Resources.TangentBuffer = GDummyCulledDispatchVertexIdsBuffer.SRVFloat; }
 
-	if (!Out.PrevResources.PreviousPositionBuffer) 			{ Out.PrevResources.PreviousPositionBuffer = GDummyCulledDispatchVertexIdsBuffer.SRVFloat; }
+	if (!Out.PrevResources.PreviousPositionBuffer) 			{ Out.PrevResources.PreviousPositionBuffer = GDummyCulledDispatchVertexIdsBuffer.SRVByteAddress; }
 	if (!Out.PrevResources.PreviousPositionOffsetBuffer) 	{ Out.PrevResources.PreviousPositionOffsetBuffer = GDummyCulledDispatchVertexIdsBuffer.SRVFloat; }
 
 	Out.Culling.bCullingEnable = HairGroupPublicData->GetCullingResultAvailable();
diff --git a/Engine/Plugins/Runtime/HairStrands/Source/HairStrandsCore/Private/Niagara/NiagaraDataInterfaceHairStrands.cpp b/Engine/Plugins/Runtime/HairStrands/Source/HairStrandsCore/Private/Niagara/NiagaraDataInterfaceHairStrands.cpp
--- a/Engine/Plugins/Runtime/HairStrands/Source/HairStrandsCore/Private/Niagara/NiagaraDataInterfaceHairStrands.cpp
+++ b/Engine/Plugins/Runtime/HairStrands/Source/HairStrandsCore/Private/Niagara/NiagaraDataInterfaceHairStrands.cpp
@@ -824,7 +824,7 @@ class FInterpolateGroomGuidesCS : public FGlobalShader
 		SHADER_PARAMETER_RDG_BUFFER_SRV(Buffer<uint>,	RootBarycentricCoordinatesBuffer)
 		SHADER_PARAMETER_RDG_BUFFER_SRV(Buffer<uint>,	RootToUniqueTriangleIndexBuffer)
 
-		SHADER_PARAMETER_RDG_BUFFER_SRV(Buffer<uint>,		PointToCurveIndexBuffer)
+		SHADER_PARAMETER_RDG_BUFFER_SRV(ByteAddressBuffer,	PointToCurveIndexBuffer)
 		SHADER_PARAMETER_RDG_BUFFER_UAV(RWByteAddressBuffer,DeformedPositionBuffer)
 		SHADER_PARAMETER_RDG_BUFFER_SRV(ByteAddressBuffer,	RestPositionBuffer)
 		SHADER_PARAMETER_RDG_BUFFER_SRV(Buffer<uint>,		CurvesOffsetsBuffer)
@@ -3525,7 +3525,7 @@ void UNiagaraDataInterfaceHairStrands::SetShaderParameters(const FNiagaraDataInt
 		FRDGBufferSRVRef DummyStructuredBuffer = GraphBuilder.CreateSRV(GSystemTextures.GetDefaultStructuredBuffer(GraphBuilder, 16u, FUintVector4(0,0,0,0)));
 		FRDGBufferSRVRef DummyByteAddressBuffer = GraphBuilder.CreateSRV(GSystemTextures.GetDefaultByteAddressBuffer(GraphBuilder, 16u));
 
-		ShaderParameters->CurvesOffsetsBuffer = Context.GetComputeDispatchInterface().GetEmptyBufferSRV(GraphBuilder, FHairStrandsCurveFormat::Format);
+		ShaderParameters->CurvesOffsetsBuffer = DummyByteAddressBuffer;
 		ShaderParameters->RestPositionBuffer = DummyByteAddressBuffer;
 		ShaderParameters->DeformedPositionOffset = DummyStructuredBuffer;
 		ShaderParameters->RestTrianglePositionBuffer = Context.GetComputeDispatchInterface().GetEmptyBufferSRV(GraphBuilder, FHairStrandsMeshTrianglePositionFormat::Format);
diff --git a/Engine/Plugins/Runtime/HairStrands/Source/HairStrandsCore/Public/HairStrandsInterpolation.h b/Engine/Plugins/Runtime/HairStrands/Source/HairStrandsCore/Public/HairStrandsInterpolation.h
--- a/Engine/Plugins/Runtime/HairStrands/Source/HairStrandsCore/Public/HairStrandsInterpolation.h
+++ b/Engine/Plugins/Runtime/HairStrands/Source/HairStrandsCore/Public/HairStrandsInterpolation.h
@@ -138,7 +138,6 @@ void AddHairStrandsInterpolationPass(
 	const FHairStrandsDeformedRootResource* SimDeformedRootResources,
 	const FRDGBufferSRVRef& RenRestPosePositionBuffer,
 	const FRDGBufferSRVRef& RenCurveBuffer,
-	const FRDGBufferSRVRef& SimCurveBuffer,
 	const bool bUseSingleGuide,
 	const FRDGBufferSRVRef& CurveInterpolationBuffer,
 	const FRDGBufferSRVRef& PoinInterpolationBuffer,
diff --git a/Engine/Plugins/Runtime/HairStrands/Source/HairStrandsDeformer/Private/DeformerDataInterfaceGroom.cpp b/Engine/Plugins/Runtime/HairStrands/Source/HairStrandsDeformer/Private/DeformerDataInterfaceGroom.cpp
--- a/Engine/Plugins/Runtime/HairStrands/Source/HairStrandsDeformer/Private/DeformerDataInterfaceGroom.cpp
+++ b/Engine/Plugins/Runtime/HairStrands/Source/HairStrandsDeformer/Private/DeformerDataInterfaceGroom.cpp
@@ -212,7 +212,6 @@ void FOptimusGroomDataProviderProxy::AllocateResources(FRDGBuilder& GraphBuilder
 	Resources.Empty();
 	FallbackByteAddressSRV = nullptr;
 	FallbackStructuredSRV = nullptr;
-	FallbackVertexSRV = nullptr;
 	const uint32 InstanceCount = GroomComponent ? GroomComponent->GetGroupCount() : 0;
 	for (uint32 Index =0; Index <InstanceCount;++Index)
 	{
@@ -236,7 +235,6 @@ void FOptimusGroomDataProviderProxy::AllocateResources(FRDGBuilder& GraphBuilder
 
 			if (!FallbackByteAddressSRV){ FallbackByteAddressSRV = GraphBuilder.CreateSRV(GSystemTextures.GetDefaultByteAddressBuffer(GraphBuilder, 16u)); }
 			if (!FallbackStructuredSRV) { FallbackStructuredSRV  = GraphBuilder.CreateSRV(GSystemTextures.GetDefaultStructuredBuffer(GraphBuilder, 16u)); }
-			if (!FallbackVertexSRV) 	{ FallbackVertexSRV 	 = GraphBuilder.CreateSRV(GSystemTextures.GetDefaultBuffer(GraphBuilder, 4u, 0u), PF_R32_UINT); }
 		}
 	}
 }
@@ -258,8 +256,8 @@ void FOptimusGroomDataProviderProxy::GatherDispatchData(FDispatchData const& InD
 		Parameters.Resources.PositionOffsetBuffer 	= FallbackStructuredSRV;
 		Parameters.Resources.CurveAttributeBuffer	= FallbackByteAddressSRV;
 		Parameters.Resources.PointAttributeBuffer	= FallbackByteAddressSRV;
-		Parameters.Resources.PointToCurveBuffer		= FallbackVertexSRV;
-		Parameters.Resources.CurveBuffer			= FallbackVertexSRV;
+		Parameters.Resources.PointToCurveBuffer		= FallbackByteAddressSRV;
+		Parameters.Resources.CurveBuffer			= FallbackByteAddressSRV;
 		Parameters.Interpolation.CurveInterpolationBuffer = FallbackByteAddressSRV;
 		Parameters.Interpolation.PointInterpolationBuffer = FallbackByteAddressSRV;
 
diff --git a/Engine/Plugins/Runtime/HairStrands/Source/HairStrandsDeformer/Private/DeformerDataInterfaceGroom.h b/Engine/Plugins/Runtime/HairStrands/Source/HairStrandsDeformer/Private/DeformerDataInterfaceGroom.h
--- a/Engine/Plugins/Runtime/HairStrands/Source/HairStrandsDeformer/Private/DeformerDataInterfaceGroom.h
+++ b/Engine/Plugins/Runtime/HairStrands/Source/HairStrandsDeformer/Private/DeformerDataInterfaceGroom.h
@@ -72,5 +72,4 @@ private:
 	TArray<FHairStrandsInstanceInterpolationParameters> Interpolations;
 	FRDGBufferSRVRef FallbackByteAddressSRV;
 	FRDGBufferSRVRef FallbackStructuredSRV;
-	FRDGBufferSRVRef FallbackVertexSRV;
 };
diff --git a/Engine/Plugins/Runtime/HairStrands/Source/HairStrandsDeformer/Private/DeformerDataInterfaceGroomGuide.cpp b/Engine/Plugins/Runtime/HairStrands/Source/HairStrandsDeformer/Private/DeformerDataInterfaceGroomGuide.cpp
--- a/Engine/Plugins/Runtime/HairStrands/Source/HairStrandsDeformer/Private/DeformerDataInterfaceGroomGuide.cpp
+++ b/Engine/Plugins/Runtime/HairStrands/Source/HairStrandsDeformer/Private/DeformerDataInterfaceGroomGuide.cpp
@@ -151,7 +151,6 @@ void FOptimusGroomGuideDataProviderProxy::AllocateResources(FRDGBuilder& GraphBu
 	Resources.Empty();
 	FallbackByteAddressSRV = nullptr;
 	FallbackStructuredSRV = nullptr;
-	FallbackVertexSRV = nullptr;
 	const uint32 InstanceCount = GroomComponent ? GroomComponent->GetGroupCount() : 0;
 	for (uint32 Index =0; Index <InstanceCount;++Index)
 	{
@@ -168,7 +167,6 @@ void FOptimusGroomGuideDataProviderProxy::AllocateResources(FRDGBuilder& GraphBu
 
 			if (!FallbackByteAddressSRV){ FallbackByteAddressSRV = GraphBuilder.CreateSRV(GSystemTextures.GetDefaultByteAddressBuffer(GraphBuilder, 16u)); }
 			if (!FallbackStructuredSRV) { FallbackStructuredSRV  = GraphBuilder.CreateSRV(GSystemTextures.GetDefaultStructuredBuffer(GraphBuilder, 16u)); }
-			if (!FallbackVertexSRV) 	{ FallbackVertexSRV 	 = GraphBuilder.CreateSRV(GSystemTextures.GetDefaultBuffer(GraphBuilder, 4u, 0u), PF_R32_UINT); }
 		}
 	}
 }
@@ -188,8 +186,8 @@ void FOptimusGroomGuideDataProviderProxy::GatherDispatchData(FDispatchData const
 		Parameters.Common.CurveCount 				= 0;
 		Parameters.Resources.PositionBuffer			= FallbackByteAddressSRV;
 		Parameters.Resources.PositionOffsetBuffer 	= FallbackStructuredSRV;
-		Parameters.Resources.PointToCurveBuffer		= FallbackVertexSRV;
-		Parameters.Resources.CurveBuffer			= FallbackVertexSRV;
+		Parameters.Resources.PointToCurveBuffer		= FallbackByteAddressSRV;
+		Parameters.Resources.CurveBuffer			= FallbackByteAddressSRV;
 
 		if (FHairGroupInstance* Instance = GroomComponent->GetGroupInstance(InvocationIndex))
 		{
diff --git a/Engine/Plugins/Runtime/HairStrands/Source/HairStrandsDeformer/Private/DeformerDataInterfaceGroomGuide.h b/Engine/Plugins/Runtime/HairStrands/Source/HairStrandsDeformer/Private/DeformerDataInterfaceGroomGuide.h
--- a/Engine/Plugins/Runtime/HairStrands/Source/HairStrandsDeformer/Private/DeformerDataInterfaceGroomGuide.h
+++ b/Engine/Plugins/Runtime/HairStrands/Source/HairStrandsDeformer/Private/DeformerDataInterfaceGroomGuide.h
@@ -71,5 +71,4 @@ private:
 	TArray<FHairStrandsInstanceResourceParameters> Resources;
 	FRDGBufferSRVRef FallbackByteAddressSRV;
 	FRDGBufferSRVRef FallbackStructuredSRV;
-	FRDGBufferSRVRef FallbackVertexSRV;
 };
diff --git a/Engine/Shaders/Private/HairStrands/HairCardsDeformation.usf b/Engine/Shaders/Private/HairStrands/HairCardsDeformation.usf
--- a/Engine/Shaders/Private/HairStrands/HairCardsDeformation.usf
+++ b/Engine/Shaders/Private/HairStrands/HairCardsDeformation.usf
@@ -36,9 +36,9 @@ Buffer<float4>  TriangleRestPositionBuffer;
 
 Buffer<float4>  TriangleDeformedPositionBuffer;
 
-Buffer<uint>	GuideRootBarycentricBuffer;
-Buffer<uint>    GuideVertexToRootIndexBuffer;
-Buffer<uint>    GuideRootToUniqueTriangleIndexBuffer;
+Buffer<uint>		GuideRootBarycentricBuffer;
+ByteAddressBuffer   GuideVertexToRootIndexBuffer;
+Buffer<uint>    	GuideRootToUniqueTriangleIndexBuffer;
 
 #if PERMUTATION_DYNAMIC_GEOMETRY == 1
 float3 ComputeDynamicGeometryOffset(
@@ -144,7 +144,7 @@ void MainCS(uint DispatchThreadId : SV_DispatchThreadID)
 	const uint GuidePointIndex0 = GuideData.VertexIndex;
 	const uint GuidePointIndex1 = min(GuidePointIndex0 + 1, GuideVertexCount - 1);
 
-	const uint GuideCurveIndex = UnpackHairPointToCurveIndex(GuideVertexToRootIndexBuffer, GuidePointIndex0);
+	const uint GuideCurveIndex = ReadHairPointToCurveIndex(GuideVertexToRootIndexBuffer, GuidePointIndex0);
 	const uint TriangleIndex = GuideRootToUniqueTriangleIndexBuffer[GuideCurveIndex];
 
 	const float3 RootBarycentric = UnpackBarycentrics(GuideRootBarycentricBuffer[GuideCurveIndex]);
diff --git a/Engine/Shaders/Private/HairStrands/HairStrandsAttributeCommon.ush b/Engine/Shaders/Private/HairStrands/HairStrandsAttributeCommon.ush
--- a/Engine/Shaders/Private/HairStrands/HairStrandsAttributeCommon.ush
+++ b/Engine/Shaders/Private/HairStrands/HairStrandsAttributeCommon.ush
@@ -64,15 +64,15 @@ uint2 InternalGetHairStrandsAttributes64(uint Index, uint AttributeOffsetInBytes
 	return OutPacked;
 }
 
-uint InternalGetHairStrandsAttributesPerCurve(uint HairControlPointId, uint AttributeOffsetInBytes, uint AttributeStrideInBytes, ByteAddressBuffer InAttributeBuffer, Buffer<uint> InPointToCurveBuffer, uint InIndexToChunkDivAsShift, uint InChunkElementCount, uint InChunkStrideInBytes)
+uint InternalGetHairStrandsAttributesPerCurve(uint HairControlPointId, uint AttributeOffsetInBytes, uint AttributeStrideInBytes, ByteAddressBuffer InAttributeBuffer, ByteAddressBuffer InPointToCurveBuffer, uint InIndexToChunkDivAsShift, uint InChunkElementCount, uint InChunkStrideInBytes)
 {
-	const uint CurveIndex = UnpackHairPointToCurveIndex(InPointToCurveBuffer, HairControlPointId);
+	const uint CurveIndex = ReadHairPointToCurveIndex(InPointToCurveBuffer, HairControlPointId);
 	return InternalGetHairStrandsAttributes(CurveIndex, AttributeOffsetInBytes, AttributeStrideInBytes, InAttributeBuffer, InIndexToChunkDivAsShift, InChunkElementCount, InChunkStrideInBytes);
 }
 
-uint2 InternalGetHairStrandsAttributesPerCurve64(uint HairControlPointId, uint AttributeOffsetInBytes, uint AttributeStrideInBytes, ByteAddressBuffer InAttributeBuffer, Buffer<uint> InPointToCurveBuffer, uint InIndexToChunkDivAsShift, uint InChunkElementCount, uint InChunkStrideInBytes)
+uint2 InternalGetHairStrandsAttributesPerCurve64(uint HairControlPointId, uint AttributeOffsetInBytes, uint AttributeStrideInBytes, ByteAddressBuffer InAttributeBuffer, ByteAddressBuffer InPointToCurveBuffer, uint InIndexToChunkDivAsShift, uint InChunkElementCount, uint InChunkStrideInBytes)
 {
-	const uint CurveIndex = UnpackHairPointToCurveIndex(InPointToCurveBuffer, HairControlPointId);
+	const uint CurveIndex = ReadHairPointToCurveIndex(InPointToCurveBuffer, HairControlPointId);
 	return InternalGetHairStrandsAttributes64(CurveIndex, AttributeOffsetInBytes, AttributeStrideInBytes, InAttributeBuffer, InIndexToChunkDivAsShift, InChunkElementCount, InChunkStrideInBytes);
 }
 
diff --git a/Engine/Shaders/Private/HairStrands/HairStrandsClusterCulling.usf b/Engine/Shaders/Private/HairStrands/HairStrandsClusterCulling.usf
--- a/Engine/Shaders/Private/HairStrands/HairStrandsClusterCulling.usf
+++ b/Engine/Shaders/Private/HairStrands/HairStrandsClusterCulling.usf
@@ -72,7 +72,7 @@ uint CurvePerGroup;
 uint DispatchCountX;
 float4 ClusterInfoParameters;
 
-Buffer<uint> CurveBuffer;
+ByteAddressBuffer CurveBuffer;
 Buffer<uint> CurveToClusterIdBuffer;
 Buffer<uint> PointLODBuffer;
 StructuredBuffer<FPackedHairClusterInfo> ClusterInfoBuffer;
@@ -80,7 +80,7 @@ StructuredBuffer<FPackedHairClusterInfo> ClusterInfoBuffer;
 RWStructuredBuffer<uint>  OutPointCounter;
 RWBuffer<uint>  OutIndexBuffer;
 RWBuffer<float> OutRadiusScaleBuffer;
-RWBuffer<uint>  OutCulledCurveBuffer;
+RWByteAddressBuffer  OutCulledCurveBuffer;
 
 #ifndef GROUP_SIZE
 #error GROUP_SIZE needs to be Undefined
@@ -152,7 +152,7 @@ void Main(uint3 DispatchThreadId : SV_DispatchThreadID, uint ThreadIndex : SV_Gr
 		CurvePointCount[LocalCurveIndex] = 0;
 		if (bValidCurveIndex)
 		{
-			const FHairCurve Curve = UnpackHairCurve(CurveBuffer[GlobalCurveIndex]);
+			const FHairCurve Curve = ReadHairCurve(CurveBuffer, GlobalCurveIndex);
 			CurvePointIndex[LocalCurveIndex] = Curve.PointIndex;
 			CurvePointCount[LocalCurveIndex] = Curve.PointCount;
 		}
@@ -241,7 +241,7 @@ void Main(uint3 DispatchThreadId : SV_DispatchThreadID, uint ThreadIndex : SV_Gr
 			CulledCurve.PointIndex = WriteOffset0[LocalCurveIndex];
 			CulledCurve.PointCount =(WriteOffsetN[LocalCurveIndex] - WriteOffset0[LocalCurveIndex]) + 1;
 
-			OutCulledCurveBuffer[GlobalCurveIndex] = PackHairCurve(CulledCurve);
+			WriteHairCurve(OutCulledCurveBuffer, GlobalCurveIndex, CulledCurve);
 		}
 	}
 }
@@ -297,7 +297,7 @@ void Main(uint3 DispatchThreadId : SV_DispatchThreadID, uint ThreadIndex : SV_Gr
 		CurvePointCount = 0;
 		if (bIsCurveValid)
 		{
-			const FHairCurve Curve = UnpackHairCurve(CurveBuffer[CurveIndex]);
+			const FHairCurve Curve = ReadHairCurve(CurveBuffer, CurveIndex);
 			CurvePointIndex = Curve.PointIndex;
 			CurvePointCount = Curve.PointCount;
 		}
@@ -387,7 +387,7 @@ void Main(uint3 DispatchThreadId : SV_DispatchThreadID, uint ThreadIndex : SV_Gr
 		FHairCurve CulledCurve;
 		CulledCurve.PointCount = AllocationCount;
 		CulledCurve.PointIndex = AllocationOffset;
-		OutCulledCurveBuffer[CurveIndex] = PackHairCurve(CulledCurve);
+		WriteHairCurve(OutCulledCurveBuffer, CurveIndex, CulledCurve);
 	}
 }
 #endif // PERMUTATION_POINT_PER_CURVE == 64
@@ -446,7 +446,7 @@ float3 CPUBoundMax;
 float LODIndex;
 float4x4 LocalToTranslatedWorldMatrix;
 
-Buffer<uint> RenCurveBuffer;
+ByteAddressBuffer RenCurveBuffer;
 Buffer<uint> RenPointLODBuffer;
 ByteAddressBuffer RenderDeformedPositionBuffer;
 StructuredBuffer<float4> RenderDeformedOffsetBuffer;
@@ -485,7 +485,7 @@ groupshared float3 OutHairPositionOffset;
 		const uint CurveIndex = GroupId.x;
 		if (CurveIndex < CurveCount)
 		{
-			const FHairCurve RenCurve = UnpackHairCurve(RenCurveBuffer[CurveIndex]);
+			const FHairCurve RenCurve = ReadHairCurve(RenCurveBuffer, CurveIndex);
 			CurvePointIndex = RenCurve.PointIndex;
 			CurvePointCount = RenCurve.PointCount;
 		}
diff --git a/Engine/Shaders/Private/HairStrands/HairStrandsForwardCulling.usf b/Engine/Shaders/Private/HairStrands/HairStrandsForwardCulling.usf
--- a/Engine/Shaders/Private/HairStrands/HairStrandsForwardCulling.usf
+++ b/Engine/Shaders/Private/HairStrands/HairStrandsForwardCulling.usf
@@ -107,7 +107,7 @@ void CSMain(uint DispatchThreadID : SV_DispatchThreadID, uint GroupThreadID : SV
 	FHairCurve Curve = (FHairCurve)0;
 	if (bValidCurve)
 	{
-		Curve = UnpackHairCurve(HairStrandsVF_CurveBuffer[CurveID]);
+		Curve = ReadHairCurve(HairStrandsVF_CurveBuffer, CurveID);
 		
 		uint VisiblePointCount = 0;
 		for (uint PointIt=0; PointIt < Curve.PointCount-1; ++PointIt)
diff --git a/Engine/Shaders/Private/HairStrands/HairStrandsGuideDeform.usf b/Engine/Shaders/Private/HairStrands/HairStrandsGuideDeform.usf
--- a/Engine/Shaders/Private/HairStrands/HairStrandsGuideDeform.usf
+++ b/Engine/Shaders/Private/HairStrands/HairStrandsGuideDeform.usf
@@ -26,9 +26,9 @@ RWByteAddressBuffer OutSimDeformedPositionBuffer;
 Buffer<float4>  SimRestPositionBuffer;
 Buffer<float4>  SimDeformedPositionBuffer;
 
-Buffer<uint>	SimRootBarycentricBuffer;
-Buffer<uint>    SimPointToCurveBuffer;
-Buffer<uint>	SimRootToUniqueTriangleIndexBuffer;
+Buffer<uint>	  SimRootBarycentricBuffer;
+ByteAddressBuffer SimPointToCurveBuffer;
+Buffer<uint>	  SimRootToUniqueTriangleIndexBuffer;
 
 // Return world position of a guide vertex deformed by a skinned triangle
 float3 DeformGuideByTriangle(
@@ -73,7 +73,7 @@ float3 DisplacePosition(uint VertexIndex, float3 Pos)
 #endif
 
 #if PERMUTATION_DEFORMATION == DEFORMATION_OFFSET_GUIDE_DYNAMIC
-	const uint RootIndex = UnpackHairPointToCurveIndex(SimPointToCurveBuffer, VertexIndex);
+	const uint RootIndex = ReadHairPointToCurveIndex(SimPointToCurveBuffer, VertexIndex);
 	const float3 RootBarycentric = UnpackBarycentrics(SimRootBarycentricBuffer[RootIndex]);
 
 	const uint TriangleIndex = SimRootToUniqueTriangleIndexBuffer[RootIndex];
diff --git a/Engine/Shaders/Private/HairStrands/HairStrandsInterpolation.usf b/Engine/Shaders/Private/HairStrands/HairStrandsInterpolation.usf
--- a/Engine/Shaders/Private/HairStrands/HairStrandsInterpolation.usf
+++ b/Engine/Shaders/Private/HairStrands/HairStrandsInterpolation.usf
@@ -397,8 +397,7 @@ float3 ApplyDeformer(const FRootTriangleData RenTri, float3 InRestPostion, float
 
 uint CurveCount;
 
-Buffer<uint> SimCurveBuffer;
-Buffer<uint> RenCurveBuffer;
+ByteAddressBuffer RenCurveBuffer;
 
 #define TOTALMAXPOINTPERCURVE 32
 //-------------------------------------------------------------------------------------------------
@@ -444,7 +443,7 @@ void MainCS(uint3 DispatchThreadId : SV_DispatchThreadID, uint ThreadIndex : SV_
 		if (bValidCurveIndex)
 		{
 			// Curve's point count/offset
-			const FHairCurve Curve = UnpackHairCurve(RenCurveBuffer[GlobalCurveIndex]);
+			const FHairCurve Curve = ReadHairCurve(RenCurveBuffer, GlobalCurveIndex);
 			CurvePointIndex[LocalCurveIndex] = Curve.PointIndex;
 			CurvePointCount[LocalCurveIndex] = Curve.PointCount;	
 
@@ -608,7 +607,7 @@ void MainCS(uint3 DispatchThreadId : SV_DispatchThreadID, uint ThreadIndex : SV_
 		CurvePointCount = 0;
 		if (bIsCurveValid)
 		{
-			const FHairCurve Curve = UnpackHairCurve(RenCurveBuffer[CurveIndex]);
+			const FHairCurve Curve = ReadHairCurve(RenCurveBuffer, CurveIndex);
 			CurvePointIndex = Curve.PointIndex;
 			CurvePointCount = Curve.PointCount;
 
@@ -737,7 +736,7 @@ void MainCS(uint3 DispatchThreadId : SV_DispatchThreadID, uint ThreadIndex : SV_
 uint				CurveCount;
 uint				bUseSingleGuide;
 
-Buffer<uint>		RenCurveBuffer;
+ByteAddressBuffer	RenCurveBuffer;
 Buffer<uint>		RenCurveToClusterIdBuffer;
 RWByteAddressBuffer	OutRenAttributeBuffer;
 
@@ -783,7 +782,7 @@ void MainCS(uint2 DispatchThreadId : SV_DispatchThreadID, uint GroupThread1D : S
 		#else // if (PERMUTATION_SIMULATION > 0)
 		{
 
-			const uint VertexIndex0 = UnpackHairCurve(RenCurveBuffer[CurveIndex]).PointIndex;
+			const uint VertexIndex0 = ReadHairCurve(RenCurveBuffer, CurveIndex).PointIndex;
 
 			const FGuideCurveData SimCurveData = UnpackGuideCurveData(CurveInterpolationBuffer, CurveIndex, bUseSingleGuide);
 			const FGuidePointData SimPointData = UnpackGuidePointData(PointInterpolationBuffer, VertexIndex0, bUseSingleGuide); 
diff --git a/Engine/Shaders/Private/HairStrands/HairStrandsPack.ush b/Engine/Shaders/Private/HairStrands/HairStrandsPack.ush
--- a/Engine/Shaders/Private/HairStrands/HairStrandsPack.ush
+++ b/Engine/Shaders/Private/HairStrands/HairStrandsPack.ush
@@ -361,9 +361,9 @@ uint PackHairCurve(FHairCurve In)
 
 FHairCurve UnpackHairCurve(uint In)
 {
-	#if HAIR_MAX_NUM_POINT_PER_CURVE != 0xFF || HAIR_MAX_NUM_POINT_PER_GROUP != 0xFFFFFF
-	#error Update hair curve ENCODING_TYPE
-	#endif
+#if HAIR_MAX_NUM_POINT_PER_CURVE != 0xFF || HAIR_MAX_NUM_POINT_PER_GROUP != 0xFFFFFF
+#error Update hair curve ENCODING_TYPE
+#endif
 
 	FHairCurve Out = (FHairCurve)0;
 	Out.PointIndex = In & HAIR_MAX_NUM_POINT_PER_GROUP;
@@ -371,13 +371,23 @@ FHairCurve UnpackHairCurve(uint In)
 	return Out;
 }
 
-uint UnpackHairPointToCurveIndex(Buffer<uint> InBuffer, uint InPointIndex)
+FHairCurve ReadHairCurve(ByteAddressBuffer InBuffer, uint InIndex)
+{
+	return UnpackHairCurve(InBuffer.Load(InIndex * 4u));
+}
+
+void WriteHairCurve(RWByteAddressBuffer InBuffer, uint InIndex, FHairCurve In)
+{
+	InBuffer.Store(InIndex * 4u, PackHairCurve(In));
+}
+
+uint ReadHairPointToCurveIndex(ByteAddressBuffer InBuffer, uint InPointIndex)
 {
-	// Encoding: base curve index + detla encoding for 8 consecutive control points
+	// Encoding: base curve index + delta encoding for 8 consecutive control points
 
 	const uint PointIndex8 = InPointIndex >> 3u;
 	const uint LocalIndex8 = InPointIndex & 0x7;
-	const uint Data = InBuffer[PointIndex8];
+	const uint Data = InBuffer.Load(PointIndex8 * 4u);
 
 	// Extract the curve index delta encoding
 	const uint Mask     = BitFieldMaskU32(LocalIndex8+1, 0);
diff --git a/Engine/Shaders/Private/HairStrands/HairStrandsVertexFactoryCommon.ush b/Engine/Shaders/Private/HairStrands/HairStrandsVertexFactoryCommon.ush
--- a/Engine/Shaders/Private/HairStrands/HairStrandsVertexFactoryCommon.ush
+++ b/Engine/Shaders/Private/HairStrands/HairStrandsVertexFactoryCommon.ush
@@ -80,8 +80,8 @@ float3 ReadCardPositionOffset(StructuredBuffer<float4> In, uint InIndex, uint In
 	uint						Prefix##_PointAttributeChunkStrideInBytes; \
 	ByteAddressBuffer			Prefix##_PositionBuffer; \
 	StructuredBuffer<float4>	Prefix##_PositionOffsetBuffer; \
-	Buffer<uint>				Prefix##_CurveBuffer; \
-	Buffer<uint>				Prefix##_PointToCurveBuffer; \
+	ByteAddressBuffer			Prefix##_CurveBuffer; \
+	ByteAddressBuffer			Prefix##_PointToCurveBuffer; \
 	ByteAddressBuffer 			Prefix##_CurveAttributeBuffer; \
 	ByteAddressBuffer 			Prefix##_PointAttributeBuffer; \
 	Buffer<uint>				Prefix##_CullingIndirectBuffer; \
diff --git a/Engine/Source/Runtime/Renderer/Public/HairStrandsInterface.h b/Engine/Source/Runtime/Renderer/Public/HairStrandsInterface.h
--- a/Engine/Source/Runtime/Renderer/Public/HairStrandsInterface.h
+++ b/Engine/Source/Runtime/Renderer/Public/HairStrandsInterface.h
@@ -65,17 +65,17 @@ END_SHADER_PARAMETER_STRUCT()
 
 // Instance resources (RDG)
 BEGIN_SHADER_PARAMETER_STRUCT(FHairStrandsInstanceResourceParameters, RENDERER_API)
-	SHADER_PARAMETER_RDG_BUFFER_SRV(Buffer, PositionBuffer)
+	SHADER_PARAMETER_RDG_BUFFER_SRV(ByteAddressBuffer, PositionBuffer)
 	SHADER_PARAMETER_RDG_BUFFER_SRV(StructuredBuffer, PositionOffsetBuffer)
-	SHADER_PARAMETER_RDG_BUFFER_SRV(Buffer, CurveBuffer)
-	SHADER_PARAMETER_RDG_BUFFER_SRV(Buffer, PointToCurveBuffer)
+	SHADER_PARAMETER_RDG_BUFFER_SRV(ByteAddressBuffer, CurveBuffer)
+	SHADER_PARAMETER_RDG_BUFFER_SRV(ByteAddressBuffer, PointToCurveBuffer)
 	SHADER_PARAMETER_RDG_BUFFER_SRV(ByteAddressBuffer, CurveAttributeBuffer)
 	SHADER_PARAMETER_RDG_BUFFER_SRV(ByteAddressBuffer, PointAttributeBuffer)
 END_SHADER_PARAMETER_STRUCT()
 
 // Instance prev. resources (RDG)
 BEGIN_SHADER_PARAMETER_STRUCT(FHairStrandsInstancePrevResourceParameters, RENDERER_API)
-	SHADER_PARAMETER_RDG_BUFFER_SRV(Buffer, PrevPositionBuffer)
+	SHADER_PARAMETER_RDG_BUFFER_SRV(ByteAddressBuffer, PrevPositionBuffer)
 	SHADER_PARAMETER_RDG_BUFFER_SRV(StructuredBuffer, PrevPositionOffsetBuffer)
 END_SHADER_PARAMETER_STRUCT()
 
@@ -98,8 +98,8 @@ END_SHADER_PARAMETER_STRUCT()
 BEGIN_SHADER_PARAMETER_STRUCT(FHairStrandsInstanceResourceRawParameters, RENDERER_API)
 	SHADER_PARAMETER_SRV(ByteAddressBuffer, PositionBuffer)
 	SHADER_PARAMETER_SRV(StructuredBuffer<float4>, PositionOffsetBuffer)
-	SHADER_PARAMETER_SRV(Buffer<uint>, CurveBuffer)
-	SHADER_PARAMETER_SRV(Buffer<uint>, PointToCurveBuffer)
+	SHADER_PARAMETER_SRV(ByteAddressBuffer, CurveBuffer)
+	SHADER_PARAMETER_SRV(ByteAddressBuffer, PointToCurveBuffer)
 	SHADER_PARAMETER_SRV(Buffer<float4>, TangentBuffer)
 	SHADER_PARAMETER_SRV(ByteAddressBuffer, CurveAttributeBuffer)
 	SHADER_PARAMETER_SRV(ByteAddressBuffer, PointAttributeBuffer)
