Modified EHairViewRayTracingMask to be a enum class and use ENUM_CLASS_FLAGS
- necessary to prevent conflicts with namespace RayTracing used elsewhere in codebase

diff --git a/Engine/Plugins/Runtime/HairStrands/Source/HairStrandsCore/Private/GroomComponent.cpp b/Engine/Plugins/Runtime/HairStrands/Source/HairStrandsCore/Private/GroomComponent.cpp
--- a/Engine/Plugins/Runtime/HairStrands/Source/HairStrandsCore/Private/GroomComponent.cpp
+++ b/Engine/Plugins/Runtime/HairStrands/Source/HairStrandsCore/Private/GroomComponent.cpp
@@ -643,7 +643,7 @@ public:
 		}
 
 		const bool bWireframe = AllowDebugViewmodes() && Context.ReferenceViewFamily.EngineShowFlags.Wireframe;
-		const uint32 ViewRayTracingMask = Context.ReferenceViewFamily.EngineShowFlags.PathTracing ? EHairViewRayTracingMask::PathTracing : EHairViewRayTracingMask::RayTracing;
+		const EHairViewRayTracingMask ViewRayTracingMask = Context.ReferenceViewFamily.EngineShowFlags.PathTracing ? EHairViewRayTracingMask::PathTracing : EHairViewRayTracingMask::RayTracing;
 		if (bWireframe)
 			return;
 
@@ -659,7 +659,7 @@ public:
 
 			FHairStrandsRaytracingResource* RTGeometry = nullptr;
 			bool bIsHairStrands = false;
-			uint32 InstanceViewRayTracingMask = EHairViewRayTracingMask::PathTracing | EHairViewRayTracingMask::RayTracing;
+			EHairViewRayTracingMask InstanceViewRayTracingMask = EHairViewRayTracingMask::PathTracing | EHairViewRayTracingMask::RayTracing;
 			switch (GeometryType)
 			{
 				case EHairGeometryType::Strands:
@@ -682,7 +682,7 @@ public:
 			}
 
 			// If the view and the instance raytracing mask don't match skip this instance.
-			if ((InstanceViewRayTracingMask & ViewRayTracingMask) == 0)
+			if (!EnumHasAnyFlags(InstanceViewRayTracingMask, ViewRayTracingMask))
 			{
 				continue;
 			}
diff --git a/Engine/Plugins/Runtime/HairStrands/Source/HairStrandsCore/Private/GroomManager.cpp b/Engine/Plugins/Runtime/HairStrands/Source/HairStrandsCore/Private/GroomManager.cpp
--- a/Engine/Plugins/Runtime/HairStrands/Source/HairStrandsCore/Private/GroomManager.cpp
+++ b/Engine/Plugins/Runtime/HairStrands/Source/HairStrandsCore/Private/GroomManager.cpp
@@ -1109,7 +1109,7 @@ static void RunHairStrandsInterpolation_Strands(
 	};
 
 	// Gather all strands instances
-	const uint32 ViewRayTracingMask = RHI_RAYTRACING ? (View->Family->EngineShowFlags.PathTracing ? EHairViewRayTracingMask::PathTracing : EHairViewRayTracingMask::RayTracing) : 0u;
+	const EHairViewRayTracingMask ViewRayTracingMask = RHI_RAYTRACING ? (View->Family->EngineShowFlags.PathTracing ? EHairViewRayTracingMask::PathTracing : EHairViewRayTracingMask::RayTracing) : EHairViewRayTracingMask::None;
 	bool bHasAnySimCacheInstances = false;
 	bool bHasAnyRenCacheInstances = false;
 	bool bHasAnyRaytracingInstances = false;
@@ -1171,7 +1171,7 @@ static void RunHairStrandsInterpolation_Strands(
 		}
 
 		#if RHI_RAYTRACING
-		InstanceData.bNeedRaytracing = InstanceData.Instance->Strands.RenRaytracingResource && (InstanceData.Instance->Strands.ViewRayTracingMask & ViewRayTracingMask) != 0;
+		InstanceData.bNeedRaytracing = InstanceData.Instance->Strands.RenRaytracingResource && EnumHasAnyFlags(InstanceData.Instance->Strands.ViewRayTracingMask, ViewRayTracingMask);
 		if (InstanceData.bNeedRaytracing)
 		{
 			InstanceData.RDGResources.Raytracing_PositionBuffer = Register(GraphBuilder, InstanceData.Instance->Strands.RenRaytracingResource->PositionBuffer, ERDGImportedBufferFlags::CreateViews);
diff --git a/Engine/Plugins/Runtime/HairStrands/Source/HairStrandsCore/Private/HairStrandsInterpolation.cpp b/Engine/Plugins/Runtime/HairStrands/Source/HairStrandsCore/Private/HairStrandsInterpolation.cpp
--- a/Engine/Plugins/Runtime/HairStrands/Source/HairStrandsCore/Private/HairStrandsInterpolation.cpp
+++ b/Engine/Plugins/Runtime/HairStrands/Source/HairStrandsCore/Private/HairStrandsInterpolation.cpp
@@ -1704,7 +1704,7 @@ FHairGroupPublicData::FVertexFactoryInput InternalComputeHairStrandsVertexInputD
 	// Flag bUseRaytracingGeometry only if RT geometry has been allocated for RayTracing view (not for PathTracing view). 
 	// This flag is used later for selecting if voxelization needs to flags voxel has shadow caster or if shadow casting is handled 
 	// by the RT geometry.
-	bRaytracingGeometry = Instance->Strands.RenRaytracingResource != nullptr && (Instance->Strands.ViewRayTracingMask & EHairViewRayTracingMask::RayTracing) != 0;
+	bRaytracingGeometry = Instance->Strands.RenRaytracingResource != nullptr && EnumHasAnyFlags(Instance->Strands.ViewRayTracingMask, EHairViewRayTracingMask::RayTracing);
 #endif
 
 	OutVFInput.Strands.Common.RegisteredIndex = Instance->RegisteredIndex;
diff --git a/Engine/Plugins/Runtime/HairStrands/Source/HairStrandsCore/Public/GroomInstance.h b/Engine/Plugins/Runtime/HairStrands/Source/HairStrandsCore/Public/GroomInstance.h
--- a/Engine/Plugins/Runtime/HairStrands/Source/HairStrandsCore/Public/GroomInstance.h
+++ b/Engine/Plugins/Runtime/HairStrands/Source/HairStrandsCore/Public/GroomInstance.h
@@ -67,11 +67,13 @@ enum class EHairLODSelectionType
 	Forced		// Forced LOD value
 };
 
-enum EHairViewRayTracingMask
+enum class EHairViewRayTracingMask
 {
+	None		= 0x0,
 	RayTracing  = 0x1, // Visible for raytracing effects (RT shadow, RT refleciton, Lumen, ...)
 	PathTracing = 0x2, // Visible for pathtracing rendering
 };
+ENUM_CLASS_FLAGS(EHairViewRayTracingMask);
 
 // Represent/Describe data & resources of a hair group belonging to a groom
 struct HAIRSTRANDSCORE_API FHairGroupInstance : public FHairStrandsInstance
@@ -134,7 +136,7 @@ struct HAIRSTRANDSCORE_API FHairGroupInstance : public FHairStrandsInstance
 		#if RHI_RAYTRACING
 		FHairStrandsRaytracingResource* RenRaytracingResource = nullptr;
 		bool RenRaytracingResourceOwned = false;
-		uint32 ViewRayTracingMask = 0u;
+		EHairViewRayTracingMask ViewRayTracingMask = EHairViewRayTracingMask::None;
 		float CachedHairScaledRadius = 0;
 		float CachedHairRootScale = 0;
 		float CachedHairTipScale = 0;
