Fix crash when visualization hair instance.

diff --git a/Engine/Plugins/Runtime/HairStrands/Source/HairStrandsCore/Private/GroomDebug.cpp b/Engine/Plugins/Runtime/HairStrands/Source/HairStrandsCore/Private/GroomDebug.cpp
--- a/Engine/Plugins/Runtime/HairStrands/Source/HairStrandsCore/Private/GroomDebug.cpp
+++ b/Engine/Plugins/Runtime/HairStrands/Source/HairStrandsCore/Private/GroomDebug.cpp
@@ -823,7 +823,7 @@ static void AddHairDebugPrintInstancePass(
 				D.Data0.Z = Instance->Strands.GetData().GetNumCurves(); // Change this later on for having dynamic value
 				D.Data0.W = Instance->Strands.GetData().GetNumPoints(); // Change this later on for having dynamic value
 				const int32 MeshLODIndex = Instance->HairGroupPublicData->MeshLODIndex;
-				if (MeshLODIndex>=0 && Instance->Strands.RestRootResource->IsValid(MeshLODIndex))
+				if (MeshLODIndex>=0 && Instance->Strands.RestRootResource && Instance->Strands.RestRootResource->IsValid(MeshLODIndex))
 				{
 					D.Data1.X = Instance->Strands.RestRootResource->GetLOD(MeshLODIndex)->BulkData.Header.UniqueSectionIndices.Num();
 					D.Data1.Y = Instance->Strands.RestRootResource->GetLOD(MeshLODIndex)->BulkData.Header.UniqueTriangleCount;
