Add hair cards/meshes binding & deformation support for mobile renderer.

diff --git a/Engine/Config/Android/AndroidEngine.ini b/Engine/Config/Android/AndroidEngine.ini
--- a/Engine/Config/Android/AndroidEngine.ini
+++ b/Engine/Config/Android/AndroidEngine.ini
@@ -84,7 +84,7 @@ r.Forward.LightGridSizeZ = 8
 r.Forward.MaxCulledLightsPerCell = 8
 r.HairStrands.Simulation=0
 r.HairStrands.Strands=0
-r.HairStrands.Binding=0
+r.HairStrands.Binding=1
 ; R11G11B10 UAV is not supported on Android
 r.TemporalAA.R11G11B10History=0
 ; Do not backtrace vulkan memory allocations
diff --git a/Engine/Config/IOS/IOSEngine.ini b/Engine/Config/IOS/IOSEngine.ini
--- a/Engine/Config/IOS/IOSEngine.ini
+++ b/Engine/Config/IOS/IOSEngine.ini
@@ -43,7 +43,7 @@ DefaultPlatformService=IOSLocalNotification
 Slate.CacheRenderData=0
 r.HairStrands.Simulation=0
 r.HairStrands.Strands=0
-r.HairStrands.Binding=0
+r.HairStrands.Binding=1
 
 [TextureStreaming]
 ; PoolSizeVRAMPercentage is how much percentage of GPU Dedicated VRAM should be used as a TexturePool cache for streaming textures (0 - unlimited streaming)
diff --git a/Engine/Plugins/Runtime/HairStrands/Source/HairStrandsCore/Private/GroomComponent.cpp b/Engine/Plugins/Runtime/HairStrands/Source/HairStrandsCore/Private/GroomComponent.cpp
--- a/Engine/Plugins/Runtime/HairStrands/Source/HairStrandsCore/Private/GroomComponent.cpp
+++ b/Engine/Plugins/Runtime/HairStrands/Source/HairStrandsCore/Private/GroomComponent.cpp
@@ -1067,7 +1067,6 @@ public:
 		}
 
 		FPrimitiveViewRelevance Result;
-		Result.bHairStrands = IsShown(View) || bForceDrawRelevance;
 
 		// Special pass for hair strands geometry (not part of the base pass, and shadowing is handlded in a custom fashion). When cards rendering is enabled we reusethe base pass
 		Result.bDrawRelevance		= IsShown(View) || bForceDrawRelevance;
@@ -1085,6 +1084,9 @@ public:
 		}
 		#endif
 		MaterialRelevance.SetPrimitiveViewRelevance(Result);
+
+		// Override the MaterialRelevance output
+		Result.bHairStrands = IsShown(View) || bForceDrawRelevance;
 		return Result;
 	}
 
diff --git a/Engine/Plugins/Runtime/HairStrands/Source/HairStrandsCore/Private/HairCardsVertexFactory.cpp b/Engine/Plugins/Runtime/HairStrands/Source/HairStrandsCore/Private/HairCardsVertexFactory.cpp
--- a/Engine/Plugins/Runtime/HairStrands/Source/HairStrandsCore/Private/HairCardsVertexFactory.cpp
+++ b/Engine/Plugins/Runtime/HairStrands/Source/HairStrandsCore/Private/HairCardsVertexFactory.cpp
@@ -343,18 +343,43 @@ void FHairCardsVertexFactory::InitResources(FRHICommandListBase& RHICmdList)
 	AddPrimitiveIdStreamElement(EVertexInputStreamType::Default, Elements, HAIR_CARDS_VF_PRIMITIVEID_STREAM_INDEX /*AttributeIndex*/, HAIR_CARDS_VF_PRIMITIVEID_STREAM_INDEX /*AttributeIndex_Mobile*/);
 
 	// Note this is a local version of the VF's bSupportsManualVertexFetch, which take into account the feature level
+	// When manual fetch is not supported, buffers are bound through input assembly based on vertex declaration. 
+	// A vertex declaraction only access FVertexBuffer buffers, so we create wrappers of pooled buffers
 	const bool bManualFetch = SupportsManualVertexFetch(CurrentFeatureLevel);
 	if (!bManualFetch)
 	{
 		if (Data.GeometryType == EHairGeometryType::Cards)
 		{
 			const FHairGroupInstance::FCards::FLOD& LOD = Data.Instance->Cards.LODs[Data.LODIndex];
-		
-			Elements.Add(AccessStreamComponent(FVertexStreamComponent(&LOD.RestResource->RestPositionBuffer,	0, 0,									FHairCardsPositionFormat::SizeInByte,												FHairCardsPositionFormat::VertexElementType,	EVertexStreamUsage::Default), 0));
-			Elements.Add(AccessStreamComponent(FVertexStreamComponent(&LOD.RestResource->NormalsBuffer,			0, 0,									FHairCardsNormalFormat::SizeInByte * FHairCardsNormalFormat::ComponentCount,		FHairCardsNormalFormat::VertexElementType,		EVertexStreamUsage::Default), 1));
-			Elements.Add(AccessStreamComponent(FVertexStreamComponent(&LOD.RestResource->NormalsBuffer,			0, FHairCardsNormalFormat::SizeInByte,	FHairCardsNormalFormat::SizeInByte * FHairCardsNormalFormat::ComponentCount,		FHairCardsNormalFormat::VertexElementType,		EVertexStreamUsage::Default), 2));
-			Elements.Add(AccessStreamComponent(FVertexStreamComponent(&LOD.RestResource->UVsBuffer,				0, 0,									FHairCardsUVFormat::SizeInByte,														FHairCardsUVFormat::VertexElementType,			EVertexStreamUsage::Default), 3));
-			Elements.Add(AccessStreamComponent(FVertexStreamComponent(&LOD.RestResource->MaterialsBuffer,		0, 0,									FHairCardsMaterialFormat::SizeInByte,												FHairCardsMaterialFormat::VertexElementType,	EVertexStreamUsage::Default), 4));
+			
+			const bool bDynamic = LOD.DeformedResource != nullptr;
+			if (bDynamic)
+			{
+				DeformedPositionVertexBuffer[0] = FRDGWrapperVertexBuffer(LOD.DeformedResource->GetBuffer(FHairCardsDeformedResource::EFrameType::Current));
+				DeformedPositionVertexBuffer[0].InitResource(RHICmdList);
+
+				DeformedPositionVertexBuffer[1] = FRDGWrapperVertexBuffer(LOD.DeformedResource->GetBuffer(FHairCardsDeformedResource::EFrameType::Previous));
+				DeformedPositionVertexBuffer[1].InitResource(RHICmdList);
+
+				DeformedNormalVertexBuffer = FRDGWrapperVertexBuffer(LOD.DeformedResource->DeformedNormalBuffer);
+				DeformedNormalVertexBuffer.InitResource(RHICmdList);
+
+				Elements.Add(AccessStreamComponent(FVertexStreamComponent(&(DeformedPositionVertexBuffer[0]),		0, 0,									FHairCardsPositionFormat::SizeInByte,												FHairCardsPositionFormat::VertexElementType,	EVertexStreamUsage::Default), 0));
+				Elements.Add(AccessStreamComponent(FVertexStreamComponent(&(DeformedPositionVertexBuffer[1]),		0, 0,									FHairCardsPositionFormat::SizeInByte,												FHairCardsPositionFormat::VertexElementType,	EVertexStreamUsage::Default), 5));
+
+				Elements.Add(AccessStreamComponent(FVertexStreamComponent(&DeformedNormalVertexBuffer,				0, 0,									FHairCardsNormalFormat::SizeInByte * FHairCardsNormalFormat::ComponentCount,		FHairCardsNormalFormat::VertexElementType,		EVertexStreamUsage::Default), 1));
+				Elements.Add(AccessStreamComponent(FVertexStreamComponent(&DeformedNormalVertexBuffer,				0, FHairCardsNormalFormat::SizeInByte,	FHairCardsNormalFormat::SizeInByte * FHairCardsNormalFormat::ComponentCount,		FHairCardsNormalFormat::VertexElementType,		EVertexStreamUsage::Default), 2));
+			}
+			else
+			{
+				Elements.Add(AccessStreamComponent(FVertexStreamComponent(&LOD.RestResource->RestPositionBuffer,	0, 0,									FHairCardsPositionFormat::SizeInByte,												FHairCardsPositionFormat::VertexElementType,	EVertexStreamUsage::Default), 0));
+				Elements.Add(AccessStreamComponent(FVertexStreamComponent(&LOD.RestResource->RestPositionBuffer,	0, 0,									FHairCardsPositionFormat::SizeInByte,												FHairCardsPositionFormat::VertexElementType,	EVertexStreamUsage::Default), 5));
+
+				Elements.Add(AccessStreamComponent(FVertexStreamComponent(&LOD.RestResource->NormalsBuffer,			0, 0,									FHairCardsNormalFormat::SizeInByte * FHairCardsNormalFormat::ComponentCount,		FHairCardsNormalFormat::VertexElementType,		EVertexStreamUsage::Default), 1));
+				Elements.Add(AccessStreamComponent(FVertexStreamComponent(&LOD.RestResource->NormalsBuffer,			0, FHairCardsNormalFormat::SizeInByte,	FHairCardsNormalFormat::SizeInByte * FHairCardsNormalFormat::ComponentCount,		FHairCardsNormalFormat::VertexElementType,		EVertexStreamUsage::Default), 2));
+			}
+			Elements.Add(AccessStreamComponent(FVertexStreamComponent(&LOD.RestResource->UVsBuffer,					0, 0,									FHairCardsUVFormat::SizeInByte,														FHairCardsUVFormat::VertexElementType,			EVertexStreamUsage::Default), 3));
+			Elements.Add(AccessStreamComponent(FVertexStreamComponent(&LOD.RestResource->MaterialsBuffer,			0, 0,									FHairCardsMaterialFormat::SizeInByte,												FHairCardsMaterialFormat::VertexElementType,	EVertexStreamUsage::Default), 4));
 
 			// Ensure the rest resources are in correct states for the VF
 			const FHairCardsRestResource* RestResource = LOD.RestResource;
@@ -372,11 +397,27 @@ void FHairCardsVertexFactory::InitResources(FRHICommandListBase& RHICmdList)
 			const FHairGroupInstance::FMeshes::FLOD& LOD = Data.Instance->Meshes.LODs[Data.LODIndex];
 
 			// Note: Use the 'Normal' buffer as a dummy input for 'Material' buffer, as Material data is fetched through textures for meshes
-			Elements.Add(AccessStreamComponent(FVertexStreamComponent(&LOD.RestResource->RestPositionBuffer,	0, 0,									FHairCardsPositionFormat::SizeInByte,												FHairCardsPositionFormat::VertexElementType,	EVertexStreamUsage::Default), 0));
-			Elements.Add(AccessStreamComponent(FVertexStreamComponent(&LOD.RestResource->NormalsBuffer,			0, 0,									FHairCardsNormalFormat::SizeInByte * FHairCardsNormalFormat::ComponentCount,		FHairCardsNormalFormat::VertexElementType,		EVertexStreamUsage::Default), 1));
-			Elements.Add(AccessStreamComponent(FVertexStreamComponent(&LOD.RestResource->NormalsBuffer,			0, FHairCardsNormalFormat::SizeInByte,	FHairCardsNormalFormat::SizeInByte * FHairCardsNormalFormat::ComponentCount,		FHairCardsNormalFormat::VertexElementType,		EVertexStreamUsage::Default), 2));
-			Elements.Add(AccessStreamComponent(FVertexStreamComponent(&LOD.RestResource->UVsBuffer,				0, 0,									FHairCardsUVFormat::SizeInByte,														FHairCardsUVFormat::VertexElementType,			EVertexStreamUsage::Default), 3));
-			Elements.Add(AccessStreamComponent(FVertexStreamComponent(&LOD.RestResource->NormalsBuffer,			0, 0,									FHairCardsMaterialFormat::SizeInByte,												FHairCardsMaterialFormat::VertexElementType,	EVertexStreamUsage::Default), 4)); 
+			const bool bDynamic = LOD.DeformedResource != nullptr;
+			if (bDynamic)
+			{
+				DeformedPositionVertexBuffer[0] = FRDGWrapperVertexBuffer(LOD.DeformedResource->GetBuffer(FHairMeshesDeformedResource::EFrameType::Current));
+				DeformedPositionVertexBuffer[0].InitResource(RHICmdList);
+
+				DeformedPositionVertexBuffer[1] = FRDGWrapperVertexBuffer(LOD.DeformedResource->GetBuffer(FHairMeshesDeformedResource::EFrameType::Previous));
+				DeformedPositionVertexBuffer[1].InitResource(RHICmdList);
+
+				Elements.Add(AccessStreamComponent(FVertexStreamComponent(&(DeformedPositionVertexBuffer[0]),		0, 0,									FHairCardsPositionFormat::SizeInByte,												FHairCardsPositionFormat::VertexElementType,	EVertexStreamUsage::Default), 0));
+				Elements.Add(AccessStreamComponent(FVertexStreamComponent(&(DeformedPositionVertexBuffer[1]),		0, 0,									FHairCardsPositionFormat::SizeInByte,												FHairCardsPositionFormat::VertexElementType,	EVertexStreamUsage::Default), 5));
+			}
+			else
+			{
+				Elements.Add(AccessStreamComponent(FVertexStreamComponent(&LOD.RestResource->RestPositionBuffer,	0, 0,									FHairCardsPositionFormat::SizeInByte,												FHairCardsPositionFormat::VertexElementType,	EVertexStreamUsage::Default), 0));
+				Elements.Add(AccessStreamComponent(FVertexStreamComponent(&LOD.RestResource->RestPositionBuffer,	0, 0,									FHairCardsPositionFormat::SizeInByte,												FHairCardsPositionFormat::VertexElementType,	EVertexStreamUsage::Default), 5));
+			}
+			Elements.Add(AccessStreamComponent(FVertexStreamComponent(&LOD.RestResource->NormalsBuffer,				0, 0,									FHairCardsNormalFormat::SizeInByte * FHairCardsNormalFormat::ComponentCount,		FHairCardsNormalFormat::VertexElementType,		EVertexStreamUsage::Default), 1));
+			Elements.Add(AccessStreamComponent(FVertexStreamComponent(&LOD.RestResource->NormalsBuffer,				0, FHairCardsNormalFormat::SizeInByte,	FHairCardsNormalFormat::SizeInByte * FHairCardsNormalFormat::ComponentCount,		FHairCardsNormalFormat::VertexElementType,		EVertexStreamUsage::Default), 2));
+			Elements.Add(AccessStreamComponent(FVertexStreamComponent(&LOD.RestResource->UVsBuffer,					0, 0,									FHairCardsUVFormat::SizeInByte,														FHairCardsUVFormat::VertexElementType,			EVertexStreamUsage::Default), 3));
+			Elements.Add(AccessStreamComponent(FVertexStreamComponent(&LOD.RestResource->NormalsBuffer,				0, 0,									FHairCardsMaterialFormat::SizeInByte,												FHairCardsMaterialFormat::VertexElementType,	EVertexStreamUsage::Default), 4)); 
 
 			// Ensure the rest resources are in correct states for the VF
 			const FHairMeshesRestResource* RestResource = LOD.RestResource;
@@ -416,6 +457,14 @@ void FHairCardsVertexFactory::InitResources(FRHICommandListBase& RHICmdList)
 	}
 }
 
+void FHairCardsVertexFactory::ReleaseResource()
+{
+	FVertexFactory::ReleaseResource();
+	DeformedPositionVertexBuffer[0].ReleaseResource();
+	DeformedPositionVertexBuffer[1].ReleaseResource();
+	DeformedNormalVertexBuffer.ReleaseResource();
+}
+
 void FHairCardsVertexFactory::InitRHI(FRHICommandListBase& RHICmdList)
 {
 	// Nothing as the initialization is done when needed
diff --git a/Engine/Plugins/Runtime/HairStrands/Source/HairStrandsCore/Public/HairCardsVertexFactory.h b/Engine/Plugins/Runtime/HairStrands/Source/HairStrandsCore/Public/HairCardsVertexFactory.h
--- a/Engine/Plugins/Runtime/HairStrands/Source/HairStrandsCore/Public/HairCardsVertexFactory.h
+++ b/Engine/Plugins/Runtime/HairStrands/Source/HairStrandsCore/Public/HairCardsVertexFactory.h
@@ -15,6 +15,26 @@
 class FMaterial;
 class FSceneView;
 
+// Wrapper to reinterepet FRDGPooledBuffer as a FVertexBuffer
+class FRDGWrapperVertexBuffer : public FVertexBuffer
+{
+public:
+	FRDGWrapperVertexBuffer() {}
+	FRDGWrapperVertexBuffer(FRDGExternalBuffer& In): ExternalBuffer(In) { check(ExternalBuffer.Buffer); }
+	virtual void InitRHI(FRHICommandListBase& RHICmdList) override
+	{
+		check(ExternalBuffer.Buffer && ExternalBuffer.Buffer->GetRHI());
+		VertexBufferRHI = ExternalBuffer.Buffer->GetRHI();
+	}
+
+	virtual void ReleaseRHI() override
+	{
+		VertexBufferRHI = nullptr;
+	}
+
+	FRDGExternalBuffer ExternalBuffer;
+};
+
 /**
  * A vertex factory which simply transforms explicit vertex attributes from local to world space.
  */
@@ -54,6 +74,7 @@ public:
 	void Copy(const FHairCardsVertexFactory& Other);
 
 	void InitResources(FRHICommandListBase& RHICmdList);
+	virtual void ReleaseResource() override;
 
 	// FRenderResource interface.
 	virtual void InitRHI(FRHICommandListBase& RHICmdList) override;
@@ -64,6 +85,9 @@ protected:
 
 	bool bIsInitialized = false;
 
+	FRDGWrapperVertexBuffer DeformedPositionVertexBuffer[2];
+	FRDGWrapperVertexBuffer DeformedNormalVertexBuffer;
+
 	struct FDebugName
 	{
 		FDebugName(const char* InDebugName)
diff --git a/Engine/Shaders/Private/HairStrands/HairCardsVertexFactory.ush b/Engine/Shaders/Private/HairStrands/HairCardsVertexFactory.ush
--- a/Engine/Shaders/Private/HairStrands/HairCardsVertexFactory.ush
+++ b/Engine/Shaders/Private/HairStrands/HairCardsVertexFactory.ush
@@ -37,7 +37,8 @@ struct FVertexFactoryInput
 #if MANUAL_VERTEX_FETCH
 	uint VertexId : SV_VertexID;
 #else
-	float4 Position  : ATTRIBUTE0; 
+	float4 Position      : ATTRIBUTE0; 
+	float4 PrevPosition  : ATTRIBUTE5;
 
 	// TangentZ.w contains sign of tangent basis determinant
   #if METAL_PROFILE
@@ -145,7 +146,7 @@ float3 GetVertexPreviousPosition(FVertexFactoryInput Input)
 	const uint VertexFetchId = GetHairVertexFetchId(Input);
 	const float4 Data = HairCardsVF.PreviousPositionBuffer[VertexFetchId];
 #else
-	const float4 Data = Input.Position;
+	const float4 Data = Input.PrevPosition;
 #endif
 	return Data.xyz;
 }
diff --git a/Engine/Source/Runtime/Renderer/Private/HairStrands/HairStrandsDebug.cpp b/Engine/Source/Runtime/Renderer/Private/HairStrands/HairStrandsDebug.cpp
--- a/Engine/Source/Runtime/Renderer/Private/HairStrands/HairStrandsDebug.cpp
+++ b/Engine/Source/Runtime/Renderer/Private/HairStrands/HairStrandsDebug.cpp
@@ -469,7 +469,7 @@ class FHairDebugPS : public FGlobalShader
 	END_SHADER_PARAMETER_STRUCT()
 
 public:
-	static bool ShouldCompilePermutation(const FGlobalShaderPermutationParameters& Parameters) { return IsHairStrandsSupported(EHairStrandsShaderType::All, Parameters.Platform); }
+	static bool ShouldCompilePermutation(const FGlobalShaderPermutationParameters& Parameters) { return IsHairStrandsSupported(EHairStrandsShaderType::Strands, Parameters.Platform); }
 	static void ModifyCompilationEnvironment(const FGlobalShaderPermutationParameters& Parameters, FShaderCompilerEnvironment& OutEnvironment)
 	{
 		FGlobalShader::ModifyCompilationEnvironment(Parameters, OutEnvironment);
diff --git a/Engine/Source/Runtime/Renderer/Private/HairStrands/HairStrandsInterface.cpp b/Engine/Source/Runtime/Renderer/Private/HairStrands/HairStrandsInterface.cpp
--- a/Engine/Source/Runtime/Renderer/Private/HairStrands/HairStrandsInterface.cpp
+++ b/Engine/Source/Runtime/Renderer/Private/HairStrands/HairStrandsInterface.cpp
@@ -169,11 +169,7 @@ bool IsHairStrandsSupported(EHairStrandsShaderType Type, EShaderPlatform Platfor
 {
 	if (!IsGroomEnabled()) return false;
 
-	// Important:
-	// EHairStrandsShaderType::All: Mobile is excluded as we don't need any interpolation/simulation code for this. It only do rigid transformation. 
-	//                              The runtime setting in these case are r.HairStrands.Binding=0 & r.HairStrands.Simulation=0
 	const bool Cards_Meshes_All = true;
-	const bool bIsMobile = IsMobilePlatform(Platform);
 
 	switch (Type)
 	{
@@ -181,7 +177,7 @@ bool IsHairStrandsSupported(EHairStrandsShaderType Type, EShaderPlatform Platfor
 	case EHairStrandsShaderType::Cards:	  return Cards_Meshes_All;
 	case EHairStrandsShaderType::Meshes:  return Cards_Meshes_All;
 	case EHairStrandsShaderType::Tool:	  return (IsD3DPlatform(Platform) || IsVulkanPlatform(Platform)) && IsPCPlatform(Platform) && IsFeatureLevelSupported(Platform, ERHIFeatureLevel::SM5);
-	case EHairStrandsShaderType::All:	  return Cards_Meshes_All && !bIsMobile;
+	case EHairStrandsShaderType::All:	  return Cards_Meshes_All;
 	}
 	return false;
 }
@@ -191,10 +187,6 @@ bool IsHairStrandsEnabled(EHairStrandsShaderType Type, EShaderPlatform Platform)
 	const bool HairStrandsGlobalEnable = IsGroomEnabled();
 	if (!HairStrandsGlobalEnable) return false;
 
-	// Important:
-	// EHairStrandsShaderType::All: Mobile is excluded as we don't need any interpolation/simulation code for this. It only do rigid transformation. 
-	//                              The runtime setting in these case are r.HairStrands.Binding=0 & r.HairStrands.Simulation=0
-	const bool bIsMobile = Platform != EShaderPlatform::SP_NumPlatforms ? IsMobilePlatform(Platform) : false;
 	const int32 HairStrandsEnable = CVarHairStrandsEnable.GetValueOnAnyThread();
 	const int32 HairCardsEnable   = CVarHairCardsEnable.GetValueOnAnyThread();
 	const int32 HairMeshesEnable  = CVarHairMeshesEnable.GetValueOnAnyThread();
@@ -208,7 +200,7 @@ bool IsHairStrandsEnabled(EHairStrandsShaderType Type, EShaderPlatform Platform)
 #else
 	case EHairStrandsShaderType::Tool:		return false;
 #endif
-	case EHairStrandsShaderType::All :		return HairStrandsGlobalEnable && (HairCardsEnable > 0 || HairMeshesEnable > 0 || HairStrandsEnable > 0) && !bIsMobile;
+	case EHairStrandsShaderType::All :		return HairStrandsGlobalEnable && (HairCardsEnable > 0 || HairMeshesEnable > 0 || HairStrandsEnable > 0);
 	}
 	return false;
 }
diff --git a/Engine/Source/Runtime/Renderer/Private/MobileShadingRenderer.cpp b/Engine/Source/Runtime/Renderer/Private/MobileShadingRenderer.cpp
--- a/Engine/Source/Runtime/Renderer/Private/MobileShadingRenderer.cpp
+++ b/Engine/Source/Runtime/Renderer/Private/MobileShadingRenderer.cpp
@@ -1111,6 +1111,36 @@ void FMobileSceneRenderer::Render(FRDGBuilder& GraphBuilder)
 			PendingRDGResources.CommitToSceneAndViewUniformBuffers(GraphBuilder, /* out */ ExternalAccessQueue);
 		}
 
+		// Hair update
+		if (IsHairStrandsEnabled(EHairStrandsShaderType::All, Scene->GetShaderPlatform()) && RendererOutput == ERendererOutput::FinalSceneColor)
+		{
+			FHairStrandsBookmarkParameters& HairStrandsBookmarkParameters = *GraphBuilder.AllocObject<FHairStrandsBookmarkParameters>();
+			CreateHairStrandsBookmarkParameters(Scene, Views, AllFamilyViews, HairStrandsBookmarkParameters);
+			check(Scene->HairStrandsSceneData.TransientResources);
+			HairStrandsBookmarkParameters.TransientResources = Scene->HairStrandsSceneData.TransientResources;
+
+			// Not need for hair uniform buffer, as this is only used for strands rendering
+			// If some shader refers to it, we can create a default one with HairStrands::CreateDefaultHairStrandsViewUniformBuffer(GraphBuilder, View);
+			for (FViewInfo& View : Views)
+			{
+				View.HairStrandsViewData.UniformBuffer = nullptr;
+			}
+
+			// Interpolation needs to happen after the skin cache run as there is a dependency 
+			// on the skin cache output.
+			const bool bRunHairStrands = HairStrandsBookmarkParameters.HasInstances() && (Views.Num() > 0);
+			if (bRunHairStrands)
+			{
+				// 1. Update groom visible in primary views
+				RunHairStrandsBookmark(GraphBuilder, EHairStrandsBookmark::ProcessCardsAndMeshesInterpolation_PrimaryView, HairStrandsBookmarkParameters);
+
+				// 2. Update groom only visible in shadow 
+				// For now, not running on mobile to keep computation light
+				// UpdateHairStrandsBookmarkParameters(Scene, Views, HairStrandsBookmarkParameters);
+				// RunHairStrandsBookmark(GraphBuilder, EHairStrandsBookmark::ProcessCardsAndMeshesInterpolation_ShadowView, HairStrandsBookmarkParameters);
+			}
+		}
+
 		GraphBuilder.SetCommandListStat(GET_STATID(STAT_CLMM_Shadows));
 		RenderShadowDepthMaps(GraphBuilder, nullptr, InstanceCullingManager, ExternalAccessQueue);
 		GraphBuilder.AddDispatchHint();
