GroomCache:
* Exposed GroomCache setter to BP/Sequencer.
* Allow adding GroomCache track in Sequencer to Groom Component track even if the GroomCache is not set.

diff --git a/Engine/Plugins/Runtime/HairStrands/Source/HairStrandsCore/Public/GroomComponent.h b/Engine/Plugins/Runtime/HairStrands/Source/HairStrandsCore/Public/GroomComponent.h
--- a/Engine/Plugins/Runtime/HairStrands/Source/HairStrandsCore/Public/GroomComponent.h
+++ b/Engine/Plugins/Runtime/HairStrands/Source/HairStrandsCore/Public/GroomComponent.h
@@ -281,6 +281,9 @@ public:
 
 	/** GroomCache */
 	UGroomCache* GetGroomCache() const { return GroomCache; }
+
+	/* Accessor function for changing GroomCache asset from blueprint/sequencer */
+	UFUNCTION(BlueprintCallable, Category = "Groom")
 	void SetGroomCache(UGroomCache* InGroomCache);
 
 	float GetGroomCacheDuration() const;
diff --git a/Engine/Plugins/Runtime/HairStrands/Source/HairStrandsEditor/Private/GroomCacheTrackEditor.cpp b/Engine/Plugins/Runtime/HairStrands/Source/HairStrandsEditor/Private/GroomCacheTrackEditor.cpp
--- a/Engine/Plugins/Runtime/HairStrands/Source/HairStrandsEditor/Private/GroomCacheTrackEditor.cpp
+++ b/Engine/Plugins/Runtime/HairStrands/Source/HairStrandsEditor/Private/GroomCacheTrackEditor.cpp
@@ -39,10 +39,7 @@ static UGroomComponent* AcquireGroomComponentFromObjectGuid(const FGuid& Guid, T
 	}
 	else if (UGroomComponent* GroomComp = Cast<UGroomComponent>(BoundObject))
 	{
-		if (GroomComp->GetGroomCache())
-		{
-			return GroomComp;
-		}
+		return GroomComp;
 	}
 
 	return nullptr;
