HairCardSolver: Update HCS button and icon.

diff --git a/Engine/Plugins/Runtime/HairStrands/Source/HairStrandsEditor/Private/GroomAssetDetails.cpp b/Engine/Plugins/Runtime/HairStrands/Source/HairStrandsEditor/Private/GroomAssetDetails.cpp
--- a/Engine/Plugins/Runtime/HairStrands/Source/HairStrandsEditor/Private/GroomAssetDetails.cpp
+++ b/Engine/Plugins/Runtime/HairStrands/Source/HairStrandsEditor/Private/GroomAssetDetails.cpp
@@ -1663,16 +1663,35 @@ void FGroomRenderingDetails::OnGenerateElementForHairGroup(TSharedRef<IPropertyH
 						FText RegenToolTipText = LOCTEXT("GenerateNewCardsTooltip", "Generate new card assets (meshes and textures) using the current procedural settings. NOTE: This will overwrite preexisting card assets for this LOD.");
 						ProceduralGenButtons->AddSlot()
 							.AutoWidth()
+							.Padding(0.0f, 2.0f, 0.0f, 4.0f)
 							[
 								SNew(SButton)
 									.VAlign(VAlign_Center)
 									.HAlign(HAlign_Center)
-									.ButtonStyle(FAppStyle::Get(), "HoverHintOnly")
+									.ContentPadding(0)
+									.ButtonStyle(FAppStyle::Get(), "Button")
 									.ToolTipText(RegenToolTipText)
 									.OnClicked(this, &FGroomRenderingDetails::OnGenerateCardDataUsingPlugin, GroupIndex)
 									[
-										SNew(SImage)
-											.Image(Brush)
+										SNew(SHorizontalBox)
+										+ SHorizontalBox::Slot()
+										.AutoWidth()
+										.VAlign(VAlign_Center)
+										.HAlign(HAlign_Center)
+										.Padding(-1.0f, 2.0f, 2.0f, 2.0f)
+										[
+											SNew(SImage)
+												.Image(Brush)
+										]
+										+ SHorizontalBox::Slot()
+										.AutoWidth()
+										.VAlign(VAlign_Center)
+										.HAlign(HAlign_Center)
+										.Padding(2.0f, 2.0f, -1.0f, 2.0f)
+										[
+											SNew(STextBlock)
+												.Text(LOCTEXT("GenerateNewCardsButtonText","Generate Hair Cards"))
+										]
 									]
 							];
 					}
diff --git a/Engine/Plugins/Runtime/HairStrands/Source/HairStrandsEditor/Public/GroomEditorStyle.h b/Engine/Plugins/Runtime/HairStrands/Source/HairStrandsEditor/Public/GroomEditorStyle.h
--- a/Engine/Plugins/Runtime/HairStrands/Source/HairStrandsEditor/Public/GroomEditorStyle.h
+++ b/Engine/Plugins/Runtime/HairStrands/Source/HairStrandsEditor/Public/GroomEditorStyle.h
@@ -8,6 +8,7 @@
 #include "Styling/SlateStyleRegistry.h"
 
 #define IMAGE_BRUSH(RelativePath, ...) FSlateImageBrush(RootToContentDir(RelativePath, TEXT(".png")), __VA_ARGS__)
+#define IMAGE_BRUSH_SVG( RelativePath, ... ) FSlateVectorImageBrush( RootToContentDir( RelativePath, TEXT(".svg") ), __VA_ARGS__ )
 
 /**
  * Implements the visual style of the Groom asset editor UI.
@@ -33,7 +34,8 @@ public:
 //		Set("HairToolsManagerCommands.BeginAddPrimitiveTool", new IMAGE_BRUSH("HairSpray128", Icon40x40));
 		Set("GroomEditor.SimulationOptions", new IMAGE_BRUSH("S_SimulationOptions_40x", Icon40x40));
 		Set("GroomEditor.SimulationOptions.Small", new IMAGE_BRUSH("S_SimulationOptions_40x", Icon20x20));
-		Set("GroomEditor.GroomCardGenerator", new IMAGE_BRUSH("S_GroomCardGenerator_128", Icon16x16));
+		Set("GroomEditor.GroomCardGenerator", new IMAGE_BRUSH_SVG("S_GroomCardGenerator", Icon16x16));
+		// Set("GroomEditor.GroomCardGenerator", new IMAGE_BRUSH("S_GroomCardGenerator_128", Icon16x16));
 		if (FSlateStyleRegistry::FindSlateStyle(GetStyleName()) == nullptr)
 		{
 			FSlateStyleRegistry::RegisterSlateStyle(*this);
