Fix hair motion blur issue with MRQ.

This changes the hair position swap at component tick time.

#changelist validated
#virtualized

diff --git a/Engine/Plugins/Runtime/HairStrands/Source/HairStrandsCore/Private/GroomManager.cpp b/Engine/Plugins/Runtime/HairStrands/Source/HairStrandsCore/Private/GroomManager.cpp
--- a/Engine/Plugins/Runtime/HairStrands/Source/HairStrandsCore/Private/GroomManager.cpp
+++ b/Engine/Plugins/Runtime/HairStrands/Source/HairStrandsCore/Private/GroomManager.cpp
@@ -38,7 +38,7 @@ static FAutoConsoleVariableRef CVarGHairStrandsMinLOD(TEXT("r.HairStrands.MinLOD
 static int32 GHairStrands_UseCards = 0;
 static FAutoConsoleVariableRef CVarHairStrands_UseCards(TEXT("r.HairStrands.UseCardsInsteadOfStrands"), GHairStrands_UseCards, TEXT("Force cards geometry on all groom elements. If no cards data is available, nothing will be displayed"), ECVF_Scalability);
 
-static int32 GHairStrands_SwapBufferType = 3;
+static int32 GHairStrands_SwapBufferType = 2;
 static FAutoConsoleVariableRef CVarGHairStrands_SwapBufferType(TEXT("r.HairStrands.SwapType"), GHairStrands_SwapBufferType, TEXT("Swap rendering buffer at the end of frame. This is an experimental toggle. Default:1"));
 
 static int32 GHairStrands_ManualSkinCache = 0;
