Do not run hair update if the instance is explicit set to 'not visible'.

#changelist validated
#virtualized

diff --git a/Engine/Plugins/Runtime/HairStrands/Source/HairStrandsCore/Private/GroomManager.cpp b/Engine/Plugins/Runtime/HairStrands/Source/HairStrandsCore/Private/GroomManager.cpp
--- a/Engine/Plugins/Runtime/HairStrands/Source/HairStrandsCore/Private/GroomManager.cpp
+++ b/Engine/Plugins/Runtime/HairStrands/Source/HairStrandsCore/Private/GroomManager.cpp
@@ -2976,6 +2976,17 @@ static void RunHairLODSelection(
 			InstanceData.Instance->BindingType = EHairBindingType::NoneBinding;
 		}
 
+		// Cull instance which are not visibles
+		for (const FSceneView* View : Views)
+		{
+			if (!InstanceData.Instance->Debug.Proxy->IsShown(View) && !InstanceData.Instance->Debug.Proxy->IsShadowCast(View))
+			{
+				InstanceData.Instance->GeometryType = EHairGeometryType::NoneGeometry;
+				InstanceData.Instance->BindingType = EHairBindingType::NoneBinding;
+				break;
+			}
+		}
+
 		// Update the local-to-world transform based on the binding type 
 		if (GetHairSwapBufferType() == EHairBufferSwapType::BeginOfFrame || GetHairSwapBufferType() == EHairBufferSwapType::EndOfFrame)
 		{
