Add project settings for hair tristrip geometry as not all project are compatible with this.

#changelist validated
#virtualized

diff --git a/Engine/Plugins/Runtime/HairStrands/Source/HairStrandsCore/Private/GroomComponent.cpp b/Engine/Plugins/Runtime/HairStrands/Source/HairStrandsCore/Private/GroomComponent.cpp
--- a/Engine/Plugins/Runtime/HairStrands/Source/HairStrandsCore/Private/GroomComponent.cpp
+++ b/Engine/Plugins/Runtime/HairStrands/Source/HairStrandsCore/Private/GroomComponent.cpp
@@ -263,6 +263,16 @@ static FHairGroupDesc GetGroomGroupsDesc(const UGroomAsset* Asset, const UGroomC
 	return O;
 }
 
+static EPrimitiveType GetPrimitiveType(EHairGeometryType In)
+{
+	return (In == EHairGeometryType::Strands && GetHairStrandsUsesTriangleStrips()) ? PT_TriangleStrip : PT_TriangleList;
+}
+
+static uint32 GetPointToVertexCount()
+{
+	return GetHairStrandsUsesTriangleStrips() ? HAIR_POINT_TO_VERTEX_FOR_TRISTRP : HAIR_POINT_TO_VERTEX_FOR_TRILIST;
+}
+
 ////////////////////////////////////////////////////////////////////////////////////////////////////////////////////////
 
 /**
@@ -923,7 +933,7 @@ public:
 			VertexFactory = (FVertexFactory*)Instance->Strands.VertexFactory;
 			PrimitiveIdMode = Instance->Strands.VertexFactory->GetPrimitiveIdMode(FeatureLevel);
 			HairVertexCount = Instance->HairGroupPublicData->GetActiveStrandsPointCount();
-			MaxVertexIndex = HairVertexCount * HAIR_POINT_TO_VERTEX;
+			MaxVertexIndex = HairVertexCount * GetPointToVertexCount();
 			bUseCulling = Instance->Strands.bCullingEnable;
 			NumPrimitive = bUseCulling ? 0 : HairVertexCount * HAIR_POINT_TO_TRIANGLE;
 			if (MaterialRenderProxy == nullptr)
@@ -1039,11 +1049,7 @@ public:
 		BatchElement.UserData = reinterpret_cast<void*>(uint64(ComponentId));
 		Mesh.ReverseCulling = bUseCardsOrMeshes ? IsLocalToWorldDeterminantNegative() : false;
 		Mesh.bDisableBackfaceCulling = GeometryType == EHairGeometryType::Strands;
-		#if USE_HAIR_TRIANGLE_STRIP
-		Mesh.Type = GeometryType == EHairGeometryType::Strands ? PT_TriangleStrip : PT_TriangleList;
-		#else
-		Mesh.Type = PT_TriangleList;
-		#endif
+		Mesh.Type = GetPrimitiveType(GeometryType);
 		Mesh.DepthPriorityGroup = SDPG_World;
 		Mesh.bCanApplyViewModeOverrides = false;
 		Mesh.BatchHitProxyId = PrimSceneInfo->DefaultDynamicHitProxyId;
@@ -3143,11 +3149,7 @@ void UGroomComponent::CollectPSOPrecacheData(const FPSOPrecacheParams& BasePreca
 		UMaterialInterface* MaterialInterface = GetMaterial(GetMaterialIndexWithFallback(VFsPerMaterial.MaterialIndex), VFsPerMaterial.HairGeometryType, true);
 		if (MaterialInterface)
 		{
-#if USE_HAIR_TRIANGLE_STRIP
-			PrecachePSOParams.PrimitiveType = VFsPerMaterial.HairGeometryType == EHairGeometryType::Strands ? PT_TriangleStrip : PT_TriangleList;
-#else
-			PrecachePSOParams.PrimitiveType = PT_TriangleList;
-#endif
+			PrecachePSOParams.PrimitiveType = GetPrimitiveType(VFsPerMaterial.HairGeometryType);
 
 			FMaterialInterfacePSOPrecacheParams& ComponentParams = OutParams[OutParams.AddDefaulted()];
 			ComponentParams.MaterialInterface = MaterialInterface;
diff --git a/Engine/Shaders/Shared/HairStrandsDefinitions.h b/Engine/Shaders/Shared/HairStrandsDefinitions.h
--- a/Engine/Shaders/Shared/HairStrandsDefinitions.h
+++ b/Engine/Shaders/Shared/HairStrandsDefinitions.h
@@ -45,14 +45,16 @@
 // Max split for raytracing geometry
 #define STRANDS_PROCEDURAL_INTERSECTOR_MAX_SPLITS 4
 
-// Use triangle strip for HW raster path
-#define USE_HAIR_TRIANGLE_STRIP 1
-
 // Number of vertex per control-point
+#define HAIR_POINT_TO_VERTEX_FOR_TRISTRP 2u
+#define HAIR_POINT_TO_VERTEX_FOR_TRILIST 6u
+
+#ifndef __cplusplus
 #if USE_HAIR_TRIANGLE_STRIP
-#define HAIR_POINT_TO_VERTEX 2u
+#define HAIR_POINT_TO_VERTEX HAIR_POINT_TO_VERTEX_FOR_TRISTRP
 #else
-#define HAIR_POINT_TO_VERTEX 6u
+#define HAIR_POINT_TO_VERTEX HAIR_POINT_TO_VERTEX_FOR_TRILIST
+#endif
 #endif
 
 // Number of triangle per control-point
diff --git a/Engine/Source/Runtime/Engine/Private/ShaderCompiler/ShaderCompiler.cpp b/Engine/Source/Runtime/Engine/Private/ShaderCompiler/ShaderCompiler.cpp
--- a/Engine/Source/Runtime/Engine/Private/ShaderCompiler/ShaderCompiler.cpp
+++ b/Engine/Source/Runtime/Engine/Private/ShaderCompiler/ShaderCompiler.cpp
@@ -8441,6 +8441,11 @@ void GlobalBeginCompileShader(
 		SET_SHADER_DEFINE(Input.Environment, SUPPORT_VSM_FOWARD_QUALITY, bHighQualityShadow ? 1 : 0);
 	}
 
+	{
+		const bool bUseTriangleStrips = GetHairStrandsUsesTriangleStrips();
+		SET_SHADER_DEFINE(Input.Environment, USE_HAIR_TRIANGLE_STRIP, bUseTriangleStrips ? 1 : 0);
+	}
+
 	const bool bSubstrate = Substrate::IsSubstrateEnabled() && IsSubstrateSupportForShaderPipeline(Input);
 	{
 		SET_SHADER_DEFINE(Input.Environment, SUBSTRATE_ENABLED, bSubstrate ? 1 : 0);
diff --git a/Engine/Source/Runtime/RenderCore/Private/RenderUtils.cpp b/Engine/Source/Runtime/RenderCore/Private/RenderUtils.cpp
--- a/Engine/Source/Runtime/RenderCore/Private/RenderUtils.cpp
+++ b/Engine/Source/Runtime/RenderCore/Private/RenderUtils.cpp
@@ -1683,6 +1683,12 @@ static TAutoConsoleVariable<int32> CVarIESProfileTranslucent(
 	TEXT("Enable IES profiles support for translucent surfaces. When enabled, it will add an extrat sampler to the pixel shader (limited to 16 on dx11 based system)"),
 	ECVF_ReadOnly | ECVF_RenderThreadSafe);
 
+static TAutoConsoleVariable<int32> CVarHairStrandsUsesTriangleStrips(
+	TEXT("r.HairStrands.Strands.UseTriangleStrips"),
+	1,
+	TEXT("Enable triangle strip geometry for hair strands rendering. This improves performances, but removes the last segments of each curve."),
+	ECVF_ReadOnly | ECVF_RenderThreadSafe);
+
 int32 GetLightFunctionAtlasFormat()
 {
 	return CVarLightFunctionAtlasFormat.GetValueOnAnyThread();
@@ -1708,6 +1714,11 @@ bool GetTranslucentUsesLightIESProfiles()
 	return CVarIESProfileTranslucent.GetValueOnAnyThread() > 0;
 }
 
+bool GetHairStrandsUsesTriangleStrips()
+{
+	return CVarHairStrandsUsesTriangleStrips.GetValueOnAnyThread() > 0;
+}
+
 ///////////////////////////////////////////////////////////////////////////////////////////////////
 // Substrate settings interface
 
diff --git a/Engine/Source/Runtime/RenderCore/Private/Shader.cpp b/Engine/Source/Runtime/RenderCore/Private/Shader.cpp
--- a/Engine/Source/Runtime/RenderCore/Private/Shader.cpp
+++ b/Engine/Source/Runtime/RenderCore/Private/Shader.cpp
@@ -2009,6 +2009,11 @@ void ShaderMapAppendKeyString(EShaderPlatform Platform, FString& KeyString)
 		}
 	}
 
+	if (GetHairStrandsUsesTriangleStrips())
+	{
+		KeyString += FString::Printf(TEXT("_STRDSTRIP"));
+	}
+
 	if (Substrate::IsSubstrateEnabled())
 	{
 		{
diff --git a/Engine/Source/Runtime/RenderCore/Public/RenderUtils.h b/Engine/Source/Runtime/RenderCore/Public/RenderUtils.h
--- a/Engine/Source/Runtime/RenderCore/Public/RenderUtils.h
+++ b/Engine/Source/Runtime/RenderCore/Public/RenderUtils.h
@@ -578,6 +578,8 @@ RENDERCORE_API bool GetTranslucentUsesLightFunctionAtlas();
 RENDERCORE_API bool GetTranslucentUsesLightRectLights();
 RENDERCORE_API bool GetTranslucentUsesLightIESProfiles();
 
+RENDERCORE_API bool GetHairStrandsUsesTriangleStrips();
+
 // LuminanceMax is the amount of light that will cause the sensor to saturate at EV100.
 //  See also https://en.wikipedia.org/wiki/Film_speed and https://en.wikipedia.org/wiki/Exposure_value for more info.
 FORCEINLINE float EV100ToLuminance(float LuminanceMax, float EV100)
