Disable hair compressed position, as decompression is sometime incorrect.

#changelist validated
#virtualized

diff --git a/Engine/Plugins/Runtime/HairStrands/Source/HairStrandsCore/Private/GroomAsset.cpp b/Engine/Plugins/Runtime/HairStrands/Source/HairStrandsCore/Private/GroomAsset.cpp
--- a/Engine/Plugins/Runtime/HairStrands/Source/HairStrandsCore/Private/GroomAsset.cpp
+++ b/Engine/Plugins/Runtime/HairStrands/Source/HairStrandsCore/Private/GroomAsset.cpp
@@ -85,7 +85,7 @@ static TAutoConsoleVariable<int32> GHairStrandsWarningLogVerbosity(
 static int32 GHairStrandsDDCLogEnable = 0;
 static FAutoConsoleVariableRef CVarHairStrandsDDCLogEnable(TEXT("r.HairStrands.DDCLog"), GHairStrandsDDCLogEnable, TEXT("Enable DDC logging for groom assets and groom binding assets"));
 
-static int32 GHairStrandsSupportCompressedPosition = 1;
+static int32 GHairStrandsSupportCompressedPosition = 0;
 static FAutoConsoleVariableRef CVarHairStrandsSupportCompressedPosition(TEXT("r.HairStrands.CompressedPosition"), GHairStrandsSupportCompressedPosition, TEXT("Optional compessed position"), ECVF_ReadOnly);
 
 bool DoesHairStrandsSupportCompressedPosition()
