* Add coverage scaling computation to hair strands builder to automatically compute coverage preservation for AutoLOD
* Add UI for configuring AutoLOD. Now AutoLOD can be enabled/disabled per project. Groom asset will use project setting, but can override it per asset (Default, Manual, Auto).
* Global CVar still allow to force/toggle AutoLOD for debugging purpose only.
* Add a default screen bias of 10% to mitigate early curve decimation for most groom

diff --git a/Engine/Plugins/Runtime/HairStrands/Source/HairStrandsCore/Private/GroomAsset.cpp b/Engine/Plugins/Runtime/HairStrands/Source/HairStrandsCore/Private/GroomAsset.cpp
--- a/Engine/Plugins/Runtime/HairStrands/Source/HairStrandsCore/Private/GroomAsset.cpp
+++ b/Engine/Plugins/Runtime/HairStrands/Source/HairStrandsCore/Private/GroomAsset.cpp
@@ -85,6 +85,14 @@ static TAutoConsoleVariable<int32> GHairStrandsWarningLogVerbosity(
 static int32 GHairStrandsDDCLogEnable = 0;
 static FAutoConsoleVariableRef CVarHairStrandsDDCLogEnable(TEXT("r.HairStrands.DDCLog"), GHairStrandsDDCLogEnable, TEXT("Enable DDC logging for groom assets and groom binding assets"));
 
+static int32 GHairStrands_LODMode = 0;
+static FAutoConsoleVariableRef CVarHairStrands_LODMode(TEXT("r.HairStrands.LODMode"), GHairStrands_LODMode, TEXT("Enable hair strands Auto LOD mode by default. Otherwise use Manual LOD mode. Auto LOD mode adapts hair curves based on screen coverage. Manual LOD mode relies on LODs manually setup per groom asset. This global behavior can be overridden per groom asset."), ECVF_RenderThreadSafe);
+
+EGroomLODMode GetHairStrandsLODMode()
+{
+	return GHairStrands_LODMode > 0 ? EGroomLODMode::Auto : EGroomLODMode::Manual;
+}
+
 static int32 GHairStrandsSupportCompressedPosition = 0;
 static FAutoConsoleVariableRef CVarHairStrandsSupportCompressedPosition(TEXT("r.HairStrands.CompressedPosition"), GHairStrandsSupportCompressedPosition, TEXT("Optional compessed position"), ECVF_ReadOnly);
 
@@ -1602,7 +1610,14 @@ void UGroomAsset::PostEditChangeProperty(FPropertyChangedEvent& PropertyChangedE
 	}
 
 	const bool bHairStrandsRaytracingRadiusChanged = PropertyName == GET_MEMBER_NAME_CHECKED(FHairShadowSettings, HairRaytracingRadiusScale);
-	
+
+	// Fast path to only rebuild the proxies
+	if (PropertyName == UGroomAsset::GetAutoLODBiasMemberName())
+	{
+		FGroomComponentRecreateRenderStateContext Context(this);
+		return;
+	}
+
 	// By pass update if bStrandsInterpolationChanged has the resources have already been recreated
 	if (!bNeedRebuildDerivedData)
 	{
@@ -3831,4 +3846,17 @@ void UGroomAsset::ChangeFeatureLevel(ERHIFeatureLevel::Type In)
 		CachedResourcesFeatureLevel = In;
 	}
 }
-#endif
\ No newline at end of file
+#endif
+
+
+GROOMASSET_DEFINE_MEMBER_NAME(LODMode)
+EGroomLODMode UGroomAsset::GetLODMode() const
+{
+	return LODMode == EGroomLODMode::Default ? GetHairStrandsLODMode() : LODMode;
+}
+
+GROOMASSET_DEFINE_MEMBER_NAME(AutoLODBias)
+float UGroomAsset::GetAutoLODBias() const
+{
+	return AutoLODBias;
+}
\ No newline at end of file
diff --git a/Engine/Plugins/Runtime/HairStrands/Source/HairStrandsCore/Private/GroomBuilder.cpp b/Engine/Plugins/Runtime/HairStrands/Source/HairStrandsCore/Private/GroomBuilder.cpp
--- a/Engine/Plugins/Runtime/HairStrands/Source/HairStrandsCore/Private/GroomBuilder.cpp
+++ b/Engine/Plugins/Runtime/HairStrands/Source/HairStrandsCore/Private/GroomBuilder.cpp
@@ -51,7 +51,18 @@ bool DoesHairStrandsSupportCompressedPosition();
 
 FString FGroomBuilder::GetVersion()
 {
-	return TEXT("v14");
+	return TEXT("v15c");
+}
+
+namespace GroomBuilder_Voxelization
+{
+	struct FCoverageScale
+	{
+		float ScreenSize = 0;
+		float CoverageScale = 0;
+	};
+
+	void ComputeHairCoverageScale(const FHairStrandsDatas& In, TArray<FCoverageScale>& Out);
 }
 
 // For debug purpose
@@ -182,7 +193,7 @@ namespace HairStrandsBuilder
 	}
 
 	/** Build the bulk/packed datas for gpu rendering/simulation */
-	void BuildBulkData(const FHairStrandsDatas& HairStrands, const TArray<uint8>& RandomSeeds, FHairStrandsBulkData& OutBulkData, bool bAllowTranscoding, uint32 InGroupFlags)
+	void BuildBulkData(const FHairStrandsDatas& HairStrands, const TArray<uint8>& RandomSeeds, const TArray<GroomBuilder_Voxelization::FCoverageScale>& InCoverageScales, FHairStrandsBulkData& OutBulkData, bool bAllowTranscoding, uint32 InGroupFlags)
 	{
 		TRACE_CPUPROFILER_EVENT_SCOPE(HairStrandsBuilder::BuildBulkData);
 
@@ -770,6 +781,12 @@ namespace HairStrandsBuilder
 			OutBulkData.Header.CurveToPointCount.Add(Curve.PointOffset + Curve.PointCount);
 		}
 
+		OutBulkData.Header.CoverageScales.Reserve(InCoverageScales.Num());
+		for (const GroomBuilder_Voxelization::FCoverageScale& Scale : InCoverageScales)
+		{
+			OutBulkData.Header.CoverageScales.Add(Scale.CoverageScale);
+		}
+
 		// Stride datas
 		OutBulkData.Header.Strides.PositionStride = FHairStrandsPositionFormat::SizeInByte;
 		OutBulkData.Header.Strides.CurveStride = FHairStrandsCurveFormat::SizeInByte;
@@ -1596,7 +1613,7 @@ namespace HairInterpolationBuilder
 						
 						const FVector3f& SimPointPosition0 = SimStrandsData.StrandsPoints.PointsPosition[Desc.Index0 + SimOffset];
 						const FVector3f& SimPointPosition1 = SimStrandsData.StrandsPoints.PointsPosition[Desc.Index1 + SimOffset];
-						const float Weight = 1.0f / FMath::Max(MinWeightDistance, FVector3f::Distance(RenPointPosition, FMath::Lerp(SimPointPosition0, SimPointPosition1, Desc.T)));
+						const float Weight = 1.0f / FMath::Max(MinWeightDistance, FVector3f::Distance(RenPointPosition, FMath::LerpStable(SimPointPosition0, SimPointPosition1, Desc.T)));
 
 						InterpolationData.PointSimIndices[PointGlobalIndex][KIndex] = Desc.Index0/* + SimOffset*/;
 						InterpolationData.PointSimLerps[PointGlobalIndex][KIndex] = Desc.T;
@@ -2389,8 +2406,11 @@ void FGroomBuilder::BuildBulkData(
 	FHairStrandsBulkData& OutBulkData, 
 	bool bAllowTranscoding)
 {
+	const bool bIsNotGuide = bAllowTranscoding;
+
 	OutBulkData.Reset();
 
+	// 1. Initialize curve seed
 	const int32 NumCurves = InData.GetNumCurves();
 	FRandomStream Random;
 	Random.Initialize(InInfo.GroupID);
@@ -2400,7 +2420,23 @@ void FGroomBuilder::BuildBulkData(
 	{
 		CurveSeeds[Index] = Random.RandHelper(255);
 	}
-	HairStrandsBuilder::BuildBulkData(InData, CurveSeeds, OutBulkData, bAllowTranscoding, InInfo.Flags);
+
+	// 2. Compute coverage scaling based on screen size & decimation
+	TArray<GroomBuilder_Voxelization::FCoverageScale> CoverageScales;
+	if (bIsNotGuide)
+	{
+		GroomBuilder_Voxelization::ComputeHairCoverageScale(InData, CoverageScales);
+	}
+	else
+	{
+		GroomBuilder_Voxelization::FCoverageScale Dummy;
+		Dummy.ScreenSize = 1;
+		Dummy.CoverageScale = 1;
+		CoverageScales.Init(Dummy, 10);
+	}
+
+	// 3. Build bulk data
+	HairStrandsBuilder::BuildBulkData(InData, CurveSeeds, CoverageScales, OutBulkData, bAllowTranscoding, InInfo.Flags);
 }
 
 void FGroomBuilder::BuildInterplationData(
@@ -2761,11 +2797,11 @@ struct FHairGrid
 	TArray<FVoxel> Voxels;
 };
 
-static void Voxelize(const FHairDescriptionGroups& InGroups, FHairGrid& Out)
+static void AllocateVoxels(const FBox3f& InBoundingBox, FHairGrid& Out)
 {
 	// 1. Compute the overal bound of for all hair groups
-	Out.MinBound = InGroups.Bounds.GetBox().Min;
-	Out.MaxBound = InGroups.Bounds.GetBox().Max;
+	Out.MinBound = InBoundingBox.Min;
+	Out.MaxBound = InBoundingBox.Max;
 	Out.Resolution = FIntVector::ZeroValue;
 
 	// 2. Based on the bound, determine the voxel grid resolution
@@ -2787,67 +2823,82 @@ static void Voxelize(const FHairDescriptionGroups& InGroups, FHairGrid& Out)
 		Out.MaxBound = Out.MinBound + FVector3f(Out.Resolution) * Out.VoxelSize;
 	}
 	Out.Voxels.SetNum(FMath::Max(Out.Resolution.X * Out.Resolution.Y * Out.Resolution.Z, 1));
+}
 
+static void VoxelizeCurves(const FHairStrandsDatas& InData, uint32 InGroupIndex, FHairGrid& Out)
+{
 	// 3. Voxelization curves
-	uint32 GroupIndex = 0;
-	for (const FHairDescriptionGroup& Group : InGroups.HairGroups)
+	// Local copy of the hair strands data, and build the derived data so that curve offset are available
+	FHairStrandsDatas In = InData;
+	HairStrandsBuilder::BuildInternalData(In);
+
+	const uint32 Attributes = In.GetAttributes();
+
+	// Fill in voxel (TODO: make it parallel)
+	for (uint32 CurveIt = 0, CurveCount = In.GetNumCurves(); CurveIt < CurveCount; ++CurveIt)
 	{
-		// Local copy of the hair strands data, and build the derived data so that curve offset are available
-		FHairStrandsDatas In = Group.Strands;
-		HairStrandsBuilder::BuildInternalData(In);
-
-		const uint32 Attributes = In.GetAttributes();
-
-		// Fill in voxel (TODO: make it parallel)
-		for (uint32 CurveIt = 0, CurveCount = In.GetNumCurves(); CurveIt < CurveCount; ++CurveIt)
-		{
-			const uint32 PointOffset = In.StrandsCurves.CurvesOffset[CurveIt];
-			const uint32 PointCount = In.StrandsCurves.CurvesCount[CurveIt];
-
-			uint32 PrevLinearCoord = ~0;
-			for (uint32 PointIndex = 0; PointIndex < PointCount - 1; ++PointIndex)
-			{
-				const uint32 Index0 = PointOffset + PointIndex;
-				const uint32 Index1 = PointOffset + PointIndex + 1;
-				const FVector3f& P0 = In.StrandsPoints.PointsPosition[Index0];
-				const FVector3f& P1 = In.StrandsPoints.PointsPosition[Index1];
-				const float R0 = In.StrandsPoints.PointsRadius[Index0];
-				const float R1 = In.StrandsPoints.PointsRadius[Index1];
-				const FVector3f C0 = HasHairAttribute(Attributes, EHairAttribute::Color) ? FVector3f(In.StrandsPoints.PointsBaseColor[Index0].R, In.StrandsPoints.PointsBaseColor[Index0].G, In.StrandsPoints.PointsBaseColor[Index0].B) : FVector3f(0);
-				const FVector3f C1 = HasHairAttribute(Attributes, EHairAttribute::Color) ? FVector3f(In.StrandsPoints.PointsBaseColor[Index1].R, In.StrandsPoints.PointsBaseColor[Index1].G, In.StrandsPoints.PointsBaseColor[Index1].B) : FVector3f(0);
-				const float Rough0 = HasHairAttribute(Attributes, EHairAttribute::Roughness) ? In.StrandsPoints.PointsRoughness[Index0] : 0;
-				const float Rough1 = HasHairAttribute(Attributes, EHairAttribute::Roughness) ? In.StrandsPoints.PointsRoughness[Index1] : 0;
-				const FVector3f Segment = P1 - P0;
-
-				// This is a coarse/non-conservative voxelization, by ray-marching segment, instead of walking intersected voxels
-				const float Length = Segment.Size();
-				const uint32 StepCount = FMath::CeilToInt(Length / Out.VoxelSize);
-				for (uint32 StepIt = 0; StepIt < StepCount + 1; ++StepIt)
-				{
-					const float T = float(StepIt) / float(StepCount);
-					const FVector3f P = P0 + Segment * T;
-					const FIntVector Coord = ToCoord(P, Out.Resolution, Out.MinBound, Out.VoxelSize);
-					const uint32 LinearCoord = ToLinearCoord(Coord, Out.Resolution);
-					if (LinearCoord != PrevLinearCoord)
-					{
-						FHairGrid::FCurve RCurve;
-						RCurve.CurveIndex = CurveIt;
-						RCurve.GroupIndex = GroupIndex;
-						RCurve.Radius = FMath::Lerp(R0, R1, T);
-						RCurve.BaseColor = FVector3f(FMath::Lerp(C0.X, C1.X, T), FMath::Lerp(C0.Y, C1.Y, T), FMath::Lerp(C0.Z, C1.Z, T));
-						RCurve.Roughness = FMath::Lerp(Rough0, Rough1, T);
-						Out.Voxels[LinearCoord].VoxelCurves.Add(RCurve);
-
-						PrevLinearCoord = LinearCoord;
-					}
+		const uint32 PointOffset = In.StrandsCurves.CurvesOffset[CurveIt];
+		const uint32 PointCount = In.StrandsCurves.CurvesCount[CurveIt];
+
+		uint32 PrevLinearCoord = ~0;
+		for (uint32 PointIndex = 0; PointIndex < PointCount - 1; ++PointIndex)
+		{
+			const uint32 Index0 = PointOffset + PointIndex;
+			const uint32 Index1 = PointOffset + PointIndex + 1;
+			const FVector3f& P0 = In.StrandsPoints.PointsPosition[Index0];
+			const FVector3f& P1 = In.StrandsPoints.PointsPosition[Index1];
+			const float R0 = In.StrandsPoints.PointsRadius[Index0];
+			const float R1 = In.StrandsPoints.PointsRadius[Index1];
+			const FVector3f C0 = HasHairAttribute(Attributes, EHairAttribute::Color) ? FVector3f(In.StrandsPoints.PointsBaseColor[Index0].R, In.StrandsPoints.PointsBaseColor[Index0].G, In.StrandsPoints.PointsBaseColor[Index0].B) : FVector3f(0);
+			const FVector3f C1 = HasHairAttribute(Attributes, EHairAttribute::Color) ? FVector3f(In.StrandsPoints.PointsBaseColor[Index1].R, In.StrandsPoints.PointsBaseColor[Index1].G, In.StrandsPoints.PointsBaseColor[Index1].B) : FVector3f(0);
+			const float Rough0 = HasHairAttribute(Attributes, EHairAttribute::Roughness) ? In.StrandsPoints.PointsRoughness[Index0] : 0;
+			const float Rough1 = HasHairAttribute(Attributes, EHairAttribute::Roughness) ? In.StrandsPoints.PointsRoughness[Index1] : 0;
+			const FVector3f Segment = P1 - P0;
+
+			// This is a coarse/non-conservative voxelization, by ray-marching segment, instead of walking intersected voxels
+			const float Length = Segment.Size();
+			const uint32 StepCount = FMath::CeilToInt(Length / Out.VoxelSize);
+			for (uint32 StepIt = 0; StepIt < StepCount + 1; ++StepIt)
+			{
+				const float T = float(StepIt) / float(StepCount);
+				const FVector3f P = P0 + Segment * T;
+				const FIntVector Coord = ToCoord(P, Out.Resolution, Out.MinBound, Out.VoxelSize);
+				const uint32 LinearCoord = ToLinearCoord(Coord, Out.Resolution);
+				if (LinearCoord != PrevLinearCoord)
+				{
+					FHairGrid::FCurve RCurve;
+					RCurve.CurveIndex = CurveIt;
+					RCurve.GroupIndex = InGroupIndex;
+					RCurve.Radius = FMath::LerpStable(R0, R1, T);
+					RCurve.BaseColor = FVector3f(FMath::LerpStable(C0.X, C1.X, T), FMath::LerpStable(C0.Y, C1.Y, T), FMath::LerpStable(C0.Z, C1.Z, T));
+					RCurve.Roughness = FMath::LerpStable(Rough0, Rough1, T);
+					Out.Voxels[LinearCoord].VoxelCurves.Add(RCurve);
+
+					PrevLinearCoord = LinearCoord;
 				}
 			}
 		}
+	}
+}
 
+static void Voxelize(const FHairDescriptionGroups& InGroups, FHairGrid& Out)
+{
+	AllocateVoxels(InGroups.Bounds.GetBox(), Out);
+	
+	uint32 GroupIndex = 0;
+	for (const FHairDescriptionGroup& Group : InGroups.HairGroups)
+	{
+		VoxelizeCurves(Group.Strands, GroupIndex, Out);
 		++GroupIndex;
 	}
 }
 
+static void Voxelize(const FHairStrandsDatas& InData, FHairGrid& Out)
+{
+	AllocateVoxels(InData.BoundingBox, Out);
+	VoxelizeCurves(InData, 0u /*GroupIndex not used*/, Out);
+}
+
 FORCEINLINE void SearchCell(const FHairStrandsVoxelData& In, const FIntVector& C, FHairStrandsVoxelData::FData& Out)
 {
 	const uint32 I = GroomBuilder_Voxelization::ToIndex(C, In.Resolution);
@@ -2857,6 +2908,168 @@ FORCEINLINE void SearchCell(const FHairStrandsVoxelData& In, const FIntVector& C
 	}
 }
 
+struct FVoxelCoverage
+{
+	uint32 CurveCount = 0;
+	float AvgRadius = 0;
+	float Coverage = 0;
+};
+
+static void ComputeCoverage(const float InGridVoxelSize, const TArray<GroomBuilder_Voxelization::FHairGrid::FVoxel>& InNonEmptyVoxels, TArray<FVoxelCoverage>& OutVoxelCoverages)
+{	
+	const uint32 VoxelCount = InNonEmptyVoxels.Num();
+	OutVoxelCoverages.SetNum(VoxelCount);
+
+	#if 1
+	ParallelFor(VoxelCount, 
+	[
+		InGridVoxelSize,
+		&InNonEmptyVoxels,
+		&OutVoxelCoverages
+	] (uint32 VoxelIt) 
+	#else
+	for (uint32 VoxelIt=0;VoxelIt<VoxelCount;++VoxelIt)
+	#endif
+	{
+		const GroomBuilder_Voxelization::FHairGrid::FVoxel& Voxel = InNonEmptyVoxels[VoxelIt];
+		FVoxelCoverage& Out = OutVoxelCoverages[VoxelIt];
+
+		Out.CurveCount = Voxel.VoxelCurves.Num();
+		for (const GroomBuilder_Voxelization::FHairGrid::FCurve& Curve : Voxel.VoxelCurves)
+		{
+			const float NormalizedRadius = Curve.Radius / InGridVoxelSize;
+			Out.AvgRadius += NormalizedRadius;
+		}
+		Out.AvgRadius /= FMath::Max(1u, Out.CurveCount);
+		Out.Coverage = GetHairCoverage(Out.CurveCount, Out.AvgRadius);
+	});
+}
+
+static float RemoveCurveAndComputeCoverageScaling(
+	const float InGridVoxelSize,
+	TArray<GroomBuilder_Voxelization::FHairGrid::FVoxel>& InNonEmptyVoxels,
+	const TArray<FVoxelCoverage>& InReferenceCoverage, 
+	uint32 InCurveIndexStart,
+	uint32 InCurveIndexEnd)
+{
+	check(InNonEmptyVoxels.Num() == InReferenceCoverage.Num());
+
+	// Certain platform don't support atomic<float>. To support these, we rely on fixed-point math for coverage computation.
+	// The (Radius/Count -> Coverage) LUT has precision up to 6 decimals, so quantizing to 10e6 should be enough
+	const uint64 FixedPointCoverageScale = 1000000u;
+	std::atomic<uint64> uAvgCoverage = 0;
+	//std::atomic<float> AvgCoverage = 0;
+
+	const uint32 VoxelCount = InNonEmptyVoxels.Num();
+	#if 1
+	ParallelFor(VoxelCount, 
+	[
+		InGridVoxelSize,
+		&InNonEmptyVoxels,
+		&InReferenceCoverage,
+		InCurveIndexStart,
+		InCurveIndexEnd,
+		FixedPointCoverageScale,
+		&uAvgCoverage
+	] (uint32 VoxelIt) 
+	#else
+	for (uint32 VoxelIt=0;VoxelIt<VoxelCount;++VoxelIt)
+	#endif
+	{
+		GroomBuilder_Voxelization::FHairGrid::FVoxel& Voxel = InNonEmptyVoxels[VoxelIt];
+
+		uint32 CurveCount = 0;
+		float AvgRadius = 0;
+
+		uint32 RemoveCount = 0;
+		for (GroomBuilder_Voxelization::FHairGrid::FCurve& Curve : Voxel.VoxelCurves)
+		{
+			if (InCurveIndexStart <= Curve.CurveIndex  && Curve.CurveIndex <= InCurveIndexEnd)
+			{
+				RemoveCount++;
+			}
+			else
+			{
+				CurveCount++;
+				const float NormalizedRadius = Curve.Radius / InGridVoxelSize;
+				AvgRadius += NormalizedRadius;
+			}
+		}
+
+		AvgRadius /= FMath::Max(CurveCount, 1u);
+
+		const float VoxelCoverage = GetHairCoverage(CurveCount, AvgRadius);
+		const float CoverageScale = VoxelCoverage > 0 ? InReferenceCoverage[VoxelIt].Coverage / VoxelCoverage : 1.f;
+		//AvgCoverage.fetch_add(CoverageScale, std::memory_order_relaxed);
+		//AvgCoverage += CoverageScale;
+
+		const uint64 uCoverageScale = CoverageScale * FixedPointCoverageScale;
+		uAvgCoverage  += uCoverageScale;
+
+		if (RemoveCount > 0)
+		{
+			Voxel.VoxelCurves.RemoveAll([InCurveIndexStart, InCurveIndexEnd](GroomBuilder_Voxelization::FHairGrid::FCurve& Curve) 
+			{ 
+				return InCurveIndexStart <= Curve.CurveIndex && Curve.CurveIndex <= InCurveIndexEnd;
+			});
+		}
+	});
+
+	const float AvgCoverage = uAvgCoverage / float(FixedPointCoverageScale);
+	return AvgCoverage / float(FMath::Max(VoxelCount, 1u));
+}
+
+void ComputeHairCoverageScale(const FHairStrandsDatas& In, TArray<FCoverageScale>& OutCoverageScales)
+{
+	// 1. Voxelize curves
+	GroomBuilder_Voxelization::FHairGrid Grid;
+	GroomBuilder_Voxelization::Voxelize(In, Grid);
+
+	// 2. Collect non-empty voxels
+	TArray<GroomBuilder_Voxelization::FHairGrid::FVoxel> NonEmptyVoxels;
+	NonEmptyVoxels.Reserve(Grid.Voxels.Num());
+	for (const GroomBuilder_Voxelization::FHairGrid::FVoxel& Voxel : Grid.Voxels)
+	{
+		if (Voxel.VoxelCurves.Num() > 0)
+		{
+			NonEmptyVoxels.Add(Voxel);
+		}
+	}
+
+	// 3. Compute reference coverage
+	TArray<FVoxelCoverage> ReferenceCoverage;
+	ComputeCoverage(Grid.VoxelSize, NonEmptyVoxels, ReferenceCoverage);
+
+	// 4. Compute coverage scale for every X curves removed
+	//    Compute scale for each decimated percent (i.e. 100 buckets)
+	const uint32 CurveCount = In.GetNumCurves();
+	uint32 BucketCount = FMath::Min(100u, CurveCount);
+	const uint32 BucketSize = FMath::DivideAndRoundUp(CurveCount, BucketCount);
+	BucketCount = FMath::DivideAndRoundUp(CurveCount, BucketSize);
+
+	check(BucketSize > 0);
+	check(BucketCount > 0);
+
+	FCoverageScale DefaultValue;
+	DefaultValue.ScreenSize = 1;
+	DefaultValue.CoverageScale = 1;
+	OutCoverageScales.Init(DefaultValue, BucketCount);
+	if (BucketCount > 1u)
+	{
+		float MaxCoverageScale = 1.f;
+		for (int32 BucketIt = BucketCount-1; BucketIt >= 0; --BucketIt)
+		{
+			const uint32 CurveIndexStart =  BucketIt * BucketSize;
+			const uint32 CurveIndexEnd   = FMath::Min(CurveCount-1u, (BucketIt+1) * BucketSize - 1u);
+			const uint32 CurveBucketCount = (CurveIndexEnd-CurveIndexStart)+1;
+			const float CoverageScale = RemoveCurveAndComputeCoverageScaling(Grid.VoxelSize, NonEmptyVoxels, ReferenceCoverage, CurveIndexStart, CurveIndexEnd);
+			MaxCoverageScale = FMath::Max(CoverageScale, MaxCoverageScale);
+			OutCoverageScales[BucketIt].CoverageScale = MaxCoverageScale;
+			OutCoverageScales[BucketIt].ScreenSize = 1.f; // TODO: compute the screen size / pixel area at which this decimation rate should be applied
+		}
+	}
+}
+
 } // namespace GroomBuilder_Voxelization
 
 FHairStrandsVoxelData::FData FHairStrandsVoxelData::GetData(const FVector3f& P) const
diff --git a/Engine/Plugins/Runtime/HairStrands/Source/HairStrandsCore/Private/GroomComponent.cpp b/Engine/Plugins/Runtime/HairStrands/Source/HairStrandsCore/Private/GroomComponent.cpp
--- a/Engine/Plugins/Runtime/HairStrands/Source/HairStrandsCore/Private/GroomComponent.cpp
+++ b/Engine/Plugins/Runtime/HairStrands/Source/HairStrandsCore/Private/GroomComponent.cpp
@@ -2743,7 +2743,8 @@ void UGroomComponent::InitResources(bool bIsBindingReloading)
 			HairGroupInstance->HairGroupPublicData->SetLODScreenSizes(CPULODScreenSize);
 			HairGroupInstance->HairGroupPublicData->SetLODVisibilities(LODVisibility);
 			HairGroupInstance->HairGroupPublicData->SetLODGeometryTypes(LODGeometryTypes);
-			HairGroupInstance->HairGroupPublicData->bAutoLOD = GroomAsset->GetHairGroupsLOD()[GroupIt].LODType == EGroomLODType::Auto;
+			HairGroupInstance->HairGroupPublicData->bAutoLOD = GroomAsset->GetLODMode() == EGroomLODMode::Auto;
+			HairGroupInstance->HairGroupPublicData->AutoLODBias = GroomAsset->GetAutoLODBias();
 		}
 
 		FHairResourceName ResourceName(GetFName(), GroupIt);
diff --git a/Engine/Plugins/Runtime/HairStrands/Source/HairStrandsCore/Private/GroomDebug.cpp b/Engine/Plugins/Runtime/HairStrands/Source/HairStrandsCore/Private/GroomDebug.cpp
--- a/Engine/Plugins/Runtime/HairStrands/Source/HairStrandsCore/Private/GroomDebug.cpp
+++ b/Engine/Plugins/Runtime/HairStrands/Source/HairStrandsCore/Private/GroomDebug.cpp
@@ -738,7 +738,7 @@ public:
 
 IMPLEMENT_GLOBAL_SHADER(FHairDebugPrintInstanceCS, "/Engine/Private/HairStrands/HairStrandsDebug.usf", "MainCS", SF_Compute);
 
-bool IsHairStrandsForceAutoLODEnabled();
+bool UseHairStrandsForceAutoLOD();
 const TCHAR* GetHairAttributeText(EHairAttribute In, uint32 InFlags);
 uint32 GetHairAttributeIndex(EHairAttribute In);
 EGroomCacheType GetHairInstanceCacheType(const FHairGroupInstance* Instance);
@@ -857,7 +857,7 @@ static void AddHairDebugPrintInstancePass(
 				D.Data3.Z |= FFloat16(Instance->HairGroupPublicData->VFInput.Strands.Common.Length).Encoded;
 				D.Data3.Z |= FFloat16(Instance->HairGroupPublicData->VFInput.Strands.Common.LengthScale).Encoded << 16u;
 				D.Data3.W |= FFloat16(Instance->HairGroupPublicData->GetActiveStrandsCoverageScale()).Encoded;
-				D.Data3.W |= Instance->HairGroupPublicData->bAutoLOD || IsHairStrandsForceAutoLODEnabled() ? (1u << 16u) : 0;
+				D.Data3.W |= (Instance->HairGroupPublicData->bAutoLOD || UseHairStrandsForceAutoLOD()) ? (1u << 16u) : 0;
 
 				#if RHI_RAYTRACING
 				bHasRaytracing = Instance->Strands.RenRaytracingResource != nullptr;
diff --git a/Engine/Plugins/Runtime/HairStrands/Source/HairStrandsCore/Private/GroomManager.cpp b/Engine/Plugins/Runtime/HairStrands/Source/HairStrandsCore/Private/GroomManager.cpp
--- a/Engine/Plugins/Runtime/HairStrands/Source/HairStrandsCore/Private/GroomManager.cpp
+++ b/Engine/Plugins/Runtime/HairStrands/Source/HairStrandsCore/Private/GroomManager.cpp
@@ -55,11 +55,19 @@ static int32 GHairStrands_Streaming_StreamOutThreshold = 2;
 static FAutoConsoleVariableRef CVarHairStrands_Streaming_StreamOutThreshold(TEXT("r.HairStrands.Streaming.StreamOutThreshold"), GHairStrands_Streaming_StreamOutThreshold, TEXT("Threshold used for streaming out data. In curve page. Default:2."));
 
 static int32 GHairStrands_AutoLOD_Force = 0;
-static float GHairStrands_AutoLOD_Scale = 1.f;
-static float GHairStrands_AutoLOD_Bias = 0.f;
-static FAutoConsoleVariableRef CVarHairStrands_AutoLOD_Force(TEXT("r.HairStrands.AutoLOD.Force"), GHairStrands_AutoLOD_Force, TEXT("Force all groom to use Auto LOD (experimental)."), ECVF_RenderThreadSafe);
-static FAutoConsoleVariableRef CVarHairStrands_AutoLOD_Scale(TEXT("r.HairStrands.AutoLOD.Scale"), GHairStrands_AutoLOD_Scale, TEXT("Hair strands Auto LOD rate at which curves get decimated based on screen coverage."), ECVF_RenderThreadSafe);
-static FAutoConsoleVariableRef CVarHairStrands_AutoLOD_Bias(TEXT("r.HairStrands.AutoLOD.Bias"), GHairStrands_AutoLOD_Bias, TEXT("Hair strands Auto LOD screen size bias at which curves get decimated."), ECVF_RenderThreadSafe);
+static FAutoConsoleVariableRef CVarHairStrands_AutoLOD_Force(TEXT("r.HairStrands.AutoLOD.Force"), GHairStrands_AutoLOD_Force, TEXT("Force Auto LOD on all grooms. Used for debugging purpose."), ECVF_RenderThreadSafe);
+static float GHairStrands_AutoLOD_Bias = 0.1f;
+static FAutoConsoleVariableRef CVarHairStrands_AutoLOD_Bias(TEXT("r.HairStrands.AutoLOD.Bias"), GHairStrands_AutoLOD_Bias, TEXT("Global bias for Auto LOD on all grooms. Used for debugging purpose."), ECVF_RenderThreadSafe);
+
+bool UseHairStrandsForceAutoLOD()
+{
+	return GHairStrands_AutoLOD_Force > 0;
+}
+
+float GetHairStrandsAutoLODBias()
+{
+	return FMath::Clamp(GHairStrands_AutoLOD_Bias, 0.f, 1.f);
+}
 
 ///////////////////////////////////////////////////////////////////////////////////////////////////
 // Foward declaration
@@ -108,11 +116,6 @@ EHairBufferSwapType GetHairSwapBufferType()
 	return EHairBufferSwapType::EndOfFrame;
 }
 
-bool IsHairStrandsForceAutoLODEnabled()
-{
-	return GHairStrands_AutoLOD_Force > 0;
-}
-
 uint32 GetStreamingCurvePage()
 {
 	return FMath::Clamp(uint32(GHairStrands_Streaming_CurvePage), 1u, 32768u);
@@ -2361,17 +2364,17 @@ static FVector2f ComputeProjectedScreenPos(const FVector& InWorldPos, const FSce
 	return FVector2f(ScreenPos.X, ScreenPos.Y);
 }
 
-static float ComputeActiveCurveCoverageScale(uint32 InAvailableCurveCount, uint32 InRestCurveCount)
+static float ComputeActiveCurveCoverageScale(const FHairStrandsBulkData& InData, uint32 InAvailableCurveCount)
 {
-	// Compensate lost in curve by a coverage scale increase
-	const float CurveRatio = InAvailableCurveCount / float(InRestCurveCount);
-	return 1.f/FMath::Max(CurveRatio, 0.01f);
+	const uint32 RestCurveCount = InData.Header.CurveCount;
+	const float CurveRatio = InAvailableCurveCount / float(RestCurveCount);
+	return InData.GetCoverageScale(CurveRatio);
 }
 
-static uint32 ComputeActiveCurveCount(float InScreenSize, uint32 InCurveCount, uint32 InClusterCount)
+static uint32 ComputeActiveCurveCount(float InScreenSize, float InAutoLODBias, uint32 InCurveCount, uint32 InClusterCount)
 {
-	const float Power = FMath::Max(0.1f, GHairStrands_AutoLOD_Scale);
-	const float ScreenSizeBias = FMath::Clamp(GHairStrands_AutoLOD_Bias, 0.f, 1.f);
+	const float Power = 1.f;  // This could be exposed per asset
+	const float ScreenSizeBias = FMath::Clamp(FMath::Max(InAutoLODBias, GetHairStrandsAutoLODBias()), 0.f, 1.f);
 	uint32 OutCurveCount = InCurveCount * FMath::Pow(FMath::Clamp(InScreenSize + ScreenSizeBias, 0.f, 1.0f), Power);
 	// Ensure there is at least 1 curve per cluster
 	OutCurveCount = FMath::Max(InClusterCount, OutCurveCount);
@@ -2473,9 +2476,9 @@ static FHairLOD ComputeHairLODIndex(const FHairGroupInstance* Instance, const TA
 	if (Instance->Strands.ClusterResource && Instance->Strands.RestResource)
 	{
 		uint32 EffectiveCurveCount = 0;
-		if (bNeedAutoLOD && (Instance->HairGroupPublicData->bAutoLOD || IsHairStrandsForceAutoLODEnabled()))
+		if (bNeedAutoLOD && (Instance->HairGroupPublicData->bAutoLOD || UseHairStrandsForceAutoLOD()))
 		{
-			EffectiveCurveCount = ComputeActiveCurveCount(Out.ContinuousLODScreenSize, Instance->HairGroupPublicData->RestCurveCount, Instance->HairGroupPublicData->ClusterCount);
+			EffectiveCurveCount = ComputeActiveCurveCount(Out.ContinuousLODScreenSize, Instance->HairGroupPublicData->AutoLODBias, Instance->HairGroupPublicData->RestCurveCount, Instance->HairGroupPublicData->ClusterCount);
 			Out.HairLODIndex = MinLODIndexWithStrands;
 		}
 		else
@@ -2486,7 +2489,7 @@ static FHairLOD ComputeHairLODIndex(const FHairGroupInstance* Instance, const TA
 
 		Out.ContinuousLODCurveCount = EffectiveCurveCount;
 		Out.ContinuousLODPointCount = EffectiveCurveCount > 0 ? Instance->Strands.GetData().Header.CurveToPointCount[EffectiveCurveCount - 1] : 0;
-		Out.ContinuousLODCoverageScale = ComputeActiveCurveCoverageScale(EffectiveCurveCount, Instance->HairGroupPublicData->RestCurveCount);
+		Out.ContinuousLODCoverageScale = ComputeActiveCurveCoverageScale(Instance->Strands.GetData(), EffectiveCurveCount);
 	}
 
 	return Out;
@@ -2707,7 +2710,7 @@ static bool SelectValidLOD(
 		const uint32 EffectiveCurveCount = FMath::Min(ResourceStatus.AvailableCurveCount, Instance->HairGroupPublicData->ContinuousLODCurveCount);
 		Instance->HairGroupPublicData->ContinuousLODCurveCount = EffectiveCurveCount;
 		Instance->HairGroupPublicData->ContinuousLODPointCount = EffectiveCurveCount > 0 ? Instance->Strands.GetData().Header.CurveToPointCount[EffectiveCurveCount - 1] : 0;
-		Instance->HairGroupPublicData->ContinuousLODCoverageScale = ComputeActiveCurveCoverageScale(EffectiveCurveCount, Instance->HairGroupPublicData->RestCurveCount);
+		Instance->HairGroupPublicData->ContinuousLODCoverageScale = ComputeActiveCurveCoverageScale(Instance->Strands.GetData(), EffectiveCurveCount);
 		check(Instance->HairGroupPublicData->ContinuousLODPointCount <= Instance->HairGroupPublicData->RestPointCount);
 	}
 	return true;
diff --git a/Engine/Plugins/Runtime/HairStrands/Source/HairStrandsCore/Private/HairStrandsDatas.cpp b/Engine/Plugins/Runtime/HairStrands/Source/HairStrandsCore/Private/HairStrandsDatas.cpp
--- a/Engine/Plugins/Runtime/HairStrands/Source/HairStrandsCore/Private/HairStrandsDatas.cpp
+++ b/Engine/Plugins/Runtime/HairStrands/Source/HairStrandsCore/Private/HairStrandsDatas.cpp
@@ -549,6 +549,7 @@ void FHairStrandsBulkData::SerializeHeader(FArchive& Ar, UObject* Owner)
 	Ar << Header.ImportedAttributeFlags;
 
 	Ar << Header.CurveToPointCount;
+	Ar << Header.CoverageScales;
 
 	Ar << Header.Strides.PositionStride;
 	Ar << Header.Strides.CurveStride;
@@ -673,6 +674,21 @@ void FHairStrandsBulkData::ResetLoadedSize()
 	Data.Curves.LoadedSize			= 0;
 }
 
+float FHairStrandsBulkData::GetCoverageScale(float InCurvePercentage /*[0..1]*/) const
+{ 
+	float Out = 1.f;
+	if (Header.CoverageScales.Num() > 0)
+	{
+		const uint32 ScaleCount = Header.CoverageScales.Num();
+		const float  fScaleIt = FMath::Clamp(InCurvePercentage * ScaleCount, 0u, ScaleCount-1u);
+		const uint32 ScaleIt0 = FMath::Floor(fScaleIt);
+		const uint32 ScaleIt1 = FMath::Min(ScaleIt0+1, ScaleCount-1);
+		const float S = fScaleIt - ScaleIt0;
+		Out = FMath::LerpStable(Header.CoverageScales[ScaleIt0], Header.CoverageScales[ScaleIt1], S);
+	}
+	return Out;
+}
+
 /////////////////////////////////////////////////////////////////////////////////////////
 // Interpolation bulk data
 
diff --git a/Engine/Plugins/Runtime/HairStrands/Source/HairStrandsCore/Public/GroomAsset.h b/Engine/Plugins/Runtime/HairStrands/Source/HairStrandsCore/Public/GroomAsset.h
--- a/Engine/Plugins/Runtime/HairStrands/Source/HairStrandsCore/Public/GroomAsset.h
+++ b/Engine/Plugins/Runtime/HairStrands/Source/HairStrandsCore/Public/GroomAsset.h
@@ -449,6 +449,13 @@ class HAIRSTRANDSCORE_API UGroomAsset : public UObject, public IInterface_AssetU
 	DECLARE_MULTICAST_DELEGATE(FOnGroomAsyncLoadFinished);
 #endif
 
+private:
+	UPROPERTY(EditAnywhere, Category = "HairLOD", meta = (DisplayName = "LOD Mode", ToolTip = "Define how LOD adapts curves & points for strands geometry. Auto: adapts the curve count based on screen coverage. Manual: use the discrete LOD created for each groups"))
+	EGroomLODMode LODMode = EGroomLODMode::Default;
+
+	UPROPERTY(EditAnywhere, Category = "HairLOD", meta = (DisplayName = "Auto LOD Bias", ToolTip = "When Auto LOD is selected, decrease the screen size at which curves reduction will occur.", ClampMin = "0", ClampMax = "1", UIMin = "0", UIMax = "1.0"))
+	float AutoLODBias = 0;
+
 public:
 	UE_DEPRECATED(5.3, "Please do not access this member directly; use UGroomAsset accessor.")
 	UPROPERTY(EditAnywhere, EditFixedSize, EditFixedSize, Category = "HairInfo", meta = (DisplayName = "Group"))
@@ -616,6 +623,11 @@ public:
 	const TArray<FHairGroupResources>& GetHairGroupsResources() const;
 	TArray<FHairGroupResources>& GetHairGroupsResources();
 
+	static FName GetLODModeMemberName();
+	EGroomLODMode GetLODMode() const;
+
+	static FName GetAutoLODBiasMemberName();
+	float GetAutoLODBias() const;
 public:
 
 	//~ Begin UObject Interface.
diff --git a/Engine/Plugins/Runtime/HairStrands/Source/HairStrandsCore/Public/GroomAssetInterpolation.h b/Engine/Plugins/Runtime/HairStrands/Source/HairStrandsCore/Public/GroomAssetInterpolation.h
--- a/Engine/Plugins/Runtime/HairStrands/Source/HairStrandsCore/Public/GroomAssetInterpolation.h
+++ b/Engine/Plugins/Runtime/HairStrands/Source/HairStrandsCore/Public/GroomAssetInterpolation.h
@@ -59,10 +59,11 @@ enum class EGroomGuideType : uint8
 };
 
 UENUM(BlueprintType)
-enum class EGroomLODType : uint8
+enum class EGroomLODMode : uint8
 {
+	Default		UMETA(DisplayName = "Default", ToolTip = "Hair strands curves & points adapt based on project settings LOD mode. (See 'Use Auto LOD' in project settings). "),
 	Manual 		UMETA(DisplayName = "Manual", ToolTip = "Hair strands curves & points adapt based on LOD settings"),
-	Auto		UMETA(DisplayName = "Auto (Experimental)", ToolTip = "Hair strands curves & points automatically adapt based on screen coverage. LOD settings are ignored.")
+	Auto		UMETA(DisplayName = "Auto", ToolTip = "Hair strands curves & points automatically adapt based on screen coverage. LOD settings are ignored.")
 };
 
 USTRUCT(BlueprintType)
@@ -220,9 +221,6 @@ struct HAIRSTRANDSCORE_API FHairGroupsLOD
 {
 	GENERATED_USTRUCT_BODY()
 
-	UPROPERTY(EditAnywhere, Category = "InterpolationSettings", meta = (ToolTip = "Define how LOD adapts curves & points for a given group."))
-	EGroomLODType LODType = EGroomLODType::Manual;
-
 	/** LODs  */
 	UPROPERTY(EditAnywhere, Category = "DecimationSettings", meta = (ToolTip = "Reduce the number of hair strands in a uniform manner", ClampMin = "0", UIMin = "0", UIMax = "1.0"))
 	TArray<FHairLODSettings> LODs;
diff --git a/Engine/Plugins/Runtime/HairStrands/Source/HairStrandsCore/Public/HairStrandsDatas.h b/Engine/Plugins/Runtime/HairStrands/Source/HairStrandsCore/Public/HairStrandsDatas.h
--- a/Engine/Plugins/Runtime/HairStrands/Source/HairStrandsCore/Public/HairStrandsDatas.h
+++ b/Engine/Plugins/Runtime/HairStrands/Source/HairStrandsCore/Public/HairStrandsDatas.h
@@ -622,6 +622,7 @@ struct HAIRSTRANDSCORE_API FHairStrandsBulkData : FHairStrandsBulkCommon
 	bool IsValid() const { return Header.CurveCount > 0 && Header.PointCount > 0; }
 	void Reset();
 	virtual void ResetLoadedSize() override;
+	float GetCoverageScale(float InCurvePercentage /*[0..1]*/) const;
 
 	uint32 GetNumCurves() const { return Header.CurveCount;  };
 	uint32 GetNumPoints() const { return Header.PointCount; };
@@ -657,6 +658,9 @@ struct HAIRSTRANDSCORE_API FHairStrandsBulkData : FHairStrandsBulkCommon
 		// Map 'curve' count to 'point' count (used for CLOD)
 		TArray<uint32> CurveToPointCount;
 
+		// Coverage scale
+		TArray<float> CoverageScales;
+
 		// Data transcoding parameters
 		struct FTranscoding
 		{
diff --git a/Engine/Plugins/Runtime/HairStrands/Source/HairStrandsEditor/Private/GroomAssetDetails.cpp b/Engine/Plugins/Runtime/HairStrands/Source/HairStrandsEditor/Private/GroomAssetDetails.cpp
--- a/Engine/Plugins/Runtime/HairStrands/Source/HairStrandsEditor/Private/GroomAssetDetails.cpp
+++ b/Engine/Plugins/Runtime/HairStrands/Source/HairStrandsEditor/Private/GroomAssetDetails.cpp
@@ -483,6 +483,15 @@ void FGroomRenderingDetails::AddNewGroupButton(IDetailCategoryBuilder& FilesCate
 	];
 }
 
+static void AddLODModeProperties(IDetailLayoutBuilder& DetailLayout, IDetailCategoryBuilder& FilesCategory)
+{
+	TSharedRef<IPropertyHandle> LODModeProperty = DetailLayout.GetProperty(UGroomAsset::GetLODModeMemberName(), UGroomAsset::StaticClass());
+	TSharedRef<IPropertyHandle> LODBiasProperty = DetailLayout.GetProperty(UGroomAsset::GetAutoLODBiasMemberName(), UGroomAsset::StaticClass());
+
+	FilesCategory.AddProperty(LODModeProperty);
+	FilesCategory.AddProperty(LODBiasProperty);
+}
+
 void FGroomRenderingDetails::CustomizeStrandsGroupProperties(IDetailLayoutBuilder& DetailLayout, IDetailCategoryBuilder& FilesCategory)
 {
 	switch (PanelType)
@@ -503,6 +512,8 @@ void FGroomRenderingDetails::CustomizeStrandsGroupProperties(IDetailLayoutBuilde
 		DetailLayout.HideProperty(DetailLayout.GetProperty(UGroomAsset::GetMinLODMemberName(), UGroomAsset::StaticClass()));
 		DetailLayout.HideProperty(DetailLayout.GetProperty(UGroomAsset::GetDisableBelowMinLodStrippingMemberName(), UGroomAsset::StaticClass()));
 		DetailLayout.HideProperty(DetailLayout.GetProperty(UGroomAsset::GetRiggedSkeletalMeshMemberName(), UGroomAsset::StaticClass()));
+		DetailLayout.HideProperty(DetailLayout.GetProperty(UGroomAsset::GetLODModeMemberName(), UGroomAsset::StaticClass()));
+		DetailLayout.HideProperty(DetailLayout.GetProperty(UGroomAsset::GetAutoLODBiasMemberName(), UGroomAsset::StaticClass()));
 	}
 	break;
 	case EMaterialPanelType::Meshes:
@@ -521,6 +532,8 @@ void FGroomRenderingDetails::CustomizeStrandsGroupProperties(IDetailLayoutBuilde
 		DetailLayout.HideProperty(DetailLayout.GetProperty(UGroomAsset::GetMinLODMemberName(), UGroomAsset::StaticClass()));
 		DetailLayout.HideProperty(DetailLayout.GetProperty(UGroomAsset::GetDisableBelowMinLodStrippingMemberName(), UGroomAsset::StaticClass()));
 		DetailLayout.HideProperty(DetailLayout.GetProperty(UGroomAsset::GetRiggedSkeletalMeshMemberName(), UGroomAsset::StaticClass()));
+		DetailLayout.HideProperty(DetailLayout.GetProperty(UGroomAsset::GetLODModeMemberName(), UGroomAsset::StaticClass()));
+		DetailLayout.HideProperty(DetailLayout.GetProperty(UGroomAsset::GetAutoLODBiasMemberName(), UGroomAsset::StaticClass()));
 	}
 	break;
 	case EMaterialPanelType::Strands:
@@ -539,6 +552,8 @@ void FGroomRenderingDetails::CustomizeStrandsGroupProperties(IDetailLayoutBuilde
 		DetailLayout.HideProperty(DetailLayout.GetProperty(UGroomAsset::GetMinLODMemberName(), UGroomAsset::StaticClass()));
 		DetailLayout.HideProperty(DetailLayout.GetProperty(UGroomAsset::GetDisableBelowMinLodStrippingMemberName(), UGroomAsset::StaticClass()));
 		DetailLayout.HideProperty(DetailLayout.GetProperty(UGroomAsset::GetRiggedSkeletalMeshMemberName(), UGroomAsset::StaticClass()));
+		DetailLayout.HideProperty(DetailLayout.GetProperty(UGroomAsset::GetLODModeMemberName(), UGroomAsset::StaticClass()));
+		DetailLayout.HideProperty(DetailLayout.GetProperty(UGroomAsset::GetAutoLODBiasMemberName(), UGroomAsset::StaticClass()));
 	}
 	break;
 	case EMaterialPanelType::Physics:
@@ -557,6 +572,8 @@ void FGroomRenderingDetails::CustomizeStrandsGroupProperties(IDetailLayoutBuilde
 		DetailLayout.HideProperty(DetailLayout.GetProperty(UGroomAsset::GetMinLODMemberName(), UGroomAsset::StaticClass()));
 		DetailLayout.HideProperty(DetailLayout.GetProperty(UGroomAsset::GetDisableBelowMinLodStrippingMemberName(), UGroomAsset::StaticClass()));
 		DetailLayout.HideProperty(DetailLayout.GetProperty(UGroomAsset::GetRiggedSkeletalMeshMemberName(), UGroomAsset::StaticClass()));
+		DetailLayout.HideProperty(DetailLayout.GetProperty(UGroomAsset::GetLODModeMemberName(), UGroomAsset::StaticClass()));
+		DetailLayout.HideProperty(DetailLayout.GetProperty(UGroomAsset::GetAutoLODBiasMemberName(), UGroomAsset::StaticClass()));
 	}
 	break;
 	case EMaterialPanelType::Interpolation:
@@ -575,6 +592,8 @@ void FGroomRenderingDetails::CustomizeStrandsGroupProperties(IDetailLayoutBuilde
 		DetailLayout.HideProperty(DetailLayout.GetProperty(UGroomAsset::GetMinLODMemberName(), UGroomAsset::StaticClass()));
 		DetailLayout.HideProperty(DetailLayout.GetProperty(UGroomAsset::GetDisableBelowMinLodStrippingMemberName(), UGroomAsset::StaticClass()));
 //		DetailLayout.HideProperty(DetailLayout.GetProperty(UGroomAsset::GetRiggedSkeletalMeshMemberName(), UGroomAsset::StaticClass()));
+		DetailLayout.HideProperty(DetailLayout.GetProperty(UGroomAsset::GetLODModeMemberName(), UGroomAsset::StaticClass()));
+		DetailLayout.HideProperty(DetailLayout.GetProperty(UGroomAsset::GetAutoLODBiasMemberName(), UGroomAsset::StaticClass()));
 	}
 	break;
 	case EMaterialPanelType::LODs:
@@ -593,6 +612,8 @@ void FGroomRenderingDetails::CustomizeStrandsGroupProperties(IDetailLayoutBuilde
 //		DetailLayout.HideProperty(DetailLayout.GetProperty(UGroomAsset::GetMinLODMemberName(), UGroomAsset::StaticClass()));
 //		DetailLayout.HideProperty(DetailLayout.GetProperty(UGroomAsset::GetDisableBelowMinLodStrippingMemberName(), UGroomAsset::StaticClass()));
 		DetailLayout.HideProperty(DetailLayout.GetProperty(UGroomAsset::GetRiggedSkeletalMeshMemberName(), UGroomAsset::StaticClass()));
+//		DetailLayout.HideProperty(DetailLayout.GetProperty(UGroomAsset::GetLODModeMemberName(), UGroomAsset::StaticClass()));
+//		DetailLayout.HideProperty(DetailLayout.GetProperty(UGroomAsset::GetAutoLODBiasMemberName(), UGroomAsset::StaticClass()));
 	}
 	break;
 	}
@@ -664,6 +685,7 @@ void FGroomRenderingDetails::CustomizeStrandsGroupProperties(IDetailLayoutBuilde
 		case EMaterialPanelType::LODs:
 		{
 			TSharedRef<IPropertyHandle> Property = DetailLayout.GetProperty(UGroomAsset::GetHairGroupsLODMemberName(), UGroomAsset::StaticClass());
+			AddLODModeProperties(DetailLayout, FilesCategory);
 			if (Property->IsValidHandle())
 			{
 				TSharedRef<FDetailArrayBuilder> PropertyBuilder = MakeShareable(new FDetailArrayBuilder(Property, false, false, false));
@@ -1102,7 +1124,6 @@ IDetailPropertyRow& FGroomRenderingDetails::AddPropertyWithCustomReset(TSharedPt
 
 void FGroomRenderingDetails::ExpandStructForLOD(TSharedRef<IPropertyHandle>& PropertyHandle, IDetailChildrenBuilder& ChildrenBuilder, int32 GroupIndex, int32 LODIndex, bool bOverrideReset)
 {
-
 	uint32 ChildrenCount = 0;
 	PropertyHandle->GetNumChildren(ChildrenCount);
 	for (uint32 ChildIt = 0; ChildIt < ChildrenCount; ++ChildIt)
@@ -1110,29 +1131,35 @@ void FGroomRenderingDetails::ExpandStructForLOD(TSharedRef<IPropertyHandle>& Pro
 		TSharedPtr<IPropertyHandle> ChildHandle = PropertyHandle->GetChildHandle(ChildIt);
 		const FName ChildPropertyName = ChildHandle->GetProperty()->GetFName();
 
-		// If the geometry type is not strands, then bypass the display of the strands related property
-		if (GroomAsset->GetGeometryType(GroupIndex, LODIndex) != EGroomGeometryType::Strands && 
-			(ChildPropertyName == GET_MEMBER_NAME_CHECKED(FHairLODSettings, CurveDecimation)  ||
+		const bool bIsManualLODStrandsProperties = 
+			(ChildPropertyName == GET_MEMBER_NAME_CHECKED(FHairLODSettings, CurveDecimation) ||
 			 ChildPropertyName == GET_MEMBER_NAME_CHECKED(FHairLODSettings, VertexDecimation) ||
 			 ChildPropertyName == GET_MEMBER_NAME_CHECKED(FHairLODSettings, AngularThreshold) ||
-			 ChildPropertyName == GET_MEMBER_NAME_CHECKED(FHairLODSettings, ThicknessScale)))
+			 ChildPropertyName == GET_MEMBER_NAME_CHECKED(FHairLODSettings, ThicknessScale));
+		const bool bIsScreenSize = ChildPropertyName == GET_MEMBER_NAME_CHECKED(FHairLODSettings, ScreenSize);
+
+		// If the geometry type is not strands, then bypass the display of the strands related property
+		if (GroomAsset->GetGeometryType(GroupIndex, LODIndex) != EGroomGeometryType::Strands && bIsManualLODStrandsProperties)
 		{
 			continue;
 		}
 
-		const bool bAutoLOD = GroomAsset->GetHairGroupsLOD()[GroupIndex].LODType == EGroomLODType::Auto;
+		// If using AutoLOD, then the following property are not editable, since they are dedicated to Manual LOD mode
+		const bool bAutoLOD = GroomAsset->GetLODMode() == EGroomLODMode::Auto;
+		const bool bIsVisible = bAutoLOD ? !bIsManualLODStrandsProperties && !bIsScreenSize : true;
+
 		if (bOverrideReset)
 		{
 			FIsResetToDefaultVisible IsResetVisible = FIsResetToDefaultVisible::CreateSP(this, &FGroomRenderingDetails::ShouldResetToDefault, GroupIndex, LODIndex);
 			FResetToDefaultHandler ResetHandler = FResetToDefaultHandler::CreateSP(this, &FGroomRenderingDetails::ResetToDefault, GroupIndex, LODIndex);
 			FResetToDefaultOverride ResetOverride = FResetToDefaultOverride::Create(IsResetVisible, ResetHandler);
 			IDetailPropertyRow& Row = ChildrenBuilder.AddProperty(ChildHandle.ToSharedRef()).OverrideResetToDefault(ResetOverride);
-			Row.IsEnabled(!bAutoLOD);
+			Row.IsEnabled(bIsVisible);
 		}
 		else
 		{
 			IDetailPropertyRow& Row = ChildrenBuilder.AddProperty(ChildHandle.ToSharedRef());
-			Row.IsEnabled(!bAutoLOD);
+			Row.IsEnabled(bIsVisible);
 		}
 	}
 }
@@ -1764,12 +1791,6 @@ void FGroomRenderingDetails::OnGenerateElementForHairGroup(TSharedRef<IPropertyH
 		break;
 		case EMaterialPanelType::LODs:
 		{
-			// Reload layout when the LOD type change
-			if (PropertyName == GET_MEMBER_NAME_CHECKED(FHairGroupsLOD, LODType))
-			{
-				ChildHandle->SetOnPropertyValueChanged(FSimpleDelegate::CreateSP(this, &FGroomRenderingDetails::OnGroomLODTypeChanged));
-			}
-
 			if (PropertyName == GET_MEMBER_NAME_CHECKED(FHairGroupsLOD, LODs))
 			{
 				// Add a custom builder for each LOD arrays within each group. This way we can customize this 'nested/inner' array
diff --git a/Engine/Plugins/Runtime/HairStrands/Source/HairStrandsEditor/Private/GroomCustomAssetEditorToolkit.cpp b/Engine/Plugins/Runtime/HairStrands/Source/HairStrandsEditor/Private/GroomCustomAssetEditorToolkit.cpp
--- a/Engine/Plugins/Runtime/HairStrands/Source/HairStrandsEditor/Private/GroomCustomAssetEditorToolkit.cpp
+++ b/Engine/Plugins/Runtime/HairStrands/Source/HairStrandsEditor/Private/GroomCustomAssetEditorToolkit.cpp
@@ -654,14 +654,14 @@ void FGroomCustomAssetEditorToolkit::InitCustomAssetEditor(const EToolkitMode::T
 		FGroomCustomAssetEditorToolkit* LocalToolKit = this;
 		auto InvalidateDetailViews = [LocalToolKit]()
 		{
-			LocalToolKit->DetailView_LODProperties->ForceRefresh();
-			LocalToolKit->DetailView_InterpolationProperties->ForceRefresh();
-			LocalToolKit->DetailView_RenderingProperties->ForceRefresh();
-			LocalToolKit->DetailView_PhysicsProperties->ForceRefresh();
-			LocalToolKit->DetailView_CardsProperties->ForceRefresh();
-			LocalToolKit->DetailView_MeshesProperties->ForceRefresh();
-			LocalToolKit->DetailView_MaterialProperties->ForceRefresh();
-			LocalToolKit->DetailView_BindingProperties->ForceRefresh();
+			if (LocalToolKit->DetailView_LODProperties)				{ LocalToolKit->DetailView_LODProperties->ForceRefresh(); }
+			if (LocalToolKit->DetailView_InterpolationProperties)	{ LocalToolKit->DetailView_InterpolationProperties->ForceRefresh(); }
+			if (LocalToolKit->DetailView_RenderingProperties)		{ LocalToolKit->DetailView_RenderingProperties->ForceRefresh(); }
+			if (LocalToolKit->DetailView_PhysicsProperties)			{ LocalToolKit->DetailView_PhysicsProperties->ForceRefresh(); }
+			if (LocalToolKit->DetailView_CardsProperties)			{ LocalToolKit->DetailView_CardsProperties->ForceRefresh(); }
+			if (LocalToolKit->DetailView_MeshesProperties)			{ LocalToolKit->DetailView_MeshesProperties->ForceRefresh(); }
+			if (LocalToolKit->DetailView_MaterialProperties)		{ LocalToolKit->DetailView_MaterialProperties->ForceRefresh(); }
+			if (LocalToolKit->DetailView_BindingProperties)			{ LocalToolKit->DetailView_BindingProperties->ForceRefresh(); }
 		};
 
 		PropertyListenDelegates.Add(GroomAsset->GetOnGroomAssetResourcesChanged().AddLambda(InvalidateDetailViews));
diff --git a/Engine/Plugins/Runtime/HairStrands/Source/HairStrandsEditor/Private/GroomMaterialDetails.cpp b/Engine/Plugins/Runtime/HairStrands/Source/HairStrandsEditor/Private/GroomMaterialDetails.cpp
--- a/Engine/Plugins/Runtime/HairStrands/Source/HairStrandsEditor/Private/GroomMaterialDetails.cpp
+++ b/Engine/Plugins/Runtime/HairStrands/Source/HairStrandsEditor/Private/GroomMaterialDetails.cpp
@@ -210,8 +210,10 @@ void FGroomMaterialDetails::CustomizeDetails( IDetailLayoutBuilder& DetailLayout
 		DetailLayout.HideProperty(DetailLayout.GetProperty(UGroomAsset::GetHairGroupsInfoMemberName(), UGroomAsset::StaticClass()));
 		DetailLayout.HideProperty(DetailLayout.GetProperty(UGroomAsset::GetEnableGlobalInterpolationMemberName(), UGroomAsset::StaticClass()));
 		DetailLayout.HideProperty(DetailLayout.GetProperty(UGroomAsset::GetHairInterpolationTypeMemberName(), UGroomAsset::StaticClass()));
-		DetailLayout.HideProperty(DetailLayout.GetProperty(UGroomAsset::GetMinLODMemberName(), UGroomAsset::StaticClass()));
+		DetailLayout.HideProperty(DetailLayout.GetProperty(UGroomAsset::GetAutoLODBiasMemberName(), UGroomAsset::StaticClass()));
 		DetailLayout.HideProperty(DetailLayout.GetProperty(UGroomAsset::GetDisableBelowMinLodStrippingMemberName(), UGroomAsset::StaticClass()));
+		DetailLayout.HideProperty(DetailLayout.GetProperty(UGroomAsset::GetMinLODMemberName(), UGroomAsset::StaticClass()));
+		DetailLayout.HideProperty(DetailLayout.GetProperty(UGroomAsset::GetLODModeMemberName(), UGroomAsset::StaticClass()));
 	}
 
 	GroomDetailLayout = &DetailLayout;
diff --git a/Engine/Source/Runtime/Engine/Classes/Engine/RendererSettings.h b/Engine/Source/Runtime/Engine/Classes/Engine/RendererSettings.h
--- a/Engine/Source/Runtime/Engine/Classes/Engine/RendererSettings.h
+++ b/Engine/Source/Runtime/Engine/Classes/Engine/RendererSettings.h
@@ -1112,6 +1112,15 @@ class URendererSettings : public UDeveloperSettings
 		ConfigRestartRequired = true))
 		uint32 bOrderedIndependentTransparencyEnable : 1;
 
+	/**
+	"Enable hair strands Auto LOD mode by default."
+	*/
+	UPROPERTY(config, EditAnywhere, Category = HairStrands, meta = (
+		ConsoleVariable = "r.HairStrands.LODMode", DisplayName = "Enable Hair Strands 'Auto' LOD mode",
+		ToolTip = "Enable hair strands Auto LOD mode by default. Otherwise use Manual LOD mode. Auto LOD mode adapts hair curves based on screen coverage. Manual LOD mode relies on LODs manually setup per groom asset. This global behavior can be overridden per groom asset",
+		ConfigRestartRequired = true))
+	uint32 bUseHairStrandsAutoLODMode  : 1;
+
 	/**
 	"Skin cache allows a compute shader to skin once each vertex, save those results into a new buffer and reuse those calculations when later running the depth, base and velocity passes. This also allows opting into the 'recompute tangents' for skinned mesh instance feature. Disabling will reduce the number of shader permutations required per material. Changing this setting requires restarting the editor."
 	*/
diff --git a/Engine/Source/Runtime/Renderer/Public/HairStrandsInterface.h b/Engine/Source/Runtime/Renderer/Public/HairStrandsInterface.h
--- a/Engine/Source/Runtime/Renderer/Public/HairStrandsInterface.h
+++ b/Engine/Source/Runtime/Renderer/Public/HairStrandsInterface.h
@@ -397,6 +397,7 @@ public:
 	float LODBias = 0;			// Current LOD bias
 	bool bLODVisibility = true; // Enable/disable hair rendering for this component
 	bool bAutoLOD = false;
+	float AutoLODBias = 0;
 
 	// Active/used point/curved based on select continuous LOD
 	uint32 ContinuousLODPointCount = 0;
