Fix CIS issue.

The console command referenced in Project Settings was defined inside the plugin which might not always be loaded.

diff --git a/Engine/Plugins/Runtime/HairStrands/Source/HairStrandsCore/Private/GroomAsset.cpp b/Engine/Plugins/Runtime/HairStrands/Source/HairStrandsCore/Private/GroomAsset.cpp
--- a/Engine/Plugins/Runtime/HairStrands/Source/HairStrandsCore/Private/GroomAsset.cpp
+++ b/Engine/Plugins/Runtime/HairStrands/Source/HairStrandsCore/Private/GroomAsset.cpp
@@ -85,12 +85,9 @@ static TAutoConsoleVariable<int32> GHairStrandsWarningLogVerbosity(
 static int32 GHairStrandsDDCLogEnable = 0;
 static FAutoConsoleVariableRef CVarHairStrandsDDCLogEnable(TEXT("r.HairStrands.DDCLog"), GHairStrandsDDCLogEnable, TEXT("Enable DDC logging for groom assets and groom binding assets"));
 
-static int32 GHairStrands_LODMode = 0;
-static FAutoConsoleVariableRef CVarHairStrands_LODMode(TEXT("r.HairStrands.LODMode"), GHairStrands_LODMode, TEXT("Enable hair strands Auto LOD mode by default. Otherwise use Manual LOD mode. Auto LOD mode adapts hair curves based on screen coverage. Manual LOD mode relies on LODs manually setup per groom asset. This global behavior can be overridden per groom asset."), ECVF_RenderThreadSafe);
-
-EGroomLODMode GetHairStrandsLODMode()
+static EGroomLODMode InternalGetHairStrandsLODMode()
 {
-	return GHairStrands_LODMode > 0 ? EGroomLODMode::Auto : EGroomLODMode::Manual;
+	return GetHairStrandsLODMode() > 0 ? EGroomLODMode::Auto : EGroomLODMode::Manual;
 }
 
 static int32 GHairStrandsSupportCompressedPosition = 0;
@@ -3852,7 +3849,7 @@ void UGroomAsset::ChangeFeatureLevel(ERHIFeatureLevel::Type In)
 GROOMASSET_DEFINE_MEMBER_NAME(LODMode)
 EGroomLODMode UGroomAsset::GetLODMode() const
 {
-	return LODMode == EGroomLODMode::Default ? GetHairStrandsLODMode() : LODMode;
+	return LODMode == EGroomLODMode::Default ? InternalGetHairStrandsLODMode() : LODMode;
 }
 
 GROOMASSET_DEFINE_MEMBER_NAME(AutoLODBias)
diff --git a/Engine/Source/Runtime/RenderCore/Private/RenderUtils.cpp b/Engine/Source/Runtime/RenderCore/Private/RenderUtils.cpp
--- a/Engine/Source/Runtime/RenderCore/Private/RenderUtils.cpp
+++ b/Engine/Source/Runtime/RenderCore/Private/RenderUtils.cpp
@@ -1689,6 +1689,12 @@ static TAutoConsoleVariable<int32> CVarHairStrandsUsesTriangleStrips(
 	TEXT("Enable triangle strip geometry for hair strands rendering. This improves performances, but removes the last segments of each curve."),
 	ECVF_ReadOnly | ECVF_RenderThreadSafe);
 
+static TAutoConsoleVariable<int32> CVarHairStrandsLODMode(
+	TEXT("r.HairStrands.LODMode"),
+	0,
+	TEXT("Enable hair strands Auto LOD mode by default. Otherwise use Manual LOD mode. Auto LOD mode adapts hair curves based on screen coverage. Manual LOD mode relies on LODs manually setup per groom asset. This global behavior can be overridden per groom asset."),
+	ECVF_RenderThreadSafe);
+
 int32 GetLightFunctionAtlasFormat()
 {
 	return CVarLightFunctionAtlasFormat.GetValueOnAnyThread();
@@ -1719,6 +1725,11 @@ bool GetHairStrandsUsesTriangleStrips()
 	return CVarHairStrandsUsesTriangleStrips.GetValueOnAnyThread() > 0;
 }
 
+uint32 GetHairStrandsLODMode()
+{
+	return FMath::Clamp(CVarHairStrandsLODMode.GetValueOnAnyThread(), 0, 1);
+}
+
 ///////////////////////////////////////////////////////////////////////////////////////////////////
 // Substrate settings interface
 
diff --git a/Engine/Source/Runtime/RenderCore/Public/RenderUtils.h b/Engine/Source/Runtime/RenderCore/Public/RenderUtils.h
--- a/Engine/Source/Runtime/RenderCore/Public/RenderUtils.h
+++ b/Engine/Source/Runtime/RenderCore/Public/RenderUtils.h
@@ -579,6 +579,7 @@ RENDERCORE_API bool GetTranslucentUsesLightRectLights();
 RENDERCORE_API bool GetTranslucentUsesLightIESProfiles();
 
 RENDERCORE_API bool GetHairStrandsUsesTriangleStrips();
+RENDERCORE_API uint32 GetHairStrandsLODMode();
 
 // LuminanceMax is the amount of light that will cause the sensor to saturate at EV100.
 //  See also https://en.wikipedia.org/wiki/Film_speed and https://en.wikipedia.org/wiki/Exposure_value for more info.
