Fix Groom deformer graph no longer working.

The EnqueWork for deformer graph was no longer called due to SendRenderData no longer being called.

diff --git a/Engine/Plugins/Runtime/HairStrands/Source/HairStrandsCore/Private/GroomComponent.cpp b/Engine/Plugins/Runtime/HairStrands/Source/HairStrandsCore/Private/GroomComponent.cpp
--- a/Engine/Plugins/Runtime/HairStrands/Source/HairStrandsCore/Private/GroomComponent.cpp
+++ b/Engine/Plugins/Runtime/HairStrands/Source/HairStrandsCore/Private/GroomComponent.cpp
@@ -3573,7 +3573,10 @@ void UGroomComponent::TickComponent(float DeltaTime, enum ELevelTick TickType, F
 		});
 	}
 	
-	if(GetHairSwapBufferType() == EHairBufferSwapType::RenderFrame)
+	// Need to run SendRenderDynamicData_Concurrent() 
+	// * If SwapBufferType is based on RenderFrame
+	// * If groom has a mesh deformer
+	if(GetHairSwapBufferType() == EHairBufferSwapType::RenderFrame || MeshDeformerInstance != nullptr)
 	{
 		MarkRenderDynamicDataDirty();
 	}
