Fix hair bulkd data header not being correctly reset when rebuilding an asset.

This causes array data to always increases in size, and leads to artifact (invalid curves' coverage, ...)

diff --git a/Engine/Plugins/Runtime/HairStrands/Source/HairStrandsCore/Private/GroomBuilder.cpp b/Engine/Plugins/Runtime/HairStrands/Source/HairStrandsCore/Private/GroomBuilder.cpp
--- a/Engine/Plugins/Runtime/HairStrands/Source/HairStrandsCore/Private/GroomBuilder.cpp
+++ b/Engine/Plugins/Runtime/HairStrands/Source/HairStrandsCore/Private/GroomBuilder.cpp
@@ -51,7 +51,7 @@ bool DoesHairStrandsSupportCompressedPosition();
 
 FString FGroomBuilder::GetVersion()
 {
-	return TEXT("v15c");
+	return TEXT("v15d");
 }
 
 namespace GroomBuilder_Voxelization
diff --git a/Engine/Plugins/Runtime/HairStrands/Source/HairStrandsCore/Private/HairStrandsDatas.cpp b/Engine/Plugins/Runtime/HairStrands/Source/HairStrandsCore/Private/HairStrandsDatas.cpp
--- a/Engine/Plugins/Runtime/HairStrands/Source/HairStrandsCore/Private/HairStrandsDatas.cpp
+++ b/Engine/Plugins/Runtime/HairStrands/Source/HairStrandsCore/Private/HairStrandsDatas.cpp
@@ -633,12 +633,7 @@ uint32 FHairStrandsBulkData::GetSize() const
 
 void FHairStrandsBulkData::Reset()
 {
-	Header.CurveCount = 0;
-	Header.PointCount = 0;
-	Header.MaxLength = 0;
-	Header.MaxRadius = 0;
-	Header.BoundingBox = FBox(EForceInit::ForceInit);
-	Header.Flags = 0;
+	Header = FHeader();
 	for (uint8 AttributeIt = 0; AttributeIt < HAIR_CURVE_ATTRIBUTE_COUNT; ++AttributeIt)
 	{
 		Header.CurveAttributeOffsets[AttributeIt] = HAIR_ATTRIBUTE_INVALID_OFFSET;
