GroomAsset: Fix deprecated array api on texture layout setting

diff --git a/Engine/Plugins/Runtime/HairStrands/Source/HairStrandsCore/Private/GroomAssetCards.cpp b/Engine/Plugins/Runtime/HairStrands/Source/HairStrandsCore/Private/GroomAssetCards.cpp
--- a/Engine/Plugins/Runtime/HairStrands/Source/HairStrandsCore/Private/GroomAssetCards.cpp
+++ b/Engine/Plugins/Runtime/HairStrands/Source/HairStrandsCore/Private/GroomAssetCards.cpp
@@ -21,7 +21,7 @@ void FHairGroupCardsTextures::SetLayout(EHairTextureLayout InLayout)
 	Layout = InLayout;
 
 	uint32 TextureCount = GetHairTextureLayoutTextureCount(Layout);
-	Textures.SetNum(TextureCount, true);
+	Textures.SetNum(TextureCount, EAllowShrinking::Yes);
 }
 
 void FHairGroupCardsTextures::SetTexture(int32 SlotIdx, UTexture2D* Texture)
