Add wave ops permutation for hair interpolation.

This allows to bypass LDS and keep all data registers. This reduces interpolation cost by ~10%

diff --git a/Engine/Plugins/Runtime/HairStrands/Source/HairStrandsCore/Private/GroomManager.cpp b/Engine/Plugins/Runtime/HairStrands/Source/HairStrandsCore/Private/GroomManager.cpp
--- a/Engine/Plugins/Runtime/HairStrands/Source/HairStrandsCore/Private/GroomManager.cpp
+++ b/Engine/Plugins/Runtime/HairStrands/Source/HairStrandsCore/Private/GroomManager.cpp
@@ -1386,6 +1386,7 @@ static void RunHairStrandsInterpolation_Strands(
 				AddHairStrandsInterpolationPass(
 					GraphBuilder,
 					ShaderMap,
+					View->GetShaderPlatform(),
 					ShaderPrintData,
 					InstanceData.Instance,
 					InstanceData.ActivePointCount,
@@ -1943,6 +1944,7 @@ static void RunHairStrandsInterpolation_Cards(
 				AddHairStrandsInterpolationPass(
 					GraphBuilder,
 					ShaderMap,
+					View->GetShaderPlatform(),
 					ShaderPrintData,
 					InstanceData.Instance,
 					LOD.Guides.RestResource->GetPointCount(),
diff --git a/Engine/Plugins/Runtime/HairStrands/Source/HairStrandsCore/Private/HairStrandsInterpolation.cpp b/Engine/Plugins/Runtime/HairStrands/Source/HairStrandsCore/Private/HairStrandsInterpolation.cpp
--- a/Engine/Plugins/Runtime/HairStrands/Source/HairStrandsCore/Private/HairStrandsInterpolation.cpp
+++ b/Engine/Plugins/Runtime/HairStrands/Source/HairStrandsCore/Private/HairStrandsInterpolation.cpp
@@ -29,6 +29,7 @@
 #include "SystemTextures.h"
 #include "ShaderPrintParameters.h"
 #include "HairStrandsClusterCulling.h"
+#include "DataDrivenShaderPlatformInfo.h"
 
 static float GHairRaytracingRadiusScale = 0;
 static FAutoConsoleVariableRef CVarHairRaytracingRadiusScale(TEXT("r.HairStrands.RaytracingRadiusScale"), GHairRaytracingRadiusScale, TEXT("Override the per instance scale factor for raytracing hair strands geometry (0: disabled, >0:enabled)"));
@@ -463,7 +464,8 @@ class FHairInterpolationCS : public FGlobalShader
 	class FSimulation : SHADER_PERMUTATION_BOOL("PERMUTATION_SIMULATION");
 	class FSingleGuide : SHADER_PERMUTATION_BOOL("PERMUTATION_USE_SINGLE_GUIDE");
 	class FDeformer : SHADER_PERMUTATION_BOOL("PERMUTATION_DEFORMER");
-	using FPermutationDomain = TShaderPermutationDomain<FDynamicGeometry, FSimulation, FSingleGuide, FDeformer, FPointPerCurve>;
+	class FWaveOps : SHADER_PERMUTATION_BOOL("PERMUTATION_WAVEOPS");
+	using FPermutationDomain = TShaderPermutationDomain<FDynamicGeometry, FSimulation, FSingleGuide, FDeformer, FPointPerCurve, FWaveOps>;
 
 	BEGIN_SHADER_PARAMETER_STRUCT(FParameters, )
 		SHADER_PARAMETER_STRUCT_INCLUDE(ShaderPrint::FShaderParameters, ShaderDrawParameters)
@@ -515,7 +517,25 @@ class FHairInterpolationCS : public FGlobalShader
 
 public:
 	static uint32 GetGroupSize() { return 32;}
-	static bool ShouldCompilePermutation(const FGlobalShaderPermutationParameters& Parameters) { return IsHairStrandsSupported(EHairStrandsShaderType::All, Parameters.Platform); }
+	static ERHIFeatureSupport DoesSupportsWaveOps(EShaderPlatform InPlatform, uint32 InPointPerCurve)
+	{
+		// D3D11 / SM5 or preview do not support, or work well with, wave-ops by default (or SM5 preview has issues with wave intrinsics too)
+		if (InPointPerCurve != 32 || InPlatform == SP_PCD3D_SM5 || FDataDrivenShaderPlatformInfo::GetIsPreviewPlatform(InPlatform))
+		{
+			return ERHIFeatureSupport::Unsupported;
+		}
+
+		return FDataDrivenShaderPlatformInfo::GetSupportsWaveOperations(InPlatform);
+	}
+	static bool ShouldCompilePermutation(const FGlobalShaderPermutationParameters& Parameters) 
+	{ 
+		FPermutationDomain PermutationVector(Parameters.PermutationId);
+		if (PermutationVector.Get<FWaveOps>() && DoesSupportsWaveOps(Parameters.Platform, PermutationVector.Get<FPointPerCurve>()) == ERHIFeatureSupport::Unsupported)
+		{
+			return false;
+		}
+		return IsHairStrandsSupported(EHairStrandsShaderType::All, Parameters.Platform); 
+	}
 	static void ModifyCompilationEnvironment(const FGlobalShaderPermutationParameters& Parameters, FShaderCompilerEnvironment& OutEnvironment)
 	{
 		FGlobalShader::ModifyCompilationEnvironment(Parameters, OutEnvironment);
@@ -539,6 +559,7 @@ IMPLEMENT_GLOBAL_SHADER(FHairInterpolationCS, "/Engine/Private/HairStrands/HairS
 void AddHairStrandsInterpolationPass(
 	FRDGBuilder& GraphBuilder,
 	FGlobalShaderMap* ShaderMap,
+	EShaderPlatform InPlatform,
 	const FShaderPrintData* ShaderPrintData,
 	const FHairGroupInstance* Instance,
 	const uint32 VertexCount,
@@ -672,6 +693,7 @@ void AddHairStrandsInterpolationPass(
 	Parameters->DispatchCountX = DispatchInfo.DispatchCount.X;
 
 	const bool bSingleGuidePermutation = bUseSingleGuide || GetHairStrandsForceSingleGuideInterpolation();
+	const bool bWaveOps = GRHISupportsWaveOperations && GRHIMaximumWaveSize >= 32 && FHairInterpolationCS::DoesSupportsWaveOps(InPlatform, DispatchInfo.PointPerCurve) != ERHIFeatureSupport::Unsupported;
 
 	FHairInterpolationCS::FPermutationDomain PermutationVector;
 	PermutationVector.Set<FHairInterpolationCS::FDynamicGeometry>(DynamicGeometryType);
@@ -679,10 +701,11 @@ void AddHairStrandsInterpolationPass(
 	PermutationVector.Set<FHairInterpolationCS::FSingleGuide>(bSingleGuidePermutation);
 	PermutationVector.Set<FHairInterpolationCS::FDeformer>(bSupportDeformer);
 	PermutationVector.Set<FHairInterpolationCS::FPointPerCurve>(DispatchInfo.PointPerCurve);
+	PermutationVector.Set<FHairInterpolationCS::FWaveOps>(bWaveOps);
 	TShaderMapRef<FHairInterpolationCS> ComputeShader(ShaderMap, PermutationVector);
 	FComputeShaderUtils::AddPass(
 		GraphBuilder,
-		RDG_EVENT_NAME("HairStrands::Interpolation(culling=off, PerCurve=%d, SingleGuide=%d, DynGeom:%d)", DispatchInfo.PointPerCurve, bSingleGuidePermutation ? 1u : 0u, DynamicGeometryType > 0 ? 1u : 0u),
+		RDG_EVENT_NAME("HairStrands::Interpolation(culling=off, PerCurve=%d, SingleGuide=%d, DynGeom:%d, WaveOps:%d)", DispatchInfo.PointPerCurve, bSingleGuidePermutation ? 1u : 0u, DynamicGeometryType > 0 ? 1u : 0u, bWaveOps ? 1u : 0u),
 		ComputeShader,
 		Parameters,
 		DispatchInfo.DispatchCount);
diff --git a/Engine/Plugins/Runtime/HairStrands/Source/HairStrandsCore/Public/HairStrandsInterpolation.h b/Engine/Plugins/Runtime/HairStrands/Source/HairStrandsCore/Public/HairStrandsInterpolation.h
--- a/Engine/Plugins/Runtime/HairStrands/Source/HairStrandsCore/Public/HairStrandsInterpolation.h
+++ b/Engine/Plugins/Runtime/HairStrands/Source/HairStrandsCore/Public/HairStrandsInterpolation.h
@@ -118,6 +118,7 @@ void AddClearAABBPass(
 void AddHairStrandsInterpolationPass(
 	FRDGBuilder& GraphBuilder,
 	FGlobalShaderMap* ShaderMap,
+	EShaderPlatform InPlatform,
 	const FShaderPrintData* ShaderPrintData,
 	const FHairGroupInstance* Instance,
 	const uint32 VertexCount,
diff --git a/Engine/Shaders/Private/HairStrands/HairStrandsInterpolation.usf b/Engine/Shaders/Private/HairStrands/HairStrandsInterpolation.usf
--- a/Engine/Shaders/Private/HairStrands/HairStrandsInterpolation.usf
+++ b/Engine/Shaders/Private/HairStrands/HairStrandsInterpolation.usf
@@ -581,6 +581,38 @@ void MainCS(uint3 DispatchThreadId : SV_DispatchThreadID, uint ThreadIndex : SV_
 //-------------------------------------------------------------------------------------------------
 #if PERMUTATION_POINT_PER_CURVE == TOTALMAXPOINTPERCURVE
 
+#if PERMUTATION_WAVEOPS
+FRootTriangleData WaveReadLaneFirstTriangle(FRootTriangleData In)
+{
+	FRootTriangleData Out;
+	Out.RootBarycentric = WaveReadLaneFirst(In.RootBarycentric);
+
+	Out.RestTriangle.P0 = WaveReadLaneFirst(In.RestTriangle.P0);
+	Out.RestTriangle.P1 = WaveReadLaneFirst(In.RestTriangle.P1);
+	Out.RestTriangle.P2 = WaveReadLaneFirst(In.RestTriangle.P2);
+	Out.RestTriangle.E0 = WaveReadLaneFirst(In.RestTriangle.E0);
+	Out.RestTriangle.E1 = WaveReadLaneFirst(In.RestTriangle.E1);
+	Out.RestTriangle.N  = WaveReadLaneFirst(In.RestTriangle.N) ;
+
+	Out.DeformedTriangle.P0 = WaveReadLaneFirst(In.DeformedTriangle.P0);
+	Out.DeformedTriangle.P1 = WaveReadLaneFirst(In.DeformedTriangle.P1);
+	Out.DeformedTriangle.P2 = WaveReadLaneFirst(In.DeformedTriangle.P2);
+	Out.DeformedTriangle.E0 = WaveReadLaneFirst(In.DeformedTriangle.E0);
+	Out.DeformedTriangle.E1 = WaveReadLaneFirst(In.DeformedTriangle.E1);
+	Out.DeformedTriangle.N  = WaveReadLaneFirst(In.DeformedTriangle.N) ;
+
+	return Out;
+}
+
+FGuideCurveData WaveReadLaneFirstCurveData(FGuideCurveData In)
+{
+	FGuideCurveData Out;
+	Out.CurveIndices	 = WaveReadLaneFirst(In.CurveIndices);
+	Out.CurveWeights	 = WaveReadLaneFirst(In.CurveWeights);
+	Out.RootPointIndices = WaveReadLaneFirst(In.RootPointIndices);
+	return Out;
+}
+#else
 groupshared uint CurvePointIndex;
 groupshared uint CurvePointCount;
 
@@ -589,7 +621,11 @@ groupshared float3 RenPositionOffset;
 
 groupshared FGuideCurveData SimCurveData;
 groupshared FRootTriangleData SimTris[2];
+#endif
 
+#if PERMUTATION_WAVEOPS && COMPILER_SUPPORTS_WAVE_SIZE
+WAVESIZE(32) // PERMUTATION_WAVE_OPS is true only when wave>=32 are available
+#endif
 [numthreads(GROUP_SIZE, 1, 1)]
 void MainCS(uint3 DispatchThreadId : SV_DispatchThreadID, uint ThreadIndex : SV_GroupIndex, uint2 GroupIndex : SV_GroupID)
 {
@@ -601,8 +637,26 @@ void MainCS(uint3 DispatchThreadId : SV_DispatchThreadID, uint ThreadIndex : SV_
 		return;
 	}
 	
+	#if PERMUTATION_WAVEOPS
+	// 2. Curve's point count/offset
+	uint CurvePointIndex = 0;
+	uint CurvePointCount = 0;
+
+	float3 RenPositionOffset = 0;
+	FRootTriangleData RenTri = (FRootTriangleData)0;
+
+	FGuideCurveData SimCurveData = (FGuideCurveData)0;
+	FRootTriangleData SimTris[2];
+	SimTris[0] = (FRootTriangleData)0;
+	SimTris[1] = (FRootTriangleData)0;
+	#endif
+
 	// 2. Curve's point count/offset
+	#if PERMUTATION_WAVEOPS
+	if (WaveIsFirstLane())
+	#else
 	if (ThreadIndex == 0)
+	#endif
 	{
 		CurvePointIndex = 0;
 		CurvePointCount = 0;
@@ -641,7 +695,29 @@ void MainCS(uint3 DispatchThreadId : SV_DispatchThreadID, uint ThreadIndex : SV_
 		}
 		RenPositionOffset = GetRenPositionOffset();
 	}
+
+	#if PERMUTATION_WAVEOPS
+	{
+		// Broadcast values
+		CurvePointIndex   = WaveReadLaneFirst(CurvePointIndex);
+		CurvePointCount   = WaveReadLaneFirst(CurvePointCount);
+		RenPositionOffset = WaveReadLaneFirst(RenPositionOffset);
+		#if PERMUTATION_DYNAMIC_GEOMETRY >= 1
+		RenTri = WaveReadLaneFirstTriangle(RenTri);
+		#endif
+
+		#if PERMUTATION_SIMULATION
+		SimCurveData = WaveReadLaneFirstCurveData(SimCurveData);
+		SimTris[0] = WaveReadLaneFirstTriangle(SimTris[0]);
+		if (!PERMUTATION_USE_SINGLE_GUIDE)
+		{
+			SimTris[1] = WaveReadLaneFirstTriangle(SimTris[1]);
+		}		
+		#endif
+	}
+	#else
 	GroupMemoryBarrierWithGroupSync();
+	#endif
 
 	// 2. Apply deformation
 	{
