Fix crash when reimporting binding asset.

diff --git a/Engine/Plugins/Runtime/HairStrands/Source/HairStrandsCore/Private/GroomBindingAsset.cpp b/Engine/Plugins/Runtime/HairStrands/Source/HairStrandsCore/Private/GroomBindingAsset.cpp
--- a/Engine/Plugins/Runtime/HairStrands/Source/HairStrandsCore/Private/GroomBindingAsset.cpp
+++ b/Engine/Plugins/Runtime/HairStrands/Source/HairStrandsCore/Private/GroomBindingAsset.cpp
@@ -170,7 +170,7 @@ void UGroomBindingAsset::InitResource()
 	TRACE_CPUPROFILER_EVENT_SCOPE(UGroomBindingAsset::InitResource);
 
 	// Ensure we are releasing binding resources before reallocating them
-	ReleaseResource();
+	ReleaseResource(true/*bResetLoadedSize*/);
 
 	for (UGroomBindingAsset::FHairGroupPlatformData& BulkData : GetHairGroupsPlatformData())
 	{
@@ -238,7 +238,7 @@ void UGroomBindingAsset::UpdateResource()
 	}
 }
 
-void UGroomBindingAsset::ReleaseResource()
+void UGroomBindingAsset::ReleaseResource(bool bResetLoadedSize)
 {
 	// Delay destruction to insure that the rendering thread is done with all resources usage
 	if (GetHairGroupResources().Num() > 0)
@@ -248,7 +248,7 @@ void UGroomBindingAsset::ReleaseResource()
 			FHairStrandsRestRootResource* InSimRootResources = Resource.SimRootResources;
 			FHairStrandsRestRootResource* InRenRootResources = Resource.RenRootResources;
 			ENQUEUE_RENDER_COMMAND(ReleaseHairStrandsResourceCommand)(UE::RenderCommandPipe::Groom,
-				[InSimRootResources, InRenRootResources](FRHICommandList& RHICmdList)
+				[InSimRootResources, InRenRootResources, bResetLoadedSize](FRHICommandList& RHICmdList)
 			{
 				if (InSimRootResources)
 				{
@@ -257,7 +257,10 @@ void UGroomBindingAsset::ReleaseResource()
 				}
 				if (InRenRootResources)
 				{
-					InRenRootResources->InternalResetLoadedSize();
+					if (bResetLoadedSize)
+					{
+						InRenRootResources->InternalResetLoadedSize();
+					}
 					InRenRootResources->ReleaseResource();
 					delete InRenRootResources;
 				}
@@ -308,7 +311,8 @@ void UGroomBindingAsset::ReleaseResource()
 
 void UGroomBindingAsset::Reset()
 {
-	ReleaseResource();
+	// No need for resetting LoadedSize as the bulk datas are removed (e.g., when reloading/reimporting a groom asset)
+	ReleaseResource(false/*bResetLoadedSize*/);
 	for (UGroomBindingAsset::FHairGroupPlatformData& Data : GetHairGroupsPlatformData())
 	{
 		Data.SimRootBulkDatas.Empty();
@@ -418,7 +422,7 @@ void UGroomBindingAsset::PostSaveRoot(FObjectPostSaveRootContext ObjectSaveConte
 
 void UGroomBindingAsset::BeginDestroy()
 {
-	ReleaseResource();
+	ReleaseResource(false/*bResetLoadedSize*/);
 	Super::BeginDestroy();
 
 #if WITH_EDITOR
@@ -1248,7 +1252,7 @@ TArray<UGroomBindingAsset::FHairGroupPlatformData>& UGroomBindingAsset::GetHairG
 #if WITH_EDITOR
 void UGroomBindingAsset::RecreateResources()
 {
-	ReleaseResource();
+	ReleaseResource(true /*bResetLoadedSize*/);
 	InitResource();
 	OnGroomBindingAssetChanged.Broadcast();
 }
diff --git a/Engine/Plugins/Runtime/HairStrands/Source/HairStrandsCore/Public/GroomBindingAsset.h b/Engine/Plugins/Runtime/HairStrands/Source/HairStrandsCore/Public/GroomBindingAsset.h
--- a/Engine/Plugins/Runtime/HairStrands/Source/HairStrandsCore/Public/GroomBindingAsset.h
+++ b/Engine/Plugins/Runtime/HairStrands/Source/HairStrandsCore/Public/GroomBindingAsset.h
@@ -244,7 +244,7 @@ public:
 	void UpdateResource();
 
 	/** Release the hair strands resource. */
-	void ReleaseResource();
+	void ReleaseResource(bool bResetLoadedSize);
 
 	void Reset();
 
