Enable manual cache by default.

This will only be use if skin cache is disabled.

diff --git a/Engine/Plugins/Runtime/HairStrands/Source/HairStrandsCore/Private/GroomManager.cpp b/Engine/Plugins/Runtime/HairStrands/Source/HairStrandsCore/Private/GroomManager.cpp
--- a/Engine/Plugins/Runtime/HairStrands/Source/HairStrandsCore/Private/GroomManager.cpp
+++ b/Engine/Plugins/Runtime/HairStrands/Source/HairStrandsCore/Private/GroomManager.cpp
@@ -41,7 +41,7 @@ static FAutoConsoleVariableRef CVarHairStrands_UseCards(TEXT("r.HairStrands.UseC
 static int32 GHairStrands_SwapBufferType = 2;
 static FAutoConsoleVariableRef CVarGHairStrands_SwapBufferType(TEXT("r.HairStrands.SwapType"), GHairStrands_SwapBufferType, TEXT("Swap rendering buffer at the end of frame. This is an experimental toggle. Default:1"));
 
-static int32 GHairStrands_ManualSkinCache = 0;
+static int32 GHairStrands_ManualSkinCache = 1;
 static FAutoConsoleVariableRef CVarGHairStrands_ManualSkinCache(TEXT("r.HairStrands.ManualSkinCache"), GHairStrands_ManualSkinCache, TEXT("If skin cache is not enabled, and grooms use skinning method, this enable a simple skin cache mechanisme for groom. Default:disable"));
 
 static int32 GHairStrands_InterpolationFrustumCullingEnable = 1;
