Fix incorrect RBF deformation when exporting MHs

The skel. mesh render datas were not identical between the groom binding builder, and the RBF builder, due to how the skel. mesh render data were fetched.

diff --git a/Engine/Plugins/Runtime/HairStrands/Source/HairStrandsCore/Private/GroomBindingBuilder.cpp b/Engine/Plugins/Runtime/HairStrands/Source/HairStrandsCore/Private/GroomBindingBuilder.cpp
--- a/Engine/Plugins/Runtime/HairStrands/Source/HairStrandsCore/Private/GroomBindingBuilder.cpp
+++ b/Engine/Plugins/Runtime/HairStrands/Source/HairStrandsCore/Private/GroomBindingBuilder.cpp
@@ -60,7 +60,7 @@ static FAutoConsoleVariableRef CVarHairStrandsBindingBuilderWarningEnable(TEXT("
 FString FGroomBindingBuilder::GetVersion()
 {
 	// Important to update the version when groom building changes
-	return TEXT("4a");
+	return TEXT("4d");
 }
 
 ///////////////////////////////////////////////////////////////////////////////////////////////////
diff --git a/Engine/Plugins/Runtime/HairStrands/Source/HairStrandsCore/Private/GroomRBFDeformer.cpp b/Engine/Plugins/Runtime/HairStrands/Source/HairStrandsCore/Private/GroomRBFDeformer.cpp
--- a/Engine/Plugins/Runtime/HairStrands/Source/HairStrandsCore/Private/GroomRBFDeformer.cpp
+++ b/Engine/Plugins/Runtime/HairStrands/Source/HairStrandsCore/Private/GroomRBFDeformer.cpp
@@ -7,6 +7,9 @@
 #include "GroomBuilder.h"
 #include "Rendering/SkeletalMeshRenderData.h"
 #include "MeshAttributes.h"
+#include "Misc/CoreMisc.h"
+#include "Interfaces/ITargetPlatformManagerModule.h"
+#include "Engine/SkinnedAssetAsyncCompileUtils.h"
 #include "Engine/StaticMesh.h"
 #include "Engine/TextureDefines.h"
 #include "Engine/Texture.h"
@@ -32,10 +35,11 @@ void InitMeshSamples(
 	for (uint32 SampleIndex = 0; SampleIndex < MaxSampleCount; ++SampleIndex)
 	{
 		const uint32 VertexIndex = SampleIndicesBuffer[SampleIndex];
-		if (VertexIndex >= MaxVertexCount)
-			continue;
-
-		OutSamplePositionsBuffer[SampleIndex] = VertexPositionsBuffer[VertexIndex];
+		check(VertexIndex < MaxVertexCount);
+		if (VertexIndex < MaxVertexCount)
+		{
+			OutSamplePositionsBuffer[SampleIndex] = VertexPositionsBuffer[VertexIndex];
+		}
 	}
 }
 
@@ -655,11 +659,23 @@ void FGroomRBFDeformer::GetRBFDeformedGroomAsset(const UGroomAsset* InGroomAsset
 #if WITH_EDITORONLY_DATA
 	if (InGroomAsset && BindingAsset && BindingAsset->GetTargetSkeletalMesh() && BindingAsset->GetSourceSkeletalMesh())
 	{
+		// Get the skel. mesh render data for the current platform
+		// This similar to how we fetch skel. mesh data when building groom binding. This is requires in order to get 
+		// identical skel. mesh render data and ensure that the hair deformation is done correctly.
+		check(BindingAsset->GetGroomBindingType() == EGroomBindingMeshType::SkeletalMesh);
+		USkeletalMesh* TargetSkeletalMesh = BindingAsset->GetTargetSkeletalMesh();
+		FScopedSkeletalMeshRenderData TargetSkeletalMeshScopedData(TargetSkeletalMesh);
+		{
+			ITargetPlatform* RunningPlatform = GetTargetPlatformManagerRef().GetRunningTargetPlatform();
+			FSkinnedAssetAsyncBuildScope AsyncBuildScope(TargetSkeletalMesh);
+			USkeletalMesh::GetPlatformSkeletalMeshRenderData(RunningPlatform, TargetSkeletalMeshScopedData);
+		}
+
 		// Use the LOD0 skeletal mesh to extract the vertices used for the RBF weight computation
 		const int32 MeshLODIndex = 0;
 
 		// Get the target mesh vertices (source and target)
-		const FSkeletalMeshRenderData* SkeletalMeshData_Target = BindingAsset->GetTargetSkeletalMesh()->GetResourceForRendering();
+		const FSkeletalMeshRenderData* SkeletalMeshData_Target = TargetSkeletalMeshScopedData.GetData();
 		TArray<FVector3f> MeshVertexPositionsBuffer_Target;
 		ExtractSkeletalVertexPosition(SkeletalMeshData_Target, MeshLODIndex, MeshVertexPositionsBuffer_Target);
 
@@ -761,6 +777,9 @@ void FGroomRBFDeformer::GetRBFDeformedGroomAsset(const UGroomAsset* InGroomAsset
 					RenLODData.MeshInterpolationWeightsBuffer 	= SimLODData.MeshInterpolationWeightsBuffer;
 					RenLODData.MeshSampleIndicesBuffer 			= SimLODData.MeshSampleIndicesBuffer;
 					RenLODData.RestSamplePositionsBuffer 		= SimLODData.RestSamplePositionsBuffer;
+					RenLODData.MeshSampleSectionsBuffer			= SimLODData.MeshSampleSectionsBuffer;
+					RenLODData.UniqueSectionIds					= SimLODData.UniqueSectionIds;
+					RenLODData.MeshSectionCount					= SimLODData.MeshSectionCount;
 				}
 
 				DeformedPositions[GroupIndex].RenderStrands = GetDeformedHairStrandsPositions(
