Fix crash when unstreaming mesh data.

diff --git a/Engine/Plugins/Runtime/HairStrands/Source/HairStrandsCore/Private/GroomManager.cpp b/Engine/Plugins/Runtime/HairStrands/Source/HairStrandsCore/Private/GroomManager.cpp
--- a/Engine/Plugins/Runtime/HairStrands/Source/HairStrandsCore/Private/GroomManager.cpp
+++ b/Engine/Plugins/Runtime/HairStrands/Source/HairStrandsCore/Private/GroomManager.cpp
@@ -418,8 +418,11 @@ static void AddHairSkinCacheDebugPass(
 
 	bool bIsGPUSkinCacheEnable = false;
 	{
-		static FShaderPlatformCachedIniValue<int32> PerPlatformCVar(TEXT("r.SkinCache.Mode"));
-		bIsGPUSkinCacheEnable = (PerPlatformCVar.Get(View->GetShaderPlatform()) != 0);
+		static FShaderPlatformCachedIniValue<int32> CVarSkinCacheCompileShader(TEXT("r.SkinCache.CompileShaders"));
+		static FShaderPlatformCachedIniValue<int32> CVarSkinCacheMode(TEXT("r.SkinCache.Mode"));
+		bIsGPUSkinCacheEnable = 
+			CVarSkinCacheCompileShader.Get(View->GetShaderPlatform()) != 0 &&
+			CVarSkinCacheMode.Get(View->GetShaderPlatform()) != 0;
 	}
 
 	const uint32 InstanceCount = Instances.Num();
@@ -528,6 +531,9 @@ static void GetOrAllocateCachedGeometry(
 	FRDGBufferSRVRef DeformedPreviousPositionSRV = nullptr;
 	OutHairGeometryCache.GetOrAdd(GraphBuilder, SkeletalMeshObject, &LODData, LODIndex, UniqueSections, DeformedPositionSRV, DeformedPreviousPositionSRV);
 
+	// Add reference to be sure the data are not streamed out while they are used
+	LODData.AddRef();
+
 	// Fill in result
 	for (uint32 SectionIndex : UniqueSections)
 	{
@@ -737,6 +743,15 @@ static void RunHairBindingSurfaceUpdate(
 	}
 
 	AddHairSkinCacheDebugPass(GraphBuilder, ShaderMap, View, ShaderPrintData, Instances, HairGeometryCache);
+
+	// Release reference on skel. mesh data (only for manual skin cache)
+	for (FHairGeometryCache::FData& Data : HairGeometryCache.Datas)
+	{
+		if (Data.LODData)
+		{
+			Data.LODData->Release();
+		}
+	}
 }
 
 ///////////////////////////////////////////////////////////////////////////////////////////////////
diff --git a/Engine/Shaders/Private/HairStrands/HairStrandsDebug.usf b/Engine/Shaders/Private/HairStrands/HairStrandsDebug.usf
--- a/Engine/Shaders/Private/HairStrands/HairStrandsDebug.usf
+++ b/Engine/Shaders/Private/HairStrands/HairStrandsDebug.usf
@@ -3123,14 +3123,17 @@ void MainCS(uint GroupIndex : SV_GroupIndex, uint3 DispatchThreadId : SV_Dispatc
 		PrintBool(Ctx, HairSkinCacheEnable);
 		Newline(Ctx);
 
-		bAdvanced = AddCheckbox(Ctx, TEXT("Advanced"), false, FontWhite); 
+		bAdvanced = AddCheckbox(Ctx, TEXT("Advanced"), false, FontSilver); 
 		Newline(Ctx);
 	}
 	Newline(Ctx);
 
 	// Draw instance details
-	Print(Ctx, TEXT("  Inst. Guide Strands Cards CacheType Sections LOD Name"));
-	Newline(Ctx);
+	if (bAdvanced)
+	{
+		Print(Ctx, TEXT("  Inst. Guide Strands Cards CacheType Sections LOD Name"));
+		Newline(Ctx);
+	}
 
 	LOOP
 	for (uint UniqueMeshIndex = 0; UniqueMeshIndex < UniqueMeshCount; ++UniqueMeshIndex)
