Fix CIS issue.

diff --git a/Engine/Plugins/Runtime/HairStrands/Source/HairStrandsCore/Private/GroomManager.cpp b/Engine/Plugins/Runtime/HairStrands/Source/HairStrandsCore/Private/GroomManager.cpp
--- a/Engine/Plugins/Runtime/HairStrands/Source/HairStrandsCore/Private/GroomManager.cpp
+++ b/Engine/Plugins/Runtime/HairStrands/Source/HairStrandsCore/Private/GroomManager.cpp
@@ -733,6 +733,7 @@ static void RunHairBindingSurfaceUpdate(
 	// Process manual skin cache requests
 	for (FHairGeometryCache::FData& Data : HairGeometryCache.Datas)
 	{
+		check(Data.LODData);
 		AddSkinUpdatePass(
 			GraphBuilder, 
 			ShaderMap, 
@@ -747,10 +748,8 @@ static void RunHairBindingSurfaceUpdate(
 	// Release reference on skel. mesh data (only for manual skin cache)
 	for (FHairGeometryCache::FData& Data : HairGeometryCache.Datas)
 	{
-		if (Data.LODData)
-		{
-			Data.LODData->Release();
-		}
+		check(Data.LODData);
+		Data.LODData->Release();
 	}
 }
 
