Fix FHairStrandsRaytracingResource not properly releasing FRayTracingGeometry.
- FHairStrandsRaytracingResource doesn't call RayTracingGeometry.InitResource(...) and instead initializes resource manually so also need to release it manually.

diff --git a/Engine/Plugins/Runtime/HairStrands/Source/HairStrandsCore/Private/GroomResources.cpp b/Engine/Plugins/Runtime/HairStrands/Source/HairStrandsCore/Private/GroomResources.cpp
--- a/Engine/Plugins/Runtime/HairStrands/Source/HairStrandsCore/Private/GroomResources.cpp
+++ b/Engine/Plugins/Runtime/HairStrands/Source/HairStrandsCore/Private/GroomResources.cpp
@@ -1908,7 +1908,8 @@ void FHairStrandsRaytracingResource::InternalAllocate(FRDGBuilder& GraphBuilder,
 		InternalCreateStructuredBufferRDG<FHairStrandsIndexFormat>(GraphBuilder, MaxIndexCount, IndexBuffer, ToHairResourceDebugName(TEXT("Hair.StrandsRaytracing_IndexBuffer"), ResourceName), OwnerName, EHairResourceUsageType::Dynamic);
 
 		// Force reallocating and rebuilding the RT geometry
-		RayTracingGeometry.ReleaseResource();
+		RayTracingGeometry.Initializer = {};
+		RayTracingGeometry.ReleaseRHI();
 		bIsRTGeometryInitialized = false;
 	}
 }
@@ -1917,7 +1918,8 @@ void FHairStrandsRaytracingResource::InternalRelease()
 {
 	PositionBuffer.Release();
 	IndexBuffer.Release();
-	RayTracingGeometry.ReleaseResource();
+	RayTracingGeometry.Initializer = {};
+	RayTracingGeometry.ReleaseRHI();
 	bIsRTGeometryInitialized = false;
 }
 #endif
