Add copy current to previous position for newly streamed/visible curves.

This makes the new curves previous frame data valid and allow ~correct motion vector.

diff --git a/Engine/Plugins/Runtime/HairStrands/Source/HairStrandsCore/Private/GroomManager.cpp b/Engine/Plugins/Runtime/HairStrands/Source/HairStrandsCore/Private/GroomManager.cpp
--- a/Engine/Plugins/Runtime/HairStrands/Source/HairStrandsCore/Private/GroomManager.cpp
+++ b/Engine/Plugins/Runtime/HairStrands/Source/HairStrandsCore/Private/GroomManager.cpp
@@ -1099,6 +1099,8 @@ static void RunHairStrandsInterpolation_Strands(
 		int32 MeshLODIndex = -1;
 		uint32 ActivePointCount = 0;
 		uint32 ActiveCurveCount = 0;
+		uint32 PreviousActivePointCount = 0;
+		uint32 PreviousActiveCurveCount = 0;
 		EGroomCacheType ActiveGroomCacheType = EGroomCacheType::None;
 		EHairBindingType BindingType = EHairBindingType::NoneBinding;
 		bool bSimulationEnable = false;
@@ -1148,6 +1150,8 @@ static void RunHairStrandsInterpolation_Strands(
 		InstanceData.Instance 					= Instance;
 		InstanceData.ActivePointCount 			= Instance->HairGroupPublicData->GetActiveStrandsPointCount();
 		InstanceData.ActiveCurveCount 			= Instance->HairGroupPublicData->GetActiveStrandsCurveCount();
+		InstanceData.PreviousActivePointCount 	= Instance->HairGroupPublicData->GetActiveStrandsPointCount(true /*bPrevious*/);
+		InstanceData.PreviousActiveCurveCount 	= Instance->HairGroupPublicData->GetActiveStrandsCurveCount(true /*bPrevious*/);
 		InstanceData.ActiveGroomCacheType 		= GetHairInstanceCacheType(Instance);
 		InstanceData.BindingType 				= Instance->BindingType;
 		InstanceData.HairLODIndex 				= Instance->HairGroupPublicData->LODIndex;
@@ -1446,12 +1450,20 @@ static void RunHairStrandsInterpolation_Strands(
 	{
 		if (InstanceData.ActiveGroomCacheType != EGroomCacheType::Strands && InstanceData.bNeedDeformation)
 		{
-			const bool bTransferPrevPosition =
-				InstanceData.Instance->HairGroupPublicData->VFInput.bHasLODSwitch &&
-				IsHairStrandsTransferPositionOnLODChange();
-			if (bTransferPrevPosition)
+			const bool bLODChanged = (InstanceData.Instance->HairGroupPublicData->VFInput.bHasLODSwitch && IsHairStrandsTransferPositionOnLODChange());
+			const bool bNewCurves  = InstanceData.ActivePointCount > InstanceData.PreviousActivePointCount;
+			if (bLODChanged || bNewCurves)
 			{
-				AddTransferPositionPass(GraphBuilder, ShaderMap, InstanceData.ActivePointCount, InstanceData.RDGResources.DeformedPosition.SRV, InstanceData.RDGResources.DeformedPrevPosition.UAV);
+				const uint32 PointOffset = bLODChanged ? 0 : InstanceData.PreviousActivePointCount;
+				const uint32 PointCount  = bLODChanged ? InstanceData.ActivePointCount : InstanceData.ActivePointCount - InstanceData.PreviousActivePointCount;
+				AddTransferPositionPass(
+					GraphBuilder, 
+					ShaderMap, 
+					PointOffset,
+					PointCount,
+					InstanceData.ActivePointCount, 
+					InstanceData.RDGResources.DeformedPosition.SRV, 
+					InstanceData.RDGResources.DeformedPrevPosition.UAV);
 				GraphBuilder.SetBufferAccessFinal(InstanceData.RDGResources.DeformedPrevPosition.Buffer, ERHIAccess::SRVMask);
 			}
 		}
@@ -2541,6 +2553,8 @@ static void ApplyHairLOD(const FHairLOD& In, FHairGroupInstance* OutInstance)
 	// Feedback game thread with LOD selection 
 	OutInstance->Debug.LODPredictedIndex = In.LODPredictedIndex;
 	OutInstance->HairGroupPublicData->DebugScreenSize = In.DebugScreenSize;
+	OutInstance->HairGroupPublicData->ContinuousLODPreviousPointCount = OutInstance->HairGroupPublicData->ContinuousLODPointCount;
+	OutInstance->HairGroupPublicData->ContinuousLODPreviousCurveCount = OutInstance->HairGroupPublicData->ContinuousLODCurveCount;
 	OutInstance->HairGroupPublicData->ContinuousLODPointCount = In.ContinuousLODPointCount;
 	OutInstance->HairGroupPublicData->ContinuousLODCurveCount = In.ContinuousLODCurveCount;
 	OutInstance->HairGroupPublicData->ContinuousLODScreenSize = In.ContinuousLODScreenSize;
diff --git a/Engine/Plugins/Runtime/HairStrands/Source/HairStrandsCore/Private/HairStrandsInterpolation.cpp b/Engine/Plugins/Runtime/HairStrands/Source/HairStrandsCore/Private/HairStrandsInterpolation.cpp
--- a/Engine/Plugins/Runtime/HairStrands/Source/HairStrandsCore/Private/HairStrandsInterpolation.cpp
+++ b/Engine/Plugins/Runtime/HairStrands/Source/HairStrandsCore/Private/HairStrandsInterpolation.cpp
@@ -177,9 +177,11 @@ class FTransferVelocityPassCS : public FGlobalShader
 	SHADER_USE_PARAMETER_STRUCT(FTransferVelocityPassCS, FGlobalShader);
 
 	BEGIN_SHADER_PARAMETER_STRUCT(FParameters, )
-		SHADER_PARAMETER(uint32, ElementCount)
-		SHADER_PARAMETER_RDG_BUFFER_SRV(Buffer, InBuffer)
-		SHADER_PARAMETER_RDG_BUFFER_UAV(RWBuffer, OutBuffer)
+		SHADER_PARAMETER(uint32, PointCount)
+		SHADER_PARAMETER(uint32, PointOffset)
+		SHADER_PARAMETER(uint32, TotalPointCount)
+		SHADER_PARAMETER_RDG_BUFFER_SRV(ByteAddressBuffer, InBuffer)
+		SHADER_PARAMETER_RDG_BUFFER_UAV(RWByteAddressBuffer, OutBuffer)
 	END_SHADER_PARAMETER_STRUCT()
 
 public:
@@ -198,18 +200,22 @@ IMPLEMENT_GLOBAL_SHADER(FTransferVelocityPassCS, "/Engine/Private/HairStrands/Ha
 void AddTransferPositionPass(
 	FRDGBuilder& GraphBuilder,
 	FGlobalShaderMap* ShaderMap,
-	const uint32 ElementCount,
+	const uint32 PointOffset,
+	const uint32 PointCount,
+	const uint32 TotalPointCount,
 	FRDGBufferSRVRef InBuffer,
 	FRDGBufferUAVRef OutBuffer)
 {
-	if (ElementCount == 0) return;
+	if (PointCount == 0 || TotalPointCount == 0) return;
 
 	FTransferVelocityPassCS::FParameters* Parameters = GraphBuilder.AllocParameters<FTransferVelocityPassCS::FParameters>();
-	Parameters->ElementCount = ElementCount;
+	Parameters->PointCount = PointCount;
+	Parameters->PointOffset = PointOffset;
+	Parameters->TotalPointCount = TotalPointCount;
 	Parameters->InBuffer = InBuffer;
 	Parameters->OutBuffer = OutBuffer;
 
-	const FIntVector DispatchCount(FMath::DivideAndRoundUp(ElementCount, FTransferVelocityPassCS::GetGroupSize()), 1, 1);
+	const FIntVector DispatchCount(FMath::DivideAndRoundUp(PointCount, FTransferVelocityPassCS::GetGroupSize()), 1, 1);
 	TShaderMapRef<FTransferVelocityPassCS> ComputeShader(ShaderMap);
 	FComputeShaderUtils::AddPass(
 		GraphBuilder,
diff --git a/Engine/Plugins/Runtime/HairStrands/Source/HairStrandsCore/Public/HairStrandsInterpolation.h b/Engine/Plugins/Runtime/HairStrands/Source/HairStrandsCore/Public/HairStrandsInterpolation.h
--- a/Engine/Plugins/Runtime/HairStrands/Source/HairStrandsCore/Public/HairStrandsInterpolation.h
+++ b/Engine/Plugins/Runtime/HairStrands/Source/HairStrandsCore/Public/HairStrandsInterpolation.h
@@ -183,7 +183,9 @@ void AddPatchAttributePass(
 void AddTransferPositionPass(
 	FRDGBuilder& GraphBuilder,
 	FGlobalShaderMap* ShaderMap,
-	const uint32 ElementCount,
+	const uint32 PointOffset,
+	const uint32 PointCount,
+	const uint32 TotalPointCount,
 	FRDGBufferSRVRef InBuffer,
 	FRDGBufferUAVRef OutBuffer);
 
diff --git a/Engine/Shaders/Private/HairStrands/HairStrandsInterpolation.usf b/Engine/Shaders/Private/HairStrands/HairStrandsInterpolation.usf
--- a/Engine/Shaders/Private/HairStrands/HairStrandsInterpolation.usf
+++ b/Engine/Shaders/Private/HairStrands/HairStrandsInterpolation.usf
@@ -898,19 +898,20 @@ void MainCS(uint2 DispatchThreadId : SV_DispatchThreadID, uint GroupThread1D : S
 
 #if SHADER_HAIRTRANSFER_PREV_POSITION
 
-uint ElementCount;
+uint PointCount;
+uint PointOffset;
+uint TotalPointCount;
 ByteAddressBuffer InBuffer;
 RWByteAddressBuffer OutBuffer;
 
 [numthreads(GROUP_SIZE, 1, 1)]
 void MainCS(uint DispatchThreadId : SV_DispatchThreadID)
 {
-	const uint Index = DispatchThreadId;
-	if (Index >= ElementCount)
+	const uint PointIndex = DispatchThreadId + PointOffset;
+	if (DispatchThreadId < PointCount && PointIndex < TotalPointCount)
 	{
-		return;
+		CopyPackedHairControlPoint(OutBuffer, InBuffer, PointIndex);
 	}
-	CopyPackedHairControlPoint(OutBuffer, InBuffer, Index);
 }
 
 #endif // SHADER_HAIRTRANSFER_PREV_POSITION
diff --git a/Engine/Source/Runtime/Renderer/Private/HairStrands/HairStrandsInterface.cpp b/Engine/Source/Runtime/Renderer/Private/HairStrands/HairStrandsInterface.cpp
--- a/Engine/Source/Runtime/Renderer/Private/HairStrands/HairStrandsInterface.cpp
+++ b/Engine/Source/Runtime/Renderer/Private/HairStrands/HairStrandsInterface.cpp
@@ -331,14 +331,14 @@ bool IsHairVisibilityComputeRasterContinuousLODEnabled()
 	return IsHairStrandContinuousDecimationReorderingEnabled() && IsHairVisibilityComputeRasterEnabled() && (CVarHairStrandsVisibilityComputeRaster_ContinuousLOD.GetValueOnAnyThread() > 0);
 }
 
-uint32 FHairGroupPublicData::GetActiveStrandsPointCount() const
+uint32 FHairGroupPublicData::GetActiveStrandsPointCount(bool bPrevious) const
 {
-	return ContinuousLODPointCount;
+	return bPrevious ? ContinuousLODPreviousPointCount : ContinuousLODPointCount;
 }
 
-uint32 FHairGroupPublicData::GetActiveStrandsCurveCount() const
+uint32 FHairGroupPublicData::GetActiveStrandsCurveCount(bool bPrevious) const
 {
-	return ContinuousLODCurveCount;
+	return bPrevious ? ContinuousLODPreviousCurveCount : ContinuousLODCurveCount;
 }
 
 float FHairGroupPublicData::GetActiveStrandsCoverageScale() const
diff --git a/Engine/Source/Runtime/Renderer/Public/HairStrandsInterface.h b/Engine/Source/Runtime/Renderer/Public/HairStrandsInterface.h
--- a/Engine/Source/Runtime/Renderer/Public/HairStrandsInterface.h
+++ b/Engine/Source/Runtime/Renderer/Public/HairStrandsInterface.h
@@ -287,8 +287,8 @@ public:
 	bool GetLODVisibility() const { return bLODVisibility; }
 
 	// Return the number of active point/curve for strand geometry
-	RENDERER_API uint32 GetActiveStrandsPointCount() const;
-	RENDERER_API uint32 GetActiveStrandsCurveCount() const;
+	RENDERER_API uint32 GetActiveStrandsPointCount(bool bPrevious=false) const;
+	RENDERER_API uint32 GetActiveStrandsCurveCount(bool bPrevious=false) const;
 	RENDERER_API float  GetActiveStrandsCoverageScale() const;
 	RENDERER_API float  GetActiveStrandsRadiusScale() const;
 
@@ -398,6 +398,8 @@ public:
 	// Active/used point/curved based on select continuous LOD
 	uint32 ContinuousLODPointCount = 0;
 	uint32 ContinuousLODCurveCount = 0;
+	uint32 ContinuousLODPreviousPointCount = 0;
+	uint32 ContinuousLODPreviousCurveCount = 0;
 	float ContinuousLODScreenSize = 1.f;
 	float ContinuousLODCoverageScale = 1.f;
 	float ContinuousLODRadiusScale = 1.f;
