Fix the FScopedSkeletalMeshRenderData which now tick the FSkinnedAssetCompilingManager after the event trigger in the destructor. This fix the issue, but the groom code should be change because keeping the lock when doing general code to create hair and groom is dangerous. A skeletal mesh api should be create to return a copy of the render data in place of the FScopedSkeletalMeshRenderData which should be remove. The new API will lock the skeletalmesh only for the copy time.

Also fix a crash when removing LODs if the lodsettings do not contain valid LOD information for the LOD index we remove.

#rnx

diff --git a/Engine/Plugins/Runtime/HairStrands/Source/HairStrandsCore/Private/GroomBindingBuilder.cpp b/Engine/Plugins/Runtime/HairStrands/Source/HairStrandsCore/Private/GroomBindingBuilder.cpp
--- a/Engine/Plugins/Runtime/HairStrands/Source/HairStrandsCore/Private/GroomBindingBuilder.cpp
+++ b/Engine/Plugins/Runtime/HairStrands/Source/HairStrandsCore/Private/GroomBindingBuilder.cpp
@@ -2191,8 +2191,10 @@ static bool InternalBuildBinding_CPU(const FGroomBindingBuilder::FInput& In, uin
 		// * Then use an async build scope to allow accessing skeletal mesh property safely.
 		//   If skel.meshes are nullptr, this will act as a NOP
 		USkeletalMesh* InSourceSkeletalMesh = In.SourceSkeletalMesh == In.TargetSkeletalMesh ? nullptr : In.SourceSkeletalMesh;
+		PRAGMA_DISABLE_DEPRECATION_WARNINGS
 		FScopedSkeletalMeshRenderData SourceSkeletalMeshScopedData(InSourceSkeletalMesh);
 		FScopedSkeletalMeshRenderData TargetSkeletalMeshScopedData(In.TargetSkeletalMesh);
+		PRAGMA_ENABLE_DEPRECATION_WARNINGS
 
 		TUniquePtr<GroomBinding_Mesh::IMeshData> SourceMeshData;
 		TUniquePtr<GroomBinding_Mesh::IMeshData> TargetMeshData;
@@ -2202,7 +2204,9 @@ static bool InternalBuildBinding_CPU(const FGroomBindingBuilder::FInput& In, uin
 			{
 				FSkinnedAssetAsyncBuildScope AsyncBuildScope(InSourceSkeletalMesh);
 				USkeletalMesh::GetPlatformSkeletalMeshRenderData(TargetPlatform, SourceSkeletalMeshScopedData);
+				PRAGMA_DISABLE_DEPRECATION_WARNINGS
 				SourceMeshData = TUniquePtr<GroomBinding_Mesh::FSkeletalMeshData, TDefaultDelete<GroomBinding_Mesh::IMeshData>>(new GroomBinding_Mesh::FSkeletalMeshData(InSourceSkeletalMesh, SourceSkeletalMeshScopedData.GetData()));
+				PRAGMA_ENABLE_DEPRECATION_WARNINGS
 			}
 			else
 			{
@@ -2213,7 +2217,9 @@ static bool InternalBuildBinding_CPU(const FGroomBindingBuilder::FInput& In, uin
 			{
 				FSkinnedAssetAsyncBuildScope AsyncBuildScope(In.TargetSkeletalMesh);
 				USkeletalMesh::GetPlatformSkeletalMeshRenderData(TargetPlatform, TargetSkeletalMeshScopedData);
+				PRAGMA_DISABLE_DEPRECATION_WARNINGS
 				TargetMeshData = TUniquePtr<GroomBinding_Mesh::FSkeletalMeshData, TDefaultDelete<GroomBinding_Mesh::IMeshData>>(new GroomBinding_Mesh::FSkeletalMeshData(In.TargetSkeletalMesh, TargetSkeletalMeshScopedData.GetData()));
+				PRAGMA_ENABLE_DEPRECATION_WARNINGS
 			}
 			else
 			{
diff --git a/Engine/Plugins/Runtime/HairStrands/Source/HairStrandsCore/Private/GroomRBFDeformer.cpp b/Engine/Plugins/Runtime/HairStrands/Source/HairStrandsCore/Private/GroomRBFDeformer.cpp
--- a/Engine/Plugins/Runtime/HairStrands/Source/HairStrandsCore/Private/GroomRBFDeformer.cpp
+++ b/Engine/Plugins/Runtime/HairStrands/Source/HairStrandsCore/Private/GroomRBFDeformer.cpp
@@ -664,7 +664,9 @@ void FGroomRBFDeformer::GetRBFDeformedGroomAsset(const UGroomAsset* InGroomAsset
 		// identical skel. mesh render data and ensure that the hair deformation is done correctly.
 		check(BindingAsset->GetGroomBindingType() == EGroomBindingMeshType::SkeletalMesh);
 		USkeletalMesh* TargetSkeletalMesh = BindingAsset->GetTargetSkeletalMesh();
+		PRAGMA_DISABLE_DEPRECATION_WARNINGS
 		FScopedSkeletalMeshRenderData TargetSkeletalMeshScopedData(TargetSkeletalMesh);
+		PRAGMA_ENABLE_DEPRECATION_WARNINGS
 		{
 			ITargetPlatform* RunningPlatform = GetTargetPlatformManagerRef().GetRunningTargetPlatform();
 			FSkinnedAssetAsyncBuildScope AsyncBuildScope(TargetSkeletalMesh);
@@ -675,7 +677,9 @@ void FGroomRBFDeformer::GetRBFDeformedGroomAsset(const UGroomAsset* InGroomAsset
 		const int32 MeshLODIndex = 0;
 
 		// Get the target mesh vertices (source and target)
+		PRAGMA_DISABLE_DEPRECATION_WARNINGS
 		const FSkeletalMeshRenderData* SkeletalMeshData_Target = TargetSkeletalMeshScopedData.GetData();
+		PRAGMA_ENABLE_DEPRECATION_WARNINGS
 		TArray<FVector3f> MeshVertexPositionsBuffer_Target;
 		ExtractSkeletalVertexPosition(SkeletalMeshData_Target, MeshLODIndex, MeshVertexPositionsBuffer_Target);
 
diff --git a/Engine/Source/Runtime/Engine/Classes/Engine/SkeletalMesh.h b/Engine/Source/Runtime/Engine/Classes/Engine/SkeletalMesh.h
--- a/Engine/Source/Runtime/Engine/Classes/Engine/SkeletalMesh.h
+++ b/Engine/Source/Runtime/Engine/Classes/Engine/SkeletalMesh.h
@@ -417,8 +417,12 @@ namespace NSSkeletalMeshSourceFileLabels
 struct ENGINE_API FScopedSkeletalMeshRenderData
 {
 public:
+	UE_DEPRECATED(5.4, "This class is deprecated, there is no replacement because this is not a safe async design.")
 	FScopedSkeletalMeshRenderData(USkeletalMesh* Mesh);
+	
 	~FScopedSkeletalMeshRenderData();
+
+	UE_DEPRECATED(5.4, "This class is deprecated, there is no replacement because this is not a safe async design.")
 	const FSkeletalMeshRenderData* GetData() const;
 private:
 	FEvent* Lock = nullptr;
diff --git a/Engine/Source/Runtime/Engine/Private/SkeletalMesh.cpp b/Engine/Source/Runtime/Engine/Private/SkeletalMesh.cpp
--- a/Engine/Source/Runtime/Engine/Private/SkeletalMesh.cpp
+++ b/Engine/Source/Runtime/Engine/Private/SkeletalMesh.cpp
@@ -1574,6 +1574,9 @@ FScopedSkeletalMeshRenderData::~FScopedSkeletalMeshRenderData()
 	{
 		check(Mesh);
 		Lock->Trigger();
+		//After we trigger the event we must tick the FSkinnedAssetCompilingManager so it clear the skeletal mesh AsyncTask and call
+		//FinishAsyncTaskInternal to terminate the LockPropertiesUntil
+		FSkinnedAssetCompilingManager::Get().FinishCompilation({ Mesh });
 	}
 
 	Data = nullptr;
@@ -2059,6 +2062,8 @@ FEvent* USkeletalMesh::LockPropertiesUntil()
 {
 	TRACE_CPUPROFILER_EVENT_SCOPE(USkeletalMesh::Import);
 	
+	check(IsInGameThread());
+	
 	FEvent* Event = FPlatformProcess::GetSynchEventFromPool();
 	check(Event);
 
@@ -3466,11 +3471,6 @@ void USkeletalMesh::FinishPostLoadInternal(FSkinnedAssetPostLoadContext& Context
 	// This scope allows us to use any locked properties without causing stalls
 	FSkinnedAssetAsyncBuildScope AsyncBuildScope(this);
 
-	if (Context.bHasCachedDerivedData)
-	{
-		PostMeshCached.Broadcast(this);
-	}
-
 	//Make sure unused cloth are unbind
 	if (GetMeshClothingAssets().Num() > 0)
 	{
@@ -3652,6 +3652,15 @@ void USkeletalMesh::FinishPostLoadInternal(FSkinnedAssetPostLoadContext& Context
 #endif
 
 	ReleaseAsyncProperty();
+#if WITH_EDITOR
+	if (Context.bHasCachedDerivedData)
+	{
+		//We must call PostMeshCached after:
+		// - The async properties are release
+		// - The init resource is done
+		PostMeshCached.Broadcast(this);
+	}
+#endif //WITH_EDITOR
 }
 
 #if WITH_EDITORONLY_DATA
