Fix incorrect tip tools hair

diff --git a/Engine/Plugins/Runtime/HairStrands/Source/HairStrandsCore/Public/GroomAsset.h b/Engine/Plugins/Runtime/HairStrands/Source/HairStrandsCore/Public/GroomAsset.h
--- a/Engine/Plugins/Runtime/HairStrands/Source/HairStrandsCore/Public/GroomAsset.h
+++ b/Engine/Plugins/Runtime/HairStrands/Source/HairStrandsCore/Public/GroomAsset.h
@@ -453,7 +453,7 @@ private:
 	UPROPERTY(EditAnywhere, Category = "HairLOD", meta = (DisplayName = "LOD Mode", ToolTip = "Define how LOD adapts curves & points for strands geometry. Auto: adapts the curve count based on screen coverage. Manual: use the discrete LOD created for each groups"))
 	EGroomLODMode LODMode = EGroomLODMode::Default;
 
-	UPROPERTY(EditAnywhere, Category = "HairLOD", meta = (DisplayName = "Auto LOD Bias", ToolTip = "When Auto LOD is selected, decrease the screen size at which curves reduction will occur.", ClampMin = "0", ClampMax = "1", UIMin = "0", UIMax = "1.0"))
+	UPROPERTY(EditAnywhere, Category = "HairLOD", meta = (DisplayName = "Auto LOD Bias", ToolTip = "When LOD mode is set to Auto, decrease the screen size at which curves reduction will occur.", ClampMin = "0", ClampMax = "1", UIMin = "0", UIMax = "1.0"))
 	float AutoLODBias = 0;
 
 public:
diff --git a/Engine/Plugins/Runtime/HairStrands/Source/HairStrandsCore/Public/GroomAssetInterpolation.h b/Engine/Plugins/Runtime/HairStrands/Source/HairStrandsCore/Public/GroomAssetInterpolation.h
--- a/Engine/Plugins/Runtime/HairStrands/Source/HairStrandsCore/Public/GroomAssetInterpolation.h
+++ b/Engine/Plugins/Runtime/HairStrands/Source/HairStrandsCore/Public/GroomAssetInterpolation.h
@@ -75,31 +75,30 @@ struct HAIRSTRANDSCORE_API FHairLODSettings
 	UPROPERTY(EditAnywhere, Category = "DecimationSettings", meta = (ToolTip = "Reduce the number of hair strands in a uniform manner", ClampMin = "0", UIMin = "0", UIMax = "1.0"))
 	float CurveDecimation = 1;
 
-	/** Reduce the number of vertices for each hair strands */
-	UPROPERTY(EditAnywhere, Category = "DecimationSettings", meta = (ToolTip = "Reduce the number of vertices per strands in a uniform manner", ClampMin = "0", UIMin = "0", UIMax = "1.0"))
+	/** Reduce the number of points for each hair strands */
+	UPROPERTY(EditAnywhere, Category = "DecimationSettings", meta = (ToolTip = " Reduce the number of points for each hair strands", ClampMin = "0", UIMin = "0", UIMax = "1.0"))
 	float VertexDecimation = 1;
 
 	/** Max angular difference between adjacents vertices to remove vertices during simplification, in degrees. */
-	UPROPERTY(EditAnywhere, Category = "DecimationSettings", meta = (ToolTip = "Reduce the number of vertices per strands in a uniform manner", ClampMin = "0", ClampMax = "45", UIMin = "0", UIMax = "45.0"))
+	UPROPERTY(EditAnywhere, Category = "DecimationSettings", meta = (ToolTip = "Max angular difference between adjacents vertices to remove vertices during simplification, in degrees", ClampMin = "0", ClampMax = "45", UIMin = "0", UIMax = "45.0"))
 	float AngularThreshold = 1.f;
 
 	/** Screen size at which this LOD should be enabled */
-	UPROPERTY(EditAnywhere, Category = "DecimationSettings", meta = (ToolTip = "Reduce the number of hair strands in a uniform manner", ClampMin = "0", UIMin = "0", UIMax = "1.0"))
+	UPROPERTY(EditAnywhere, Category = "DecimationSettings", meta = (ToolTip = "Screen size at which this LOD should be enabled", ClampMin = "0", UIMin = "0", UIMax = "1.0"))
 	float ScreenSize = 1;
 
-	/** Scale the hair thickness for compensating the reduction of curves. This thickness scale properties manually increases/decreases the thickness in addition to 
-	  * the automatic increase of thickness applied onto the curve when curve decimation is lower than 1. */
-	UPROPERTY(EditAnywhere, Category = "DecimationSettings", meta = (ToolTip = "Reduce the number of hair strands in a uniform manner", ClampMin = "0", UIMin = "0", UIMax = "1.0"))
+	/** Scales the hair Strands radius. This can be used for manually compensating the reduction of curves.. */
+	UPROPERTY(EditAnywhere, Category = "DecimationSettings", meta = (ToolTip = "Scales the hair Strands radius. This can be used for manually compensating the reduction of curves.", ClampMin = "0", UIMin = "0", UIMax = "1.0"))
 	float ThicknessScale = 1;
 
 	/** If true (default), the hair group is visible. If false, the hair group is culled. */
 	UPROPERTY(EditAnywhere, Category = "DecimationSettings", meta = (ToolTip = "If disable, the hair strands won't be rendered"))
 	bool bVisible = true;
 
-	UPROPERTY(EditAnywhere, Category = "DecimationSettings", meta = (ToolTip = "If enable this LOD version will use cards representation"))
+	UPROPERTY(EditAnywhere, Category = "DecimationSettings", meta = (ToolTip = "Defines the type of geometry used by this LOD (Strands, Cards, or Meshes)"))
 	EGroomGeometryType GeometryType = EGroomGeometryType::Strands;
 
-	UPROPERTY(EditAnywhere, Category = "DecimationSettings", AdvancedDisplay, meta = (ToolTip = "If enable this LOD version will use the provided attachment points"))
+	UPROPERTY(EditAnywhere, Category = "DecimationSettings", AdvancedDisplay, meta = (ToolTip = "Defines the type of attachment"))
 	EGroomBindingType BindingType = EGroomBindingType::Skinning;
 	
 	UPROPERTY(EditAnywhere, Category = "DecimationSettings", AdvancedDisplay, meta = (ToolTip = "Groom simulation"))
