Fix crash when reloading groom asset.

There is no longer need to reload groom resources, since this is now only done when derived data are updated.

diff --git a/Engine/Plugins/Runtime/HairStrands/Source/HairStrandsCore/Private/GroomAsset.cpp b/Engine/Plugins/Runtime/HairStrands/Source/HairStrandsCore/Private/GroomAsset.cpp
--- a/Engine/Plugins/Runtime/HairStrands/Source/HairStrandsCore/Private/GroomAsset.cpp
+++ b/Engine/Plugins/Runtime/HairStrands/Source/HairStrandsCore/Private/GroomAsset.cpp
@@ -1016,6 +1016,9 @@ void UGroomAsset::UpdateResource()
 		uint32 ChangeType = (CachedHairGroupsLOD[GroupIndex] == GetHairGroupsLOD()[GroupIndex] ? 0 : GroomChangeType_LOD);
 		AllChangeType = AllChangeType | ChangeType;
 
+		// No need to reload the resources as this needs only be done when the derived data are update, 
+		// which is already called in such a case.
+		#if 0
 		check(GroupIndex < uint32(GetNumHairGroups()));
 		FHairGroupResources& GroupData = GetHairGroupsResources()[GroupIndex];
 		if (GroupData.Guides.RestResource)
@@ -1037,6 +1040,7 @@ void UGroomAsset::UpdateResource()
 			}
 			#endif
 		}
+		#endif
 	}
 
 	if (BuildCardsData())  { AllChangeType |= GroomChangeType_Cards;  }
