Fix hair LOD ThicknessScale not accurately taken into account when running with Auto LOD mode.

Auto LOD is turned on by default in 5.4, and for legacy asset which uses thickness scale, there is a small discrepancy without that fix.

diff --git a/Engine/Plugins/Runtime/HairStrands/Source/HairStrandsCore/Private/GroomManager.cpp b/Engine/Plugins/Runtime/HairStrands/Source/HairStrandsCore/Private/GroomManager.cpp
--- a/Engine/Plugins/Runtime/HairStrands/Source/HairStrandsCore/Private/GroomManager.cpp
+++ b/Engine/Plugins/Runtime/HairStrands/Source/HairStrandsCore/Private/GroomManager.cpp
@@ -2525,6 +2525,7 @@ static FHairLOD ComputeHairLODIndex(const FHairGroupInstance* Instance, const TA
 
 	// Auto LOD
 	const bool bNeedAutoLOD = MaxLODIndexWithStrands != INDEX_NONE && int32(Out.HairLODIndex) <= MaxLODIndexWithStrands;
+	const float OriginalHairLODIndex = Out.HairLODIndex;
 	if (Instance->Strands.ClusterResource && Instance->Strands.RestResource)
 	{
 		uint32 EffectiveCurveCount = 0;
@@ -2542,7 +2543,7 @@ static FHairLOD ComputeHairLODIndex(const FHairGroupInstance* Instance, const TA
 		Out.ContinuousLODCurveCount = EffectiveCurveCount;
 		Out.ContinuousLODPointCount = EffectiveCurveCount > 0 ? Instance->Strands.GetData().Header.CurveToPointCount[EffectiveCurveCount - 1] : 0;
 		Out.ContinuousLODCoverageScale = ComputeActiveCurveCoverageScale(Instance->Strands.GetData(), EffectiveCurveCount);
-		Out.ContinuousLODRadiusScale = ComputeActiveCurveRadiusScale(Instance->Strands.ClusterResource, Out.HairLODIndex);
+		Out.ContinuousLODRadiusScale = ComputeActiveCurveRadiusScale(Instance->Strands.ClusterResource, OriginalHairLODIndex); // Use OriginalHairLODIndex to query the correct CurveRadiusScale
 	}
 
 	return Out;
