Fix incorrect permutation skin.

diff --git a/Engine/Plugins/Runtime/HairStrands/Source/HairStrandsCore/Private/HairStrandsMeshProjection.cpp b/Engine/Plugins/Runtime/HairStrands/Source/HairStrandsCore/Private/HairStrandsMeshProjection.cpp
--- a/Engine/Plugins/Runtime/HairStrands/Source/HairStrandsCore/Private/HairStrandsMeshProjection.cpp
+++ b/Engine/Plugins/Runtime/HairStrands/Source/HairStrandsCore/Private/HairStrandsMeshProjection.cpp
@@ -108,7 +108,7 @@ void AddSkinUpdatePass(
 	PermutationVector.Set<FSkinUpdateCS::FUnlimitedBoneInfluence>(SkinWeight->GetBoneInfluenceType() == GPUSkinBoneInfluenceType::UnlimitedBoneInfluence);
 	PermutationVector.Set<FSkinUpdateCS::FUseExtraInfluence>(SkinWeight->GetMaxBoneInfluences() > MAX_INFLUENCES_PER_STREAM);
 	PermutationVector.Set<FSkinUpdateCS::FBoneIndexUint16 >(SkinWeight->Use16BitBoneIndex());
-	PermutationVector.Set<FSkinUpdateCS::FBoneIndexUint16 >(SkinWeight->Use16BitBoneWeight());
+	PermutationVector.Set<FSkinUpdateCS::FBoneWeightUint16 >(SkinWeight->Use16BitBoneWeight());
 	PermutationVector.Set<FSkinUpdateCS::FPrevious>(bPrevPosition);
 	TShaderMapRef<FSkinUpdateCS> ComputeShader(ShaderMap, PermutationVector);
 
