* Fix incorrect local bound when creating dynamic instance uniform buffer. The local bound and (overriden) local to world transform were not expressed in the same space (original vs. overridden).
* Fix incorrect bound when strand data are available but trimmed (i.e., groom contains a single strand curve). In such case, the bounds are computing using the strands data header, misrepresentative the groom's size. In such a case, the available cards or meshes bounds are use to compute the actual bound size.

diff --git a/Engine/Plugins/Runtime/HairStrands/Source/HairStrandsCore/Private/GroomComponent.cpp b/Engine/Plugins/Runtime/HairStrands/Source/HairStrandsCore/Private/GroomComponent.cpp
--- a/Engine/Plugins/Runtime/HairStrands/Source/HairStrandsCore/Private/GroomComponent.cpp
+++ b/Engine/Plugins/Runtime/HairStrands/Source/HairStrandsCore/Private/GroomComponent.cpp
@@ -1008,10 +1008,22 @@ public:
 			FPrimitiveUniformShaderParametersBuilder Builder;
 			BuildUniformShaderParameters(Builder);
 
-			// Override transforms
+			// Override transforms and the local bound.
+			// The original local bound relative to the component local to world transform. If we override the local to world transform, 
+			// we need to recompute the local bound relative to this new transform. It is important that the new local bound is correct 
+			// as otherwise the GPUScene (which use bot the local to world transform and the local bound for culling purpose) will issue 
+			// incorrect visibility test.
+			FBoxSphereBounds NewLocalBound = GetLocalBounds();
+			if (!bUseProxy)
+			{
+				const FTransform InvLocalToWorld = Instance->GetCurrentLocalToWorld().Inverse();
+				const FBoxSphereBounds OriginalWorldBound = GetBounds();
+				NewLocalBound = OriginalWorldBound.TransformBy(InvLocalToWorld);
+			}
 			Builder
 				.LocalToWorld(CurrentLocalToWorld)
 				.PreviousLocalToWorld(PreviousLocalToWorld)
+				.LocalBounds(NewLocalBound)
 				.OutputVelocity(bOutputVelocity)
 				.UseVolumetricLightmap(false);
 
@@ -1939,6 +1951,17 @@ FBoxSphereBounds UGroomComponent::CalcBounds(const FTransform& InLocalToWorld) c
 					const FVector HairBorder(0.5f * Desc.HairWidth * FMath::Max3(1.0f, Desc.HairRootScale, Desc.HairTipScale) * InvScale);
 
 					LocalHairBound += GroupData.Strands.GetBounds();
+
+					// If the strands data only contain a single curve, it likely means that the strands data have been trimmed 
+					// and their bounds incorrectly represent the groom component. In such a case, we optionnally evaluate/include 
+					// the cards/meshes bounds
+					const bool bNeedCardsOrMeshesBound = GroupData.Strands.BulkData.Header.CurveCount <= 1;
+					if (bNeedCardsOrMeshesBound)
+					{
+						if (GroupData.Cards.HasValidData())  { LocalHairBound += GroupData.Cards.GetBounds(); }
+						if (GroupData.Meshes.HasValidData()) { LocalHairBound += GroupData.Meshes.GetBounds(); }
+					}
+
 					LocalHairBound.Min -= HairBorder;
 					LocalHairBound.Max += HairBorder;
 				}
