Reset groom simulation when Mesh LOD is changing

diff --git a/Engine/Plugins/Runtime/HairStrands/Source/HairStrandsCore/Private/Niagara/NiagaraDataInterfaceHairStrands.cpp b/Engine/Plugins/Runtime/HairStrands/Source/HairStrandsCore/Private/Niagara/NiagaraDataInterfaceHairStrands.cpp
--- a/Engine/Plugins/Runtime/HairStrands/Source/HairStrandsCore/Private/Niagara/NiagaraDataInterfaceHairStrands.cpp
+++ b/Engine/Plugins/Runtime/HairStrands/Source/HairStrandsCore/Private/Niagara/NiagaraDataInterfaceHairStrands.cpp
@@ -3352,7 +3352,8 @@ void UNiagaraDataInterfaceHairStrands::SetShaderParameters(const FNiagaraDataInt
 	FNDIHairStrandsProxy& DIProxy = Context.GetProxy<FNDIHairStrandsProxy>();
 	FNDIHairStrandsData* ProxyData = DIProxy.SystemInstancesToProxyData.Find(Context.GetSystemInstanceID());
 
-	const int32 MeshLODIndex = ProxyData && ProxyData->HairGroupInstance ? ProxyData->HairGroupInstance->Debug.MeshLODIndex : -1;
+	// TODO: refactor all of that with if condition
+	const int32 MeshLODIndex = ProxyData && ProxyData->HairGroupInstance && ProxyData->HairGroupInstance->HairGroupPublicData ? ProxyData->HairGroupInstance->HairGroupPublicData->GetMeshLODIndex() : -1;
 	const bool bIsHairValid = ProxyData != nullptr && ProxyData->HairStrandsBuffer && ProxyData->HairStrandsBuffer->IsInitialized();
 	const bool bIsHairGroupInstValid = ProxyData != nullptr && ProxyData->HairGroupInstSource != nullptr && ProxyData->HairGroupInstSource->ContainsGroupInstance(ProxyData->HairGroupInstance);
 	const bool bHasSkinningBinding = bIsHairValid && bIsHairGroupInstValid && ProxyData->HairGroupInstance && ProxyData->HairGroupInstance->BindingType == EHairBindingType::Skinning;
@@ -3418,14 +3419,15 @@ void UNiagaraDataInterfaceHairStrands::SetShaderParameters(const FNiagaraDataInt
 
 		FRDGBufferSRVRef RestSamplePositionsBufferSRV = (bHasSamples && RestMeshProjection) ? RegisterAsSRV(GraphBuilder,RestMeshProjection->RestSamplePositionsBuffer) : DummyStructuredBufferSRV;
 		FRDGBufferSRVRef MeshSampleWeightsBufferSRV = (bHasSamples && DeformedMeshProjection) ? RegisterAsSRV(GraphBuilder,DeformedMeshProjection->GetMeshSampleWeightsBuffer((FHairStrandsLODDeformedRootResource::Current))) : DummyStructuredBufferSRV;
-		
+
 		// Simulation setup (we update the rest configuration based on the deformed positions 
-		// if in restupdate mode or if we are resetting the sim and using RBF transfer since the rest positions are not matrching the physics asset)
-		const int32 NeedResetValue = (ProxyData->TickCount <= GHairSimulationMaxDelay) || !HairStrandsBuffer->bValidGeometryType;
+		// if in restupdate mode or if we are resetting the sim and using RBF transfer since the rest positions are not matching the physics asset)
+		const int32 NeedResetValue = (ProxyData->TickCount <= GHairSimulationMaxDelay) || !HairStrandsBuffer->bValidGeometryType || (HairStrandsBuffer->CurrentMeshLOD != MeshLODIndex);
 		const int32 RestUpdateValue = GHairSimulationRestUpdate || (NeedResetValue && ProxyData->bSkinningTransfer) ? 1 : 0;
 		const int32 LocalSimulationValue = ProxyData->LocalSimulation;
-
+		
 		HairStrandsBuffer->bValidGeometryType = true;
+		HairStrandsBuffer->CurrentMeshLOD = MeshLODIndex;
 
 		// Offsets / Transforms
 		FVector3f RestPositionOffsetValue = (FVector3f)ProxyData->HairStrandsBuffer->SourceRestResources->GetPositionOffset();
diff --git a/Engine/Plugins/Runtime/HairStrands/Source/HairStrandsCore/Public/Niagara/NiagaraDataInterfaceHairStrands.h b/Engine/Plugins/Runtime/HairStrands/Source/HairStrandsCore/Public/Niagara/NiagaraDataInterfaceHairStrands.h
--- a/Engine/Plugins/Runtime/HairStrands/Source/HairStrandsCore/Public/Niagara/NiagaraDataInterfaceHairStrands.h
+++ b/Engine/Plugins/Runtime/HairStrands/Source/HairStrandsCore/Public/Niagara/NiagaraDataInterfaceHairStrands.h
@@ -81,6 +81,9 @@ struct FNDIHairStrandsBuffer : public FRenderResource
 	/** Valid geometry type for hair (strands, cards, mesh)*/
 	bool bValidGeometryType = false;
 
+	/** Mesh LOD that is being used for the root resources */
+	int32 CurrentMeshLOD = INDEX_NONE;
+
 	// For debug only
 	//FRHIGPUBufferReadback* ReadbackBuffer = nullptr;
 };
