Fix groom exporting when trimming LODs.

* Reset Groom data prior to rebuild them to ensure there is no stale data. This fixes crashes when deleting cards/meshes LOD and calling CacheDerivedData().
* Fix missing binding reloading on property change.
** When a groom asset change it derived data cache, it only call the update resource deleague. This deleague was not listen by the groom binding asset. This causes mismatching between groom binding and groom asset, causing the binding to not work/be validated.
** Fix groom component reloading when binding data update. The groom component was happening too early prior to the binding resources and infos being update, which would cause the binding validation to fail down the line.
* Fix incorrect buffer type in shader declaraction.

diff --git a/Engine/Plugins/Runtime/HairStrands/Source/HairStrandsCore/Private/GroomAsset.cpp b/Engine/Plugins/Runtime/HairStrands/Source/HairStrandsCore/Private/GroomAsset.cpp
--- a/Engine/Plugins/Runtime/HairStrands/Source/HairStrandsCore/Private/GroomAsset.cpp
+++ b/Engine/Plugins/Runtime/HairStrands/Source/HairStrandsCore/Private/GroomAsset.cpp
@@ -2606,6 +2606,7 @@ bool UGroomAsset::CacheCardsData(uint32 GroupIndex, const FString& StrandsKey)
 
 		// Initialized LODs
 		const uint32 LODCount = GetHairGroupsLOD()[GroupIndex].LODs.Num();
+		GetHairGroupsPlatformData()[GroupIndex].Cards.LODs.Empty(); // Force reset data
 		GetHairGroupsPlatformData()[GroupIndex].Cards.LODs.SetNum(LODCount);
 
 		return bNeedReset;
@@ -2685,6 +2686,7 @@ bool UGroomAsset::CacheMeshesData(uint32 GroupIndex)
 
 		// Initialized LODs
 		const uint32 LODCount = GetHairGroupsLOD()[GroupIndex].LODs.Num();
+		GetHairGroupsPlatformData()[GroupIndex].Meshes.LODs.Empty(); // Force reset data
 		GetHairGroupsPlatformData()[GroupIndex].Meshes.LODs.SetNum(LODCount);
 
 		return bNeedReset;
@@ -2862,17 +2864,12 @@ bool UGroomAsset::BuildHairGroup_Cards(uint32 GroupIndex)
 	FHairGroupPlatformData& GroupData = GetHairGroupsPlatformData()[GroupIndex];
 
 	bool bDataBuilt = false;
+	GroupData.Cards.LODs.Empty(); // Force reset data
 	GroupData.Cards.LODs.SetNum(LODCount);
 	for (uint32 LODIt = 0; LODIt < LODCount; ++LODIt)
 	{
 		if (FHairGroupsCardsSourceDescription* Desc = GetSourceDescription(GetHairGroupsCards(), GroupIndex, LODIt))
 		{
-			if (bIsAlreadyBuilt[LODIt])
-			{
-				bDataBuilt = true;
-				continue;
-			}
-
 			FHairGroupPlatformData::FCards::FLOD& LOD = GroupData.Cards.LODs[LODIt];
 			LOD.BulkData.Reset();
 
@@ -3016,20 +3013,13 @@ bool UGroomAsset::BuildHairGroup_Meshes(uint32 GroupIndex)
 	const uint32 LODCount = GetHairGroupsLOD()[GroupIndex].LODs.Num();
 
 	bool bDataBuilt = false;
+	GroupData.Meshes.LODs.Empty(); // Force reset data
 	GroupData.Meshes.LODs.SetNum(LODCount);
 	for (uint32 LODIt = 0; LODIt < LODCount; ++LODIt)
 	{
 		int32 SourceIt = 0;
 		if (const FHairGroupsMeshesSourceDescription* Desc = GetSourceDescription(GetHairGroupsMeshes(), GroupIndex, LODIt, &SourceIt))
 		{
-			const FHairGroupsMeshesSourceDescription* CachedDesc = SourceIt < CachedHairGroupsMeshes.Num() ? &CachedHairGroupsMeshes[SourceIt] : nullptr;
-			const bool bLODHasChanged = CachedDesc == nullptr || !(*CachedDesc == *Desc);
-			if (!bLODHasChanged && bIsAlreadyBuilt[LODIt])
-			{
-				bDataBuilt |= bIsAlreadyBuilt[LODIt];
-				continue;
-			}
-
 			if (!IsInGameThread())
 			{
 				// Build needs to execute from the game thread
@@ -3860,4 +3850,4 @@ GROOMASSET_DEFINE_MEMBER_NAME(AutoLODBias)
 float UGroomAsset::GetAutoLODBias() const
 {
 	return AutoLODBias;
-}
\ No newline at end of file
+}
diff --git a/Engine/Plugins/Runtime/HairStrands/Source/HairStrandsCore/Private/GroomBindingAsset.cpp b/Engine/Plugins/Runtime/HairStrands/Source/HairStrandsCore/Private/GroomBindingAsset.cpp
--- a/Engine/Plugins/Runtime/HairStrands/Source/HairStrandsCore/Private/GroomBindingAsset.cpp
+++ b/Engine/Plugins/Runtime/HairStrands/Source/HairStrandsCore/Private/GroomBindingAsset.cpp
@@ -375,6 +375,7 @@ void UGroomBindingAsset::PostLoad()
 		if (LocalGroom)
 		{
 			LocalGroom->GetOnGroomAssetChanged().AddUObject(this, &UGroomBindingAsset::InvalidateBinding);
+			LocalGroom->GetOnGroomAssetResourcesChanged().AddUObject(this, &UGroomBindingAsset::InvalidateBinding);
 			bRegisterGroomAssetCallback = true;
 		}
 	#endif
@@ -440,6 +441,7 @@ void UGroomBindingAsset::BeginDestroy()
 
 	if (GetGroom() && bRegisterGroomAssetCallback)
 	{
+		GetGroom()->GetOnGroomAssetResourcesChanged().RemoveAll(this);
 		GetGroom()->GetOnGroomAssetChanged().RemoveAll(this);
 		bRegisterGroomAssetCallback = false;
 	}
@@ -902,16 +904,24 @@ void UGroomBindingAsset::CacheDerivedDatas()
 			bIsValid = bIsValid && bGroupValid;
 			bReloadResource = bReloadResource || bGroupReloadResource;
 		}
-	}
 
-	// 4. Reload resources if needed
-	if (bReloadResource)
+		// 3. Update binding infos here as they need to be valid when RecreateRenderContext is deleted
+		//    When RecreateRenderContext's Dtor is called, it will recreate component, which will run 
+		//    the binding validation to assess if the binding asset is compatible. This validation logic 
+		//    use the binding infos to know if curve count match between GroomAsset and GroomBindingAsset
+		UpdateGroomBindingAssetInfos(this);
+
+		// 4. Reload resources if needed
+		if (bReloadResource)
+		{
+			InitResource();
+		}
+	}
+	else
 	{
-		InitResource();
+		// 3. Patch hair group info if it does not match the DDC-read/deserialized data
+		UpdateGroomBindingAssetInfos(this);
 	}
-
-	// 3. Patch hair group info if it does not match the DDC-read/deserialized data
-	UpdateGroomBindingAssetInfos(this);
 }
 
 static void CacheDerivedDatas(UGroomBindingAsset* In, const uint32 InGroupIndex, const FString& DerivedDataKey, bool& bOutValid, const ITargetPlatform* TargetPlatform, UGroomBindingAsset::FHairGroupPlatformData& OutPlatformData)
diff --git a/Engine/Plugins/Runtime/HairStrands/Source/HairStrandsCore/Private/GroomBuilder.cpp b/Engine/Plugins/Runtime/HairStrands/Source/HairStrandsCore/Private/GroomBuilder.cpp
--- a/Engine/Plugins/Runtime/HairStrands/Source/HairStrandsCore/Private/GroomBuilder.cpp
+++ b/Engine/Plugins/Runtime/HairStrands/Source/HairStrandsCore/Private/GroomBuilder.cpp
@@ -51,7 +51,7 @@ bool DoesHairStrandsSupportCompressedPosition();
 
 FString FGroomBuilder::GetVersion()
 {
-	return TEXT("v15d");
+	return TEXT("v16");
 }
 
 namespace GroomBuilder_Voxelization
diff --git a/Engine/Plugins/Runtime/HairStrands/Source/HairStrandsCore/Private/HairStrandsMeshProjection.cpp b/Engine/Plugins/Runtime/HairStrands/Source/HairStrandsCore/Private/HairStrandsMeshProjection.cpp
--- a/Engine/Plugins/Runtime/HairStrands/Source/HairStrandsCore/Private/HairStrandsMeshProjection.cpp
+++ b/Engine/Plugins/Runtime/HairStrands/Source/HairStrandsCore/Private/HairStrandsMeshProjection.cpp
@@ -994,8 +994,8 @@ private:
 		SHADER_PARAMETER(uint32, InstanceRegisteredIndex)
 		SHADER_PARAMETER(FVector3f, CPUCurrPositionOffset)
 		SHADER_PARAMETER(FVector3f, CPUPrevPositionOffset)
-		SHADER_PARAMETER_RDG_BUFFER_SRV(StructuredBuffer, RootTriangleCurrPositionBuffer)
-		SHADER_PARAMETER_RDG_BUFFER_SRV(StructuredBuffer, RootTrianglePrevPositionBuffer)
+		SHADER_PARAMETER_RDG_BUFFER_SRV(Buffer, RootTriangleCurrPositionBuffer)
+		SHADER_PARAMETER_RDG_BUFFER_SRV(Buffer, RootTrianglePrevPositionBuffer)
 		SHADER_PARAMETER_RDG_BUFFER_UAV(RWStructuredBuffer, OutCurrOffsetBuffer)
 		SHADER_PARAMETER_RDG_BUFFER_UAV(RWStructuredBuffer, OutPrevOffsetBuffer)
 	END_SHADER_PARAMETER_STRUCT()
