- Handle potential nullptr for DeformedBuffer

diff --git a/Engine/Plugins/Runtime/HairStrands/Source/HairStrandsCore/Private/Niagara/NiagaraDataInterfaceHairStrands.cpp b/Engine/Plugins/Runtime/HairStrands/Source/HairStrandsCore/Private/Niagara/NiagaraDataInterfaceHairStrands.cpp
--- a/Engine/Plugins/Runtime/HairStrands/Source/HairStrandsCore/Private/Niagara/NiagaraDataInterfaceHairStrands.cpp
+++ b/Engine/Plugins/Runtime/HairStrands/Source/HairStrandsCore/Private/Niagara/NiagaraDataInterfaceHairStrands.cpp
@@ -3576,23 +3576,31 @@ void FNDIHairStrandsProxy::PreStage(const FNDIGpuComputePreStageContext& Context
 
 void FNDIHairStrandsProxy::PostSimulate(const FNDIGpuComputePostSimulateContext& Context)
 {
+	// Check we have valid hair data
 	FNDIHairStrandsData* ProxyData = SystemInstancesToProxyData.Find(Context.GetSystemInstanceID());
-
-	const bool bIsHairValid = ProxyData != nullptr && ProxyData->HairStrandsBuffer && ProxyData->HairStrandsBuffer->IsInitialized();
-	const bool bIsDeformedValid = bIsHairValid && ProxyData->HairStrandsBuffer->SourceDeformedResources && ProxyData->HairStrandsBuffer->SourceDeformedResources->IsInitialized();
+	FNDIHairStrandsBuffer* HairStrandsBuffer = ProxyData ? ProxyData->HairStrandsBuffer : nullptr;
+	if (!HairStrandsBuffer || !HairStrandsBuffer->IsInitialized())
+	{
+		return;
+	}
 
 	// MGPU DeformedPositionBuffer copy after simulation
-	if (bIsDeformedValid)
+	if (HairStrandsBuffer->SourceDeformedResources && HairStrandsBuffer->SourceDeformedResources->IsInitialized())
 	{
-		const FNDIHairStrandsBuffer* HairStrandsBuffer = ProxyData->HairStrandsBuffer;
-		FRHIBuffer* DeformedPositionBuffer = HairStrandsBuffer->SourceDeformedResources->DeformedPositionBuffer[HairStrandsBuffer->SourceDeformedResources->CurrentIndex].Buffer->GetRHI();
-		Context.GetComputeDispatchInterface().MultiGPUResourceModified(Context.GetGraphBuilder(), DeformedPositionBuffer, false, true);
+		const FRDGExternalBuffer& DeformedBuffer = HairStrandsBuffer->SourceDeformedResources->DeformedPositionBuffer[HairStrandsBuffer->SourceDeformedResources->CurrentIndex];
+		if (DeformedBuffer.Buffer.IsValid())
+		{
+			if (FRHIBuffer* DeformedPositionBuffer = DeformedBuffer.Buffer->GetRHI())
+			{
+				Context.GetComputeDispatchInterface().MultiGPUResourceModified(Context.GetGraphBuilder(), DeformedPositionBuffer, false, true);
+			}
+		}
 	}
 
-	if (bIsHairValid && Context.IsFinalPostSimulate())
+	if (Context.IsFinalPostSimulate())
 	{
-		ProxyData->HairStrandsBuffer->BoundingBoxBuffer.EndGraphUsage();
-		ProxyData->HairStrandsBuffer->ParamsScaleBuffer.EndGraphUsage();
+		HairStrandsBuffer->BoundingBoxBuffer.EndGraphUsage();
+		HairStrandsBuffer->ParamsScaleBuffer.EndGraphUsage();
 	}
 }
 
