Fix groom binding asset being rebuilt while being already cached.

The logic for filling the DDC key, used later for warming local DDC, as incorrectly run too early. This causes the logic to check if data are available to fail, and forcing the data to be rebuilt.

diff --git a/Engine/Plugins/Runtime/HairStrands/Source/HairStrandsCore/Private/GroomBindingAsset.cpp b/Engine/Plugins/Runtime/HairStrands/Source/HairStrandsCore/Private/GroomBindingAsset.cpp
--- a/Engine/Plugins/Runtime/HairStrands/Source/HairStrandsCore/Private/GroomBindingAsset.cpp
+++ b/Engine/Plugins/Runtime/HairStrands/Source/HairStrandsCore/Private/GroomBindingAsset.cpp
@@ -956,13 +956,14 @@ static void CacheDerivedDatas(UGroomBindingAsset* In, const uint32 InGroupIndex,
 		{
 			UE_CLOG(IsHairStrandsDDCLogEnable(), LogHairStrands, Log, TEXT("[GroomBinding/DDC] Found (GroomBinding:%s)."), *In->GetName());
 
-			FillDrivedDataKey(OutPlatformData);
-
 			// Header
 			FMemoryReaderView Ar(Data, /*bIsPersistent*/ true);
 			InternalSerializePlatformData(Ar, In, OutPlatformData, 0 /*Flags*/, true /*bHeader*/, false /*bData*/);
 			bHasDataInCache = true;
 
+			// Fill DDC key for each strands LOD root bulk data. Done after InternalSerializePlatformData(), as RenRootBulkDatas is not filled in yet, and OutPlatformData.RenRootBulkDatas.Num() == 0.
+			FillDrivedDataKey(OutPlatformData);
+
 			// Verify that all strands data are correctly cached into the DDC
 			{
 				for (int32 MeshLODIndex = 0, MeshLODCount = OutPlatformData.RenRootBulkDatas.Num(); MeshLODIndex < MeshLODCount; ++MeshLODIndex)
