Fix incorrect hair strands stretching when exporting MH.

RBF deformation was applied even though the RBF deformation was not enabled by content creators.

diff --git a/Engine/Plugins/Runtime/HairStrands/Source/HairStrandsCore/Private/GroomRBFDeformer.cpp b/Engine/Plugins/Runtime/HairStrands/Source/HairStrandsCore/Private/GroomRBFDeformer.cpp
--- a/Engine/Plugins/Runtime/HairStrands/Source/HairStrandsCore/Private/GroomRBFDeformer.cpp
+++ b/Engine/Plugins/Runtime/HairStrands/Source/HairStrandsCore/Private/GroomRBFDeformer.cpp
@@ -101,6 +101,7 @@ void DeformStrands(
 	const TArray<FHairStrandsMeshTrianglePositionFormat::Type>& UniqueTrianglePositionBuffer_Deformed,
 	uint32 VertexCount,
 	uint32 SampleCount,
+	const bool bHasRBF,
 	const TArray<FVector3f>& RestPosePositionBuffer,
 	const TArray<FVector4f>& RestSamplePositionsBuffer,
 	const TArray<FVector3f>& MeshSampleWeightsBuffer,
@@ -113,7 +114,7 @@ void DeformStrands(
 	{
 		const FVector3f& ControlPoint = RestPosePositionBuffer[VertexIndex];
 		const FVector3f DisplacedPosition = DisplacePosition(ControlPoint, SampleCount, RestSamplePositionsBuffer, MeshSampleWeightsBuffer);
-		OutDeformedPositionBuffer[VertexIndex] = DisplacedPosition;
+		OutDeformedPositionBuffer[VertexIndex] = bHasRBF ? DisplacedPosition : ControlPoint;
 	});
 
 	// Build curve shuffling indices
@@ -155,7 +156,7 @@ void DeformStrands(
 
 		const uint32 PackedBarycentric = RootBarycentricBuffer[ShuffleRootIndex];
 		const FVector2f B0 = FVector2f(FHairStrandsRootUtils::UnpackBarycentrics(PackedBarycentric));
-		const FVector3f   B  = FVector3f(B0.X, B0.Y, 1.f - B0.X - B0.Y);
+		const FVector3f B  = FVector3f(B0.X, B0.Y, 1.f - B0.X - B0.Y);
 
 		/* Strand hair roots translation and rotation in rest position relative to the bound triangle. Positions are relative to the rest root center */
 		const FVector3f& Rest_V0 = UniqueTrianglePositionBuffer_Rest[TriangleIndex * 3 + 0];
@@ -228,6 +229,7 @@ TArray<FVector3f> GetDeformedHairStrandsPositions(
 	const TArray<FVector3f>& MeshVertexPositionsBuffer_Target,
 	const FHairStrandsDatas& HairStrandsData,
 	const uint32 MeshLODIndex,
+	const bool bHasRBF,
 	const FSkeletalMeshRenderData* InMeshRenderData,
 	const TArray<FHairStrandsIndexFormat::Type>& PointToCurveBuffer,
 	const FHairStrandsRootData& RestLODData)
@@ -280,6 +282,7 @@ TArray<FVector3f> GetDeformedHairStrandsPositions(
 		UniqueTrianglePositionBuffer_Deformed,
 		VertexCount, 
 		MaxSampleCount, 
+		bHasRBF,
 		RestPosePositionBuffer, 
 		RestSamplePositionsBuffer, 
 		MeshSampleWeightsBuffer, 
@@ -754,10 +757,13 @@ void FGroomRBFDeformer::GetRBFDeformedGroomAsset(const UGroomAsset* InGroomAsset
 				FHairStrandsRootData SimRootData;
 				FGroomBindingBuilder::GetRootData(SimRootData, BindingAsset->GetHairGroupsPlatformData()[GroupIndex].SimRootBulkDatas[MeshLODIndex]);
 
+				const bool bHasRBF = InGroomAsset->IsGlobalInterpolationEnable(GroupIndex, 0/*HairLODIndex*/);
+
 				DeformedPositions[GroupIndex].GuideStrands = GetDeformedHairStrandsPositions(
 					MeshVertexPositionsBuffer_Target,
 					GuidesData,
 					MeshLODIndex,
+					bHasRBF,
 					SkeletalMeshData_Target,
 					SimRootDataPointToCurveBuffer,
 					SimRootData);
@@ -786,10 +792,13 @@ void FGroomRBFDeformer::GetRBFDeformedGroomAsset(const UGroomAsset* InGroomAsset
 					RenLODData.MeshSectionCount					= SimLODData.MeshSectionCount;
 				}
 
+				const bool bHasRBF = InGroomAsset->IsGlobalInterpolationEnable(GroupIndex, 0/*HairLODIndex*/);
+
 				DeformedPositions[GroupIndex].RenderStrands = GetDeformedHairStrandsPositions(
 					MeshVertexPositionsBuffer_Target,
 					StrandsData,
 					MeshLODIndex,
+					bHasRBF,
 					SkeletalMeshData_Target,
 					RenRootDataPointToCurveBuffer,
 					RenRootData);
