Fix incorrect binding when a mesh section is provided.

diff --git a/Engine/Plugins/Runtime/HairStrands/Source/HairStrandsCore/Private/GroomBindingBuilder.cpp b/Engine/Plugins/Runtime/HairStrands/Source/HairStrandsCore/Private/GroomBindingBuilder.cpp
--- a/Engine/Plugins/Runtime/HairStrands/Source/HairStrandsCore/Private/GroomBindingBuilder.cpp
+++ b/Engine/Plugins/Runtime/HairStrands/Source/HairStrandsCore/Private/GroomBindingBuilder.cpp
@@ -1149,7 +1149,8 @@ namespace GroomBinding_RootProjection
 		const FHairStrandsDatas& InStrandsData,
 		const GroomBinding_Mesh::IMeshData* InMeshData,
 		const TArray<TArray<FVector3f>>& InTransferredPositions,
-		TArray<FHairStrandsRootData>& OutRootData)
+		TArray<FHairStrandsRootData>& OutRootData,
+		uint32 InMatchingSection)
 	{
 		// 2. Project root for each mesh LOD
 		const uint32 CurveCount = InStrandsData.GetNumCurves();
@@ -1185,10 +1186,19 @@ namespace GroomBinding_RootProjection
 				return false;
 			}
 
+			// Only cull section if the 'matching section' exist
+			const bool bCullSection = InMatchingSection > 0 && InMatchingSection < SectionCount;
+
 			float ClosestTrianglePoint = FLT_MAX;
 			check(SectionCount > 0);
 			for (uint32 SectionIt = 0; SectionIt < SectionCount; ++SectionIt)
 			{
+				// Cull out sections which don't match
+				if (bCullSection && SectionIt != InMatchingSection)
+				{
+					continue;
+				}
+
 				// 2.2.1 Compute the bounding box of the skeletal mesh
 				const GroomBinding_Mesh::IMeshSectionData& Section = MeshLODData.GetSection(SectionIt);
 				const uint32 TriangleCount = Section.GetNumTriangles();
@@ -1275,6 +1285,12 @@ namespace GroomBinding_RootProjection
 				check(SectionCount < MaxSectionCount);
 				check(TriangleCount > 0);
 
+				// Cull out sections which don't match
+				if (bCullSection && SectionIt != InMatchingSection)
+				{
+					continue;
+				}
+
 				for (uint32 TriangleIt = 0; TriangleIt < TriangleCount; ++TriangleIt)
 				{
 					FTriangleGrid::FTriangle T;
@@ -2334,7 +2350,8 @@ static bool InternalBuildBinding_CPU(const FGroomBindingBuilder::FInput& In, uin
 					GuidesData,
 					TargetMeshData.Get(),
 					TransferredPositions,
-					OutData.SimRootDatas))
+					OutData.SimRootDatas,
+					In.MatchingSection))
 				{
 					UE_LOG(LogHairStrands, Error, TEXT("[Groom] Binding asset could not be built. Some guide roots are not close enough to the target mesh to be projected onto it."));
 					return false; 
@@ -2349,7 +2366,8 @@ static bool InternalBuildBinding_CPU(const FGroomBindingBuilder::FInput& In, uin
 					StrandsData,
 					TargetMeshData.Get(),
 					TransferredPositions,
-					OutData.RenRootDatas))
+					OutData.RenRootDatas,
+					In.MatchingSection))
 				{
 					UE_LOG(LogHairStrands, Error, TEXT("[Groom] Binding asset could not be built. Some strand roots are not close enough to the target mesh to be projected onto it."));
 					return false;
@@ -2371,7 +2389,8 @@ static bool InternalBuildBinding_CPU(const FGroomBindingBuilder::FInput& In, uin
 							LODGuidesData,
 							TargetMeshData.Get(),
 							TransferredPositions,
-							OutData.CardsRootDatas[CardsLODIt]))
+							OutData.CardsRootDatas[CardsLODIt],
+							0/*In.MatchingSection*/))
 						{
 							UE_LOG(LogHairStrands, Error, TEXT("[Groom] Binding asset could not be built. Some cards guide roots are not close enough to the target mesh to be projected onto it."));
 							return false; 
