Fixed hair cards vertex declaration for Materials input on platforms with MVF disabled
Added correct vertex declaration for hair cards PSO precaching on platforms with MVF disabled

diff --git a/Engine/Plugins/Runtime/HairStrands/Source/HairStrandsCore/Private/HairCardsVertexFactory.cpp b/Engine/Plugins/Runtime/HairStrands/Source/HairStrandsCore/Private/HairCardsVertexFactory.cpp
--- a/Engine/Plugins/Runtime/HairStrands/Source/HairStrandsCore/Private/HairCardsVertexFactory.cpp
+++ b/Engine/Plugins/Runtime/HairStrands/Source/HairStrandsCore/Private/HairCardsVertexFactory.cpp
@@ -289,8 +289,28 @@ void FHairCardsVertexFactory::ValidateCompiledResult(const FVertexFactoryType* T
 
 void FHairCardsVertexFactory::GetPSOPrecacheVertexFetchElements(EVertexInputStreamType VertexInputStreamType, FVertexDeclarationElementList& Elements)
 {
-	// Manual vertex fetch is available for this factory so only primitive ID stream is used
-	Elements.Add(FVertexElement(0, 0, VET_UInt, HAIR_CARDS_VF_PRIMITIVEID_STREAM_INDEX, 0, true));
+	if (!PlatformGPUSceneUsesUniformBufferView(GMaxRHIShaderPlatform))
+	{ 
+		// Manual vertex fetch is available for this factory so only primitive ID stream is used
+		Elements.Add(FVertexElement(0, 0, VET_UInt, HAIR_CARDS_VF_PRIMITIVEID_STREAM_INDEX, 0, true));
+	}
+
+	const bool bManualFetch = RHISupportsManualVertexFetch(GMaxRHIShaderPlatform);
+	if (!bManualFetch)
+	{
+		uint8 NumStreams = Elements.Num();
+		
+		Elements.Add(FVertexElement(0+NumStreams, 0, FHairCardsPositionFormat::VertexElementType, 0, FHairCardsPositionFormat::SizeInByte, false));
+		Elements.Add(FVertexElement(0+NumStreams, 0, FHairCardsPositionFormat::VertexElementType, 5, FHairCardsPositionFormat::SizeInByte, false));
+
+		uint8 NormalOffset0 = 0u;
+		uint8 NormalOffset1 = FHairCardsNormalFormat::SizeInByte;
+		Elements.Add(FVertexElement(1+NumStreams, NormalOffset0, FHairCardsNormalFormat::VertexElementType, 1, FHairCardsNormalFormat::SizeInByte * FHairCardsNormalFormat::ComponentCount, false));
+		Elements.Add(FVertexElement(1+NumStreams, NormalOffset1, FHairCardsNormalFormat::VertexElementType, 2, FHairCardsNormalFormat::SizeInByte * FHairCardsNormalFormat::ComponentCount, false));
+
+		Elements.Add(FVertexElement(2+NumStreams, 0, FHairCardsUVFormat::VertexElementType, 3, FHairCardsUVFormat::SizeInByte, false));
+		Elements.Add(FVertexElement(3+NumStreams, 0, FHairCardsMaterialFormat::VertexElementType, 4, FHairCardsMaterialFormat::SizeInByte, false));
+	}
 }
 
 EPrimitiveIdMode FHairCardsVertexFactory::GetPrimitiveIdMode(ERHIFeatureLevel::Type In) const
diff --git a/Engine/Plugins/Runtime/HairStrands/Source/HairStrandsCore/Public/HairCardsDatas.h b/Engine/Plugins/Runtime/HairStrands/Source/HairStrandsCore/Public/HairCardsDatas.h
--- a/Engine/Plugins/Runtime/HairStrands/Source/HairStrandsCore/Public/HairCardsDatas.h
+++ b/Engine/Plugins/Runtime/HairStrands/Source/HairStrandsCore/Public/HairCardsDatas.h
@@ -35,7 +35,7 @@ struct FHairCardsMaterialFormat
 	typedef uint32 Type;
 	static const uint32 ComponentCount = 1;
 	static const uint32 SizeInByte = sizeof(Type);
-	static const EVertexElementType VertexElementType = VET_UByte4;
+	static const EVertexElementType VertexElementType = VET_UByte4N; 
 	static const EPixelFormat Format = PF_R8G8B8A8;
 };
 
