Hair card grooms are invisible in Performance Rendering Mode.

diff --git a/Engine/Plugins/Runtime/HairStrands/Source/HairStrandsCore/Private/HairStrandsMeshProjection.cpp b/Engine/Plugins/Runtime/HairStrands/Source/HairStrandsCore/Private/HairStrandsMeshProjection.cpp
--- a/Engine/Plugins/Runtime/HairStrands/Source/HairStrandsCore/Private/HairStrandsMeshProjection.cpp
+++ b/Engine/Plugins/Runtime/HairStrands/Source/HairStrandsCore/Private/HairStrandsMeshProjection.cpp
@@ -195,7 +195,7 @@ static bool AddHairStrandUpdateMeshTrianglesPass(
 
 	// If a skel. mesh is not streaming yet, its SRV will be null
 	const bool bValid = EffectiveSectionCount > 0 && TotalSectionCount < GetHairStrandsMaxSectionCount();
-	const bool bReady = MeshLODData.Sections[0].IndexBuffer != nullptr && MeshLODData.Sections[0].UVsBuffer != nullptr && (MeshLODData.Sections[0].RDGPositionBuffer != nullptr || MeshLODData.Sections[0].PositionBuffer != nullptr);
+	const bool bReady = MeshLODData.Sections[0].IndexBuffer != nullptr && (MeshLODData.Sections[0].RDGPositionBuffer != nullptr || MeshLODData.Sections[0].PositionBuffer != nullptr);
 	if (!bReady || !bValid)
 	{
 		return false;
@@ -276,6 +276,13 @@ static bool AddHairStrandUpdateMeshTrianglesPass(
 
 	const bool bComputePreviousDeformedPosition = OutputPrevUAV != nullptr;
 	{
+		// UV stream is only used for debugging purpose. On some platform the TextureCoordinateSRV can be null as it is not create by default (requires CPU access flags).
+		// In such a case we bind the index buffer as a dummy data
+		if (CommonParameters.MeshUVsBuffer == nullptr)
+		{
+			CommonParameters.MeshUVsBuffer = CommonParameters.MeshIndexBuffer;
+		}
+
 		FHairUpdateMeshTriangleCS::FParameters* PassParameters = GraphBuilder.AllocParameters<FHairUpdateMeshTriangleCS::FParameters>();
 		*PassParameters = CommonParameters;
 
