Fix crash when changing groom asset property

diff --git a/Engine/Plugins/Runtime/HairStrands/Source/HairStrandsEditor/Private/GroomCustomAssetEditorToolkit.cpp b/Engine/Plugins/Runtime/HairStrands/Source/HairStrandsEditor/Private/GroomCustomAssetEditorToolkit.cpp
--- a/Engine/Plugins/Runtime/HairStrands/Source/HairStrandsEditor/Private/GroomCustomAssetEditorToolkit.cpp
+++ b/Engine/Plugins/Runtime/HairStrands/Source/HairStrandsEditor/Private/GroomCustomAssetEditorToolkit.cpp
@@ -303,16 +303,30 @@ void FGroomCustomAssetEditorToolkit::InitPreviewComponents()
 void FGroomCustomAssetEditorToolkit::OnClose() 
 {
 	// Remove all delegates
-	if (GroomAsset.IsValid() && PropertyListenDelegates.Num() > 0)
+	if (GroomAsset.IsValid() && PropertyListenDelegatesResourceChanged.Num() > 0)
 	{
-		for (FDelegateHandle Handle : PropertyListenDelegates)
+		if(PropertyListenDelegatesResourceChanged.Num() > 0)
 		{
-			if (Handle.IsValid())
+			for (FDelegateHandle Handle : PropertyListenDelegatesResourceChanged)
 			{
-				GroomAsset->GetOnGroomAssetResourcesChanged().Remove(Handle);
+				if (Handle.IsValid())
+				{
+					GroomAsset->GetOnGroomAssetResourcesChanged().Remove(Handle);
+				}
 			}
+			PropertyListenDelegatesResourceChanged.Empty();
 		}
-		PropertyListenDelegates.Empty();
+		if (PropertyListenDelegatesAssetChanged.Num() > 0)
+        {
+         	for (FDelegateHandle Handle : PropertyListenDelegatesAssetChanged)
+         	{
+         		if (Handle.IsValid())
+         		{
+         			GroomAsset->GetOnGroomAssetChanged().Remove(Handle);
+         		}
+         	}
+         	PropertyListenDelegatesAssetChanged.Empty();
+        }
 	}
 
 	PropertiesTab.Reset();
@@ -664,8 +678,8 @@ void FGroomCustomAssetEditorToolkit::InitCustomAssetEditor(const EToolkitMode::T
 			if (LocalToolKit->DetailView_BindingProperties)			{ LocalToolKit->DetailView_BindingProperties->ForceRefresh(); }
 		};
 
-		PropertyListenDelegates.Add(GroomAsset->GetOnGroomAssetResourcesChanged().AddLambda(InvalidateDetailViews));
-		PropertyListenDelegates.Add(GroomAsset->GetOnGroomAssetChanged().AddLambda(InvalidateDetailViews));
+		PropertyListenDelegatesResourceChanged.Add(GroomAsset->GetOnGroomAssetResourcesChanged().AddLambda(InvalidateDetailViews));
+		PropertyListenDelegatesAssetChanged.Add(GroomAsset->GetOnGroomAssetChanged().AddLambda(InvalidateDetailViews));
 	}
 }
 
diff --git a/Engine/Plugins/Runtime/HairStrands/Source/HairStrandsEditor/Public/GroomCustomAssetEditorToolkit.h b/Engine/Plugins/Runtime/HairStrands/Source/HairStrandsEditor/Public/GroomCustomAssetEditorToolkit.h
--- a/Engine/Plugins/Runtime/HairStrands/Source/HairStrandsEditor/Public/GroomCustomAssetEditorToolkit.h
+++ b/Engine/Plugins/Runtime/HairStrands/Source/HairStrandsEditor/Public/GroomCustomAssetEditorToolkit.h
@@ -179,7 +179,8 @@ private:
 
 	bool bIsTabManagerInitialized = false;
 	int32 ActiveGroomBindingIndex = -1;
-	TArray<FDelegateHandle> PropertyListenDelegates;
+	TArray<FDelegateHandle> PropertyListenDelegatesAssetChanged;
+	TArray<FDelegateHandle> PropertyListenDelegatesResourceChanged;
 	TWeakObjectPtr<UGroomAsset> GroomAsset;
 	TWeakObjectPtr<UGroomBindingAsset> GroomBindingAsset;
 	TWeakObjectPtr<UGroomBindingAssetList> GroomBindingAssetList;
