Fixed crash when creating a groom asset from Content Browser.

diff --git a/Engine/Plugins/Runtime/HairStrands/Source/HairStrandsEditor/Private/GroomFactory.cpp b/Engine/Plugins/Runtime/HairStrands/Source/HairStrandsEditor/Private/GroomFactory.cpp
--- a/Engine/Plugins/Runtime/HairStrands/Source/HairStrandsEditor/Private/GroomFactory.cpp
+++ b/Engine/Plugins/Runtime/HairStrands/Source/HairStrandsEditor/Private/GroomFactory.cpp
@@ -24,6 +24,10 @@ UObject* UGroomFactory::FactoryCreateNew(UClass* InClass, UObject* InParent, FNa
 	Groom->GetHairGroupsLOD()[0].LODs[0].GeometryType = EGroomGeometryType::Cards;
 	Groom->GetHairGroupsInterpolation()[0].InterpolationSettings.bUseUniqueGuide = true;
 
+	// Group must have at least 1 LOD for cards/meshes even if it has no valid data to pass validation
+	Groom->GetHairGroupsPlatformData()[0].Cards.LODs.SetNum(1);
+	Groom->GetHairGroupsPlatformData()[0].Meshes.LODs.SetNum(1);
+
 	Groom->UpdateHairGroupsInfo();
 	Groom->UpdateCachedSettings();
 	Groom->InitResources();
