Fix MetaHuman eyelashes not being displayed in rendering modes where strands are not supported

Add function UGroomBlueprintLibrary::IsHairStrandsSupportedInWorld to check if strands are going to be rendered for a given World

diff --git a/Engine/Plugins/Runtime/HairStrands/Source/HairStrandsCore/Private/GroomBlueprintLibrary.cpp b/Engine/Plugins/Runtime/HairStrands/Source/HairStrandsCore/Private/GroomBlueprintLibrary.cpp
--- a/Engine/Plugins/Runtime/HairStrands/Source/HairStrandsCore/Private/GroomBlueprintLibrary.cpp
+++ b/Engine/Plugins/Runtime/HairStrands/Source/HairStrandsCore/Private/GroomBlueprintLibrary.cpp
@@ -4,6 +4,9 @@
 #include "GeometryCache.h"
 #include "HairStrandsCore.h"
 #include "GroomBindingAsset.h"
+#include "HairStrandsInterface.h"
+#include "Engine/Engine.h"
+#include "SceneInterface.h"
 
 #include UE_INLINE_GENERATED_CPP_BY_NAME(GroomBlueprintLibrary)
 
@@ -109,3 +112,17 @@ UGroomBindingAsset* UGroomBlueprintLibrary::CreateNewGeometryCacheGroomBindingAs
 #endif
 }
 
+bool UGroomBlueprintLibrary::IsHairStrandsSupportedInWorld(const UObject* WorldContextObject)
+{
+	bool bHairEnabledSupported = false;
+
+	if (UWorld* World = GEngine->GetWorldFromContextObject(WorldContextObject, EGetWorldErrorMode::LogAndReturnNull))
+	{
+		if (FSceneInterface* Scene = World->Scene)
+		{
+			bHairEnabledSupported = IsHairStrandsSupported(EHairStrandsShaderType::Strands, Scene->GetShaderPlatform());
+		}
+	}
+
+	return bHairEnabledSupported;
+}
\ No newline at end of file
diff --git a/Engine/Plugins/Runtime/HairStrands/Source/HairStrandsCore/Public/GroomBlueprintLibrary.h b/Engine/Plugins/Runtime/HairStrands/Source/HairStrandsCore/Public/GroomBlueprintLibrary.h
--- a/Engine/Plugins/Runtime/HairStrands/Source/HairStrandsCore/Public/GroomBlueprintLibrary.h
+++ b/Engine/Plugins/Runtime/HairStrands/Source/HairStrandsCore/Public/GroomBlueprintLibrary.h
@@ -80,4 +80,12 @@ public:
 		int32 NumInterpolationPoints = 100,
 		UGeometryCache* SourceGeometryCacheForTransfer = nullptr,
 		int32 MatchingSection = 0);
+
+	/**
+	 * Check for strands support in the world of a given Actor Component
+	 * @param InWorld The world to check
+	 * @return true if strands are going to be rendered for the world's scene and false otherwise
+	 */
+	UFUNCTION(BlueprintCallable, Category = "Groom", meta = (WorldContext="WorldContextObject"))
+	static bool IsHairStrandsSupportedInWorld(const UObject* WorldContextObject);
 };
\ No newline at end of file
