Fix hair group name not being serialized after import.

The group name weren't correctly transferred from the hair group description into the UGroomAsset. Thus there weren't serialized. This was causing group to not have any names once the asset was reopened.

diff --git a/Engine/Plugins/Runtime/HairStrands/Source/HairStrandsCore/Private/GroomAsset.cpp b/Engine/Plugins/Runtime/HairStrands/Source/HairStrandsCore/Private/GroomAsset.cpp
--- a/Engine/Plugins/Runtime/HairStrands/Source/HairStrandsCore/Private/GroomAsset.cpp
+++ b/Engine/Plugins/Runtime/HairStrands/Source/HairStrandsCore/Private/GroomAsset.cpp
@@ -1161,6 +1161,17 @@ void UGroomAsset::UpdateHairGroupsInfo()
 		Info.NumCurveVertices = Data.Strands.BulkData.GetNumPoints();
 		Info.NumGuideVertices = Data.Guides.BulkData.GetNumPoints();
 		Info.MaxCurveLength = Data.Strands.BulkData.GetMaxLength();
+	#if WITH_EDITORONLY_DATA
+		// Transfer group name from the hair description group to the UGroomAsset.
+		// This pass is used right after an asset is imported.
+		if (const FHairDescriptionGroups* LocalHairDescriptionGroups = CachedHairDescriptionGroups[EHairDescriptionType::Source].Get())
+		{
+			if (LocalHairDescriptionGroups->HairGroups.IsValidIndex(GroupIndex))
+			{
+				Info.GroupName = LocalHairDescriptionGroups->HairGroups[GroupIndex].Info.GroupName;
+			}
+		}
+	#endif
 		if (bForceReset)
 		{
 			Info.bIsVisible = true;
