Remove Groom concurent update transform function as it is no needed.

diff --git a/Engine/Plugins/Runtime/HairStrands/Source/HairStrandsCore/Private/GroomComponent.cpp b/Engine/Plugins/Runtime/HairStrands/Source/HairStrandsCore/Private/GroomComponent.cpp
--- a/Engine/Plugins/Runtime/HairStrands/Source/HairStrandsCore/Private/GroomComponent.cpp
+++ b/Engine/Plugins/Runtime/HairStrands/Source/HairStrandsCore/Private/GroomComponent.cpp
@@ -3617,19 +3617,6 @@ void UGroomComponent::TickComponent(float DeltaTime, enum ELevelTick TickType, F
 	}
 }
 
-void UGroomComponent::SendRenderTransform_Concurrent()
-{
-	if (RegisteredMeshComponent)
-	{
-		if (ShouldComponentAddToScene() && ShouldRender())
-		{
-			GetWorld()->Scene->UpdatePrimitiveTransform(this);
-		}
-	}
-
-	Super::SendRenderTransform_Concurrent();
-}
-
 void UGroomComponent::SendRenderDynamicData_Concurrent()
 {
 	Super::SendRenderDynamicData_Concurrent();
diff --git a/Engine/Plugins/Runtime/HairStrands/Source/HairStrandsCore/Public/GroomComponent.h b/Engine/Plugins/Runtime/HairStrands/Source/HairStrandsCore/Public/GroomComponent.h
--- a/Engine/Plugins/Runtime/HairStrands/Source/HairStrandsCore/Public/GroomComponent.h
+++ b/Engine/Plugins/Runtime/HairStrands/Source/HairStrandsCore/Public/GroomComponent.h
@@ -136,7 +136,7 @@ public:
 	virtual void OnAttachmentChanged() override;
 	virtual void DetachFromComponent(const FDetachmentTransformRules& DetachmentRules) override;
 	virtual void TickComponent(float DeltaTime, enum ELevelTick TickType, FActorComponentTickFunction *ThisTickFunction) override;
-	virtual void SendRenderTransform_Concurrent() override;
+
 	virtual void CreateRenderState_Concurrent(FRegisterComponentContext* Context) override;
 	virtual void SendRenderDynamicData_Concurrent() override;
 	virtual void DestroyRenderState_Concurrent() override;
