Change hair strands texture generation to emit first a GBuffer like output when rasterizing the geometry. This avoids costly overdraw when invocating texel tracing, which can cause GPU TDR due to long work.

This prevents hair strands texture generation to be sensible to the input geometry.

diff --git a/Engine/Plugins/Runtime/HairStrands/Source/HairStrandsCore/Private/GroomTextureBuilder.cpp b/Engine/Plugins/Runtime/HairStrands/Source/HairStrandsCore/Private/GroomTextureBuilder.cpp
--- a/Engine/Plugins/Runtime/HairStrands/Source/HairStrandsCore/Private/GroomTextureBuilder.cpp
+++ b/Engine/Plugins/Runtime/HairStrands/Source/HairStrandsCore/Private/GroomTextureBuilder.cpp
@@ -682,17 +682,15 @@ class FHairStrandsTextureVS : public FGlobalShader
 	DECLARE_GLOBAL_SHADER(FHairStrandsTextureVS);
 	SHADER_USE_PARAMETER_STRUCT(FHairStrandsTextureVS, FGlobalShader)
 
-		BEGIN_SHADER_PARAMETER_STRUCT(FParameters, )
+	BEGIN_SHADER_PARAMETER_STRUCT(FParameters, )
 		SHADER_PARAMETER(FIntPoint, OutputResolution)
 		SHADER_PARAMETER(uint32, VertexCount)
-
 		SHADER_PARAMETER(uint32, UVsChannelIndex)
 		SHADER_PARAMETER(uint32, UVsChannelCount)
-
 		SHADER_PARAMETER_SRV(Buffer, VertexBuffer)
 		SHADER_PARAMETER_SRV(Buffer, UVsBuffer)
 		SHADER_PARAMETER_SRV(Buffer, NormalsBuffer)
-		END_SHADER_PARAMETER_STRUCT()
+	END_SHADER_PARAMETER_STRUCT()
 
 		static bool ShouldCompilePermutation(const FGlobalShaderPermutationParameters& Parameters) { return IsHairStrandsSupported(EHairStrandsShaderType::Tool, Parameters.Platform); }
 	static void ModifyCompilationEnvironment(const FGlobalShaderPermutationParameters& Parameters, FShaderCompilerEnvironment& OutEnvironment)
@@ -702,19 +700,48 @@ class FHairStrandsTextureVS : public FGlobalShader
 	}
 };
 
+class FHairStrandsTexturePS : public FGlobalShader
+{
+	DECLARE_GLOBAL_SHADER(FHairStrandsTexturePS);
+	SHADER_USE_PARAMETER_STRUCT(FHairStrandsTexturePS, FGlobalShader)
+
+	BEGIN_SHADER_PARAMETER_STRUCT(FParameters, )
+		SHADER_PARAMETER(FIntPoint, OutputResolution)
+		SHADER_PARAMETER(uint32, VertexCount)
+		SHADER_PARAMETER(uint32, UVsChannelIndex)
+		SHADER_PARAMETER(uint32, UVsChannelCount)
+		SHADER_PARAMETER_SRV(Buffer, VertexBuffer)
+		SHADER_PARAMETER_SRV(Buffer, UVsBuffer)
+		SHADER_PARAMETER_SRV(Buffer, NormalsBuffer)
+		RENDER_TARGET_BINDING_SLOTS()
+	END_SHADER_PARAMETER_STRUCT()
+
+	static bool ShouldCompilePermutation(const FGlobalShaderPermutationParameters& Parameters) { return IsHairStrandsSupported(EHairStrandsShaderType::Tool, Parameters.Platform); }
+	static void ModifyCompilationEnvironment(const FGlobalShaderPermutationParameters& Parameters, FShaderCompilerEnvironment& OutEnvironment)
+	{
+		FGlobalShader::ModifyCompilationEnvironment(Parameters, OutEnvironment);
+		OutEnvironment.SetDefine(TEXT("SHADER_PIXEL"), 1);
+	}
+};
+
+IMPLEMENT_GLOBAL_SHADER(FHairStrandsTextureVS, "/Engine/Private/HairStrands/HairStrandsTexturesGeneration.usf", "MainVS", SF_Vertex);
+IMPLEMENT_GLOBAL_SHADER(FHairStrandsTexturePS, "/Engine/Private/HairStrands/HairStrandsTexturesGeneration.usf", "MainPS", SF_Pixel);
+
 BEGIN_SHADER_PARAMETER_STRUCT(FHairStrandsInstanceRawParameters, )
 	SHADER_PARAMETER_STRUCT_INCLUDE(FHairStrandsInstanceCommonParameters, Common)
 	SHADER_PARAMETER_STRUCT_INCLUDE(FHairStrandsInstanceResourceRawParameters, Resources)
 	SHADER_PARAMETER_STRUCT_INCLUDE(FHairStrandsInstanceCullingRawParameters, Culling)
 END_SHADER_PARAMETER_STRUCT()
 
-class FHairStrandsTexturePS : public FGlobalShader
+class FHairStrandsTextureCS : public FGlobalShader
 {
-	DECLARE_GLOBAL_SHADER(FHairStrandsTexturePS);
-	SHADER_USE_PARAMETER_STRUCT(FHairStrandsTexturePS, FGlobalShader)
+	DECLARE_GLOBAL_SHADER(FHairStrandsTextureCS);
+	SHADER_USE_PARAMETER_STRUCT(FHairStrandsTextureCS, FGlobalShader)
 
 		BEGIN_SHADER_PARAMETER_STRUCT(FParameters, )
 		SHADER_PARAMETER_STRUCT_INCLUDE(ShaderPrint::FShaderParameters, ShaderPrintParameters)
+		SHADER_PARAMETER(FIntPoint, TileSize)
+		SHADER_PARAMETER(FIntPoint, TileOffsetInPixels)
 		SHADER_PARAMETER(uint32, LayoutIndex)
 		SHADER_PARAMETER(FIntPoint, OutputResolution)
 		SHADER_PARAMETER(uint32, VertexCount)
@@ -741,19 +768,23 @@ class FHairStrandsTexturePS : public FGlobalShader
 		SHADER_PARAMETER_SRV(Buffer, UVsBuffer)
 		SHADER_PARAMETER_SRV(Buffer, NormalsBuffer)
 		SHADER_PARAMETER_RDG_TEXTURE_UAV(RWTexture2DArray, OutTexture)
-		RENDER_TARGET_BINDING_SLOTS()
+
+		SHADER_PARAMETER_RDG_TEXTURE(Texture2D, PositionTexture)
+		SHADER_PARAMETER_RDG_TEXTURE(Texture2D, TangentTexture)
+		SHADER_PARAMETER_RDG_TEXTURE(Texture2D, BitangentTexture)
+		SHADER_PARAMETER_RDG_TEXTURE(Texture2D, NormalTexture)
+		SHADER_PARAMETER_RDG_TEXTURE(Texture2D, UVTexture)
 	END_SHADER_PARAMETER_STRUCT()
 
 	static bool ShouldCompilePermutation(const FGlobalShaderPermutationParameters& Parameters) { return IsHairStrandsSupported(EHairStrandsShaderType::Tool, Parameters.Platform); }
 	static void ModifyCompilationEnvironment(const FGlobalShaderPermutationParameters& Parameters, FShaderCompilerEnvironment& OutEnvironment)
 	{
 		FGlobalShader::ModifyCompilationEnvironment(Parameters, OutEnvironment);
-		OutEnvironment.SetDefine(TEXT("SHADER_PIXEL"), 1);
+		OutEnvironment.SetDefine(TEXT("SHADER_COMPUTE"), 1);
 	}
 };
 
-IMPLEMENT_GLOBAL_SHADER(FHairStrandsTextureVS, "/Engine/Private/HairStrands/HairStrandsTexturesGeneration.usf", "MainVS", SF_Vertex);
-IMPLEMENT_GLOBAL_SHADER(FHairStrandsTexturePS, "/Engine/Private/HairStrands/HairStrandsTexturesGeneration.usf", "MainPS", SF_Pixel);
+IMPLEMENT_GLOBAL_SHADER(FHairStrandsTextureCS, "/Engine/Private/HairStrands/HairStrandsTexturesGeneration.usf", "MainCS", SF_Compute);
 
 
 const int32 HairStrandsTextureTileSize = 1024;
@@ -846,102 +877,137 @@ static void InternalGenerateHairStrandsTextures(
 {
 	const FIntPoint OutputResolution = Out.Texture->Desc.Extent;
 	const FIntPoint OutTileCoord = Out.TileCoord;
+	const int32 TileSize = HairStrandsTextureTileSize;
+	const FIntPoint TileOffsetInPixels(OutTileCoord.X * TileSize, OutTileCoord.Y * TileSize);
 
-	FHairStrandsTexturePS::FParameters* ParametersPS = GraphBuilder.AllocParameters<FHairStrandsTexturePS::FParameters>();
-	ParametersPS->LayoutIndex = uint32(Out.Layout);
-	ParametersPS->OutputResolution = OutputResolution;
-	ParametersPS->VertexCount = VertexCount;
-	ParametersPS->VertexBuffer = InMeshVertexBuffer;
-	ParametersPS->UVsBuffer = InMeshUVsBuffer;
-	ParametersPS->NormalsBuffer = InMeshNormalsBuffer;
-	ParametersPS->MaxDistance = FMath::Abs(InMaxDistance);
-	ParametersPS->TracingDirection = InTracingDirection;
-
-	ParametersPS->UVsChannelIndex = UVsChannelIndex;
-	ParametersPS->UVsChannelCount = UVsChannelCount;
+	FRDGTextureRef PositionTexture 	= GraphBuilder.CreateTexture(FRDGTextureDesc::Create2D(OutputResolution, PF_A32B32G32R32F, 	FClearValueBinding(FLinearColor(0, 0, 0, 1)), TexCreate_RenderTargetable | TexCreate_ShaderResource), TEXT("HairTextureGeneration.Position"));
+	FRDGTextureRef TangentTexture 	= GraphBuilder.CreateTexture(FRDGTextureDesc::Create2D(OutputResolution, PF_FloatRGBA, 		FClearValueBinding(FLinearColor(0, 0, 0, 1)), TexCreate_RenderTargetable | TexCreate_ShaderResource), TEXT("HairTextureGeneration.Tangent"));
+	FRDGTextureRef BitangentTexture = GraphBuilder.CreateTexture(FRDGTextureDesc::Create2D(OutputResolution, PF_FloatRGBA, 		FClearValueBinding(FLinearColor(0, 0, 0, 1)), TexCreate_RenderTargetable | TexCreate_ShaderResource), TEXT("HairTextureGeneration.Bitangent"));
+	FRDGTextureRef NormalTexture 	= GraphBuilder.CreateTexture(FRDGTextureDesc::Create2D(OutputResolution, PF_FloatRGBA, 		FClearValueBinding(FLinearColor(0, 0, 0, 1)), TexCreate_RenderTargetable | TexCreate_ShaderResource), TEXT("HairTextureGeneration.Normal"));
+	FRDGTextureRef UVTexture 		= GraphBuilder.CreateTexture(FRDGTextureDesc::Create2D(OutputResolution, PF_G32R32F, 		FClearValueBinding(FLinearColor(0, 0, 0, 1)), TexCreate_RenderTargetable | TexCreate_ShaderResource), TEXT("HairTextureGeneration.UV"));
 
-	ParametersPS->InVF = Instance;
-	
-	ParametersPS->Voxel_MinBound = VoxelMinBound;
-	ParametersPS->Voxel_MaxBound = VoxelMaxBound;
-	ParametersPS->Voxel_Resolution = VoxelResolution;
-	ParametersPS->Voxel_Size = VoxelSize;
-	ParametersPS->Voxel_MaxSegmentPerVoxel = VoxelMaxSegmentPerVoxel;
-	ParametersPS->Voxel_OffsetAndCount = GraphBuilder.CreateSRV(VoxelOffsetAndCount);
-	ParametersPS->Voxel_Data = GraphBuilder.CreateSRV(VoxelData);
-
-	ParametersPS->Voxel_OffsetAndCount_MaxCount = VoxelOffsetAndCount->Desc.NumElements;
-	ParametersPS->Voxel_Data_MaxCount = VoxelData->Desc.NumElements;
-
-	if (ShaderPrintData)
+	// Raster geometry
+	// This pass emits a GBuffer like output. It avoids having multiple PS invocations due to overdraw and lack of early-Z testing (due to SV_Depth output)
 	{
-		ShaderPrint::SetParameters(GraphBuilder, *ShaderPrintData, ParametersPS->ShaderPrintParameters);
+		FHairStrandsTexturePS::FParameters* ParametersPS = GraphBuilder.AllocParameters<FHairStrandsTexturePS::FParameters>();
+		ParametersPS->OutputResolution = OutputResolution;
+		ParametersPS->VertexCount = VertexCount;
+		ParametersPS->VertexBuffer = InMeshVertexBuffer;
+		ParametersPS->UVsBuffer = InMeshUVsBuffer;
+		ParametersPS->NormalsBuffer = InMeshNormalsBuffer;
+		ParametersPS->UVsChannelIndex = UVsChannelIndex;
+		ParametersPS->UVsChannelCount = UVsChannelCount;
+		ParametersPS->RenderTargets[0] = FRenderTargetBinding(PositionTexture, ERenderTargetLoadAction::EClear);
+		ParametersPS->RenderTargets[1] = FRenderTargetBinding(TangentTexture, ERenderTargetLoadAction::EClear);
+		ParametersPS->RenderTargets[2] = FRenderTargetBinding(BitangentTexture, ERenderTargetLoadAction::EClear);
+		ParametersPS->RenderTargets[3] = FRenderTargetBinding(NormalTexture, ERenderTargetLoadAction::EClear);
+		ParametersPS->RenderTargets[4] = FRenderTargetBinding(UVTexture, ERenderTargetLoadAction::EClear);
+		ParametersPS->RenderTargets[5] = FRenderTargetBinding(Out.TriangleMaskTexture, bClear ? ERenderTargetLoadAction::EClear : ERenderTargetLoadAction::ELoad);
+		ParametersPS->RenderTargets.DepthStencil = FDepthStencilBinding(
+			Out.DepthTestTexture,
+			bClear ? ERenderTargetLoadAction::EClear : ERenderTargetLoadAction::ELoad,
+			ERenderTargetLoadAction::ENoAction,
+			FExclusiveDepthStencil::DepthWrite_StencilNop);
+
+		TShaderMapRef<FHairStrandsTextureVS> VertexShader(ShaderMap);
+		TShaderMapRef<FHairStrandsTexturePS> PixelShader(ShaderMap);
+
+		GraphBuilder.AddPass(
+			RDG_EVENT_NAME("HairStrands::TexturePS"),
+			ParametersPS,
+			ERDGPassFlags::Raster,
+			[ParametersPS, VertexShader, PixelShader, InMeshIndexBuffer, VertexCount, PrimitiveCount, IndexBaseIndex, VertexBaseIndex, OutputResolution, OutTileCoord, TileOffsetInPixels](FRHICommandList& RHICmdList)
+			{
+				FHairStrandsTextureVS::FParameters ParametersVS;
+				ParametersVS.OutputResolution	= ParametersPS->OutputResolution;
+				ParametersVS.VertexCount 		= ParametersPS->VertexCount;
+				ParametersVS.VertexBuffer 		= ParametersPS->VertexBuffer;
+				ParametersVS.UVsBuffer 			= ParametersPS->UVsBuffer;
+				ParametersVS.NormalsBuffer 		= ParametersPS->NormalsBuffer;
+				ParametersVS.UVsChannelIndex	= ParametersPS->UVsChannelIndex;
+				ParametersVS.UVsChannelCount	= ParametersPS->UVsChannelCount;
+
+				FGraphicsPipelineStateInitializer GraphicsPSOInit;
+				RHICmdList.ApplyCachedRenderTargets(GraphicsPSOInit);
+				GraphicsPSOInit.BlendState = TStaticBlendState<
+					CW_RGBA, BO_Add, BF_One, BF_Zero, BO_Add, BF_One, BF_Zero,
+					CW_RGBA, BO_Add, BF_One, BF_Zero, BO_Add, BF_One, BF_Zero,
+					CW_RGBA, BO_Add, BF_One, BF_Zero, BO_Add, BF_One, BF_Zero,
+					CW_RGBA, BO_Add, BF_One, BF_Zero, BO_Add, BF_One, BF_Zero,
+					CW_RGBA, BO_Add, BF_One, BF_Zero, BO_Add, BF_One, BF_Zero> ::GetRHI();
+
+				GraphicsPSOInit.RasterizerState = TStaticRasterizerState<>::GetRHI();
+				GraphicsPSOInit.DepthStencilState = TStaticDepthStencilState<true, CF_LessEqual>::GetRHI();
+				GraphicsPSOInit.BoundShaderState.VertexDeclarationRHI = GFilterVertexDeclaration.VertexDeclarationRHI;
+				GraphicsPSOInit.BoundShaderState.VertexShaderRHI = VertexShader.GetVertexShader();
+				GraphicsPSOInit.BoundShaderState.PixelShaderRHI = PixelShader.GetPixelShader();
+				GraphicsPSOInit.PrimitiveType = PT_TriangleList;
+				SetGraphicsPipelineState(RHICmdList, GraphicsPSOInit, 0);
+
+				SetShaderParameters(RHICmdList, VertexShader, VertexShader.GetVertexShader(), ParametersVS);
+				SetShaderParameters(RHICmdList, PixelShader, PixelShader.GetPixelShader(), *ParametersPS);
+
+				RHICmdList.SetStreamSource(0, nullptr, 0);
+				RHICmdList.SetViewport(0, 0, 0.0f, OutputResolution.X, OutputResolution.Y, 1.0f);
+			
+				// Divide the rendering work into small batches to reduce risk of TDR as the texture projection implies heavy works 
+				// (i.e. long thread running due the the large amount of strands a groom can have)
+				const int32 TileSize = HairStrandsTextureTileSize;
+				if (OutputResolution.X > TileSize)
+				{
+					RHICmdList.SetScissorRect(true, TileOffsetInPixels.X, TileOffsetInPixels.Y, TileOffsetInPixels.X + TileSize, TileOffsetInPixels.Y + TileSize);
+				}
+				RHICmdList.DrawIndexedPrimitive(InMeshIndexBuffer, VertexBaseIndex, 0, VertexCount, IndexBaseIndex, PrimitiveCount, 1);
+			});
 	}
 
-	ParametersPS->OutTexture = GraphBuilder.CreateUAV(Out.Texture);
-	ParametersPS->RenderTargets[0] = FRenderTargetBinding(Out.TriangleMaskTexture, bClear ? ERenderTargetLoadAction::EClear : ERenderTargetLoadAction::ELoad);
-	ParametersPS->RenderTargets.DepthStencil = FDepthStencilBinding(
-		Out.DepthTestTexture,
-		bClear ? ERenderTargetLoadAction::EClear : ERenderTargetLoadAction::ELoad,
-		ERenderTargetLoadAction::ENoAction,
-		FExclusiveDepthStencil::DepthWrite_StencilNop);
-
-	TShaderMapRef<FHairStrandsTextureVS> VertexShader(ShaderMap);
-	TShaderMapRef<FHairStrandsTexturePS> PixelShader(ShaderMap);
-
-	GraphBuilder.AddPass(
-		RDG_EVENT_NAME("HairStrandsTexturePS"),
-		ParametersPS,
-		ERDGPassFlags::Raster,
-		[ParametersPS, VertexShader, PixelShader, InMeshIndexBuffer, VertexCount, PrimitiveCount, IndexBaseIndex, VertexBaseIndex, OutputResolution, OutTileCoord](FRHICommandList& RHICmdList)
+	// Evaluate material
+	{
+		FHairStrandsTextureCS::FParameters* ParametersCS = GraphBuilder.AllocParameters<FHairStrandsTextureCS::FParameters>();
+		ParametersCS->TileSize = TileSize;
+		ParametersCS->TileOffsetInPixels = TileOffsetInPixels;
+		ParametersCS->LayoutIndex = uint32(Out.Layout);
+		ParametersCS->OutputResolution = OutputResolution;
+		ParametersCS->VertexCount = VertexCount;
+		ParametersCS->VertexBuffer = InMeshVertexBuffer;
+		ParametersCS->UVsBuffer = InMeshUVsBuffer;
+		ParametersCS->NormalsBuffer = InMeshNormalsBuffer;
+		ParametersCS->MaxDistance = FMath::Abs(InMaxDistance);
+		ParametersCS->TracingDirection = InTracingDirection;
+		ParametersCS->UVsChannelIndex = UVsChannelIndex;
+		ParametersCS->UVsChannelCount = UVsChannelCount;
+		ParametersCS->InVF = Instance;
+		ParametersCS->Voxel_MinBound = VoxelMinBound;
+		ParametersCS->Voxel_MaxBound = VoxelMaxBound;
+		ParametersCS->Voxel_Resolution = VoxelResolution;
+		ParametersCS->Voxel_Size = VoxelSize;
+		ParametersCS->Voxel_MaxSegmentPerVoxel = VoxelMaxSegmentPerVoxel;
+		ParametersCS->Voxel_OffsetAndCount = GraphBuilder.CreateSRV(VoxelOffsetAndCount);
+		ParametersCS->Voxel_Data = GraphBuilder.CreateSRV(VoxelData);
+		ParametersCS->Voxel_OffsetAndCount_MaxCount = VoxelOffsetAndCount->Desc.NumElements;
+		ParametersCS->Voxel_Data_MaxCount = VoxelData->Desc.NumElements;
+		
+		ParametersCS->PositionTexture 	= PositionTexture;
+		ParametersCS->TangentTexture 	= TangentTexture;
+		ParametersCS->BitangentTexture  = BitangentTexture;
+		ParametersCS->NormalTexture 	= NormalTexture;
+		ParametersCS->UVTexture 		= UVTexture;
+		ParametersCS->OutTexture 		= GraphBuilder.CreateUAV(Out.Texture);
+
+		if (ShaderPrintData)
 		{
-			FHairStrandsTextureVS::FParameters ParametersVS;
-			ParametersVS.OutputResolution = ParametersPS->OutputResolution;
-			ParametersVS.VertexCount = ParametersPS->VertexCount;
-			ParametersVS.VertexBuffer = ParametersPS->VertexBuffer;
-			ParametersVS.UVsChannelIndex = ParametersPS->UVsChannelIndex;
-			ParametersVS.UVsChannelCount = ParametersPS->UVsChannelCount;
-			ParametersVS.UVsBuffer = ParametersPS->UVsBuffer;
-			ParametersVS.NormalsBuffer = ParametersPS->NormalsBuffer;
-
-			FGraphicsPipelineStateInitializer GraphicsPSOInit;
-			RHICmdList.ApplyCachedRenderTargets(GraphicsPSOInit);
-			GraphicsPSOInit.BlendState = TStaticBlendState<
-				CW_RGBA, BO_Add, BF_One, BF_Zero, BO_Add, BF_One, BF_Zero,
-				CW_RGBA, BO_Add, BF_One, BF_Zero, BO_Add, BF_One, BF_Zero,
-				CW_RGBA, BO_Add, BF_One, BF_Zero, BO_Add, BF_One, BF_Zero,
-				CW_RGBA, BO_Add, BF_One, BF_Zero, BO_Add, BF_One, BF_Zero,
-				CW_RGBA, BO_Add, BF_One, BF_Zero, BO_Add, BF_One, BF_Zero> ::GetRHI();
-
-			GraphicsPSOInit.RasterizerState = TStaticRasterizerState<>::GetRHI();
-			GraphicsPSOInit.DepthStencilState = TStaticDepthStencilState<true, CF_LessEqual>::GetRHI();
-			GraphicsPSOInit.BoundShaderState.VertexDeclarationRHI = GFilterVertexDeclaration.VertexDeclarationRHI;
-			GraphicsPSOInit.BoundShaderState.VertexShaderRHI = VertexShader.GetVertexShader();
-			GraphicsPSOInit.BoundShaderState.PixelShaderRHI = PixelShader.GetPixelShader();
-			GraphicsPSOInit.PrimitiveType = PT_TriangleList;
-			SetGraphicsPipelineState(RHICmdList, GraphicsPSOInit, 0);
-
-			SetShaderParameters(RHICmdList, VertexShader, VertexShader.GetVertexShader(), ParametersVS);
-			SetShaderParameters(RHICmdList, PixelShader, PixelShader.GetPixelShader(), *ParametersPS);
-
-			RHICmdList.SetStreamSource(0, nullptr, 0);
-			RHICmdList.SetViewport(0, 0, 0.0f, OutputResolution.X, OutputResolution.Y, 1.0f);
-			
-			// Divide the rendering work into small batches to reduce risk of TDR as the texture projection implies heavy works 
-			// (i.e. long thread running due the the large amount of strands a groom can have)
-			const int32 TileSize = HairStrandsTextureTileSize;
-			if (OutputResolution.X < TileSize)
-			{
-				RHICmdList.DrawIndexedPrimitive(InMeshIndexBuffer, VertexBaseIndex, 0, VertexCount, IndexBaseIndex, PrimitiveCount, 1);
-			}
-			else
-			{
-				const uint32 OffsetX = OutTileCoord.X * TileSize;
-				const uint32 OffsetY = OutTileCoord.Y * TileSize;
-				RHICmdList.SetScissorRect(true, OffsetX, OffsetY, OffsetX + TileSize, OffsetY + TileSize);
-				RHICmdList.DrawIndexedPrimitive(InMeshIndexBuffer, VertexBaseIndex, 0, VertexCount, IndexBaseIndex, PrimitiveCount, 1);
-			}
-		});
+			ShaderPrint::SetParameters(GraphBuilder, *ShaderPrintData, ParametersCS->ShaderPrintParameters);
+		}
+
+		TShaderMapRef<FHairStrandsTextureCS> ComputeShader(ShaderMap);
+		const FIntVector DispatchCount = FIntVector(FMath::DivideAndRoundUp(TileSize, 8), FMath::DivideAndRoundUp(TileSize, 8), 1);
+		FComputeShaderUtils::AddPass(
+			GraphBuilder,
+			RDG_EVENT_NAME("HairStrands::TextureCS"),
+			ComputeShader,
+			ParametersCS,
+			DispatchCount);
+	}
 }
 
 
diff --git a/Engine/Shaders/Private/HairStrands/HairStrandsTexturesGeneration.usf b/Engine/Shaders/Private/HairStrands/HairStrandsTexturesGeneration.usf
--- a/Engine/Shaders/Private/HairStrands/HairStrandsTexturesGeneration.usf
+++ b/Engine/Shaders/Private/HairStrands/HairStrandsTexturesGeneration.usf
@@ -116,6 +116,34 @@ void MainVS(
 #endif // VERTEXSHADER
 
 #if SHADER_PIXEL
+void MainPS(
+	in float4  SvPosition	: SV_Position,
+	in float2  InUV			: MESH_UV,
+	in float3  InPosition	: MESH_POSITION,
+	in float3  InTangent	: MESH_TANGENT,
+	in float3  InBitangent	: MESH_BITANGENT,
+	in float3  InNormal		: MESH_NORMAL,
+	out float4 OutPosition  : SV_Target0,
+	out float4 OutTangent   : SV_Target1,
+	out float4 OutBitangent : SV_Target2,
+	out float4 OutNormal    : SV_Target3,
+	out float2 OutUV        : SV_Target4,
+	out uint   OutTriangleMask : SV_Target5)
+{
+	const float3 T = normalize(InTangent);
+	const float3 B = normalize(InBitangent);
+	const float3 N = normalize(InNormal);
+
+	OutPosition     = float4(InPosition,1);
+	OutTangent      = float4(InTangent,1);
+	OutBitangent    = float4(InBitangent,1);
+	OutNormal       = float4(InNormal,1);
+	OutUV           = InUV;
+	OutTriangleMask = 1;
+}
+#endif //SHADER_PIXEL
+
+#if SHADER_COMPUTE 
 
 #include "HairStrandsVertexFactoryCommon.ush"
 #include "HairStrandsAttributeCommon.ush"
@@ -144,6 +172,8 @@ float3 ComputeTangent(
 	return normalize(CP1.Position - CP0.Position);
 }
 
+int2 TileSize;
+int2 TileOffsetInPixels;
 uint LayoutIndex;
 int2 OutputResolution;
 uint VertexCount;
@@ -162,7 +192,6 @@ float3 Voxel_MaxBound;
 int3 Voxel_Resolution;
 float Voxel_Size;
 uint Voxel_MaxSegmentPerVoxel;
-
 uint Voxel_OffsetAndCount_MaxCount;
 uint Voxel_Data_MaxCount;
 
@@ -170,6 +199,12 @@ StructuredBuffer<uint2> Voxel_OffsetAndCount;
 StructuredBuffer<uint2> Voxel_Data;
 RWTexture2DArray<float4> OutTexture;
 
+Texture2D<float4> PositionTexture;
+Texture2D<float4> TangentTexture;
+Texture2D<float4> BitangentTexture;
+Texture2D<float4> NormalTexture;
+Texture2D<float2> UVTexture;
+
 uint3 PositionToCoord(float3 P, float3 InMinAABB, float3 InMaxAABB, uint3 InResolution)
 {
 	return clamp(
@@ -393,22 +428,38 @@ bool Trace(
 	return bIsValid;
 }
 
-void MainPS(
-	in float4  SvPosition	: SV_Position,
-	in float2  InUV			: MESH_UV,
-	in float3  InPosition	: MESH_POSITION,
-	in float3  InTangent	: MESH_TANGENT,
-	in float3  InBitangent	: MESH_BITANGENT,
-	in float3  InNormal		: MESH_NORMAL,
-	out uint   OutTriangleMask : SV_Target0,
-	out float  OutDepthTest	: SV_Depth)
+[numthreads(8, 8, 1)]
+void MainCS(
+	uint GroupIndex : SV_GroupIndex, 
+	uint3 DispatchThreadId : SV_DispatchThreadID)
 {
+	const int2 PixelCoord = DispatchThreadId.xy + TileOffsetInPixels;
+	if (any(PixelCoord > OutputResolution))
+	{
+		return;
+	}
+
+	const float2 InUV		= UVTexture[PixelCoord];
+	const float3 InPosition = PositionTexture[PixelCoord].xyz;
+	const float3 InTangent	= TangentTexture[PixelCoord].xyz;
+	const float3 InBitangent= BitangentTexture[PixelCoord].xyz;
+	const float3 InNormal	= NormalTexture[PixelCoord].xyz;
+
 	const float3 T = normalize(InTangent);
 	const float3 B = normalize(InBitangent);
 	const float3 N = normalize(InNormal);
 
-	const float3 PixelAxisX = ddx_fine(InPosition);
-	const float3 PixelAxisY = ddy_fine(InPosition);
+	// Manual ddx/ddy
+	float3 PixelAxisX = 0;
+	float3 PixelAxisY = 0; 
+	if (PixelCoord.x + 1 < TileOffsetInPixels.x)
+	{
+		PixelAxisX = PositionTexture[float2(PixelCoord.x+1, PixelCoord.y)].xyz - InPosition;
+	}
+	if (PixelCoord.y + 1 < TileOffsetInPixels.y)
+	{
+		PixelAxisY = PositionTexture[float2(PixelCoord.x, PixelCoord.y+1)].xyz - InPosition;
+	}
 
 	float4 OutTexture0 = 0;
 	float4 OutTexture1	= 0;
@@ -428,8 +479,6 @@ void MainPS(
 	float  OutCoverage = 0;
 
 	float3 Tangent_ObjectSpace = 0;
-	const int2 PixelCoord = int2(SvPosition.xy);
-
 	const bool bDebugEnable =
 	#if DEBUG_ENABLE
 		all(PixelCoord == GetCursorPos());
@@ -517,8 +566,6 @@ void MainPS(
 	}
 
 	OutDepth = saturate(OutDepth / MaxDistance);
-	OutDepthTest = OutDepth;
-	OutTriangleMask = 1;
 
 	if (bValidCount > 0)
 	{
@@ -540,8 +587,7 @@ void MainPS(
 		OutTexture[uint3(PixelCoord, 5)] = OutTexture5;
 	}
 }
-
-#endif // SHADER_PIXEL
+#endif // SHADER_COMPUTE
 
 /////////////////////////////////////////////////////////////////////////////////////////////////////////////////////////////
 
