Fix hair data streaming bulk data beyond their actual size.

diff --git a/Engine/Plugins/Runtime/HairStrands/Source/HairStrandsCore/Private/HairStrandsDatas.cpp b/Engine/Plugins/Runtime/HairStrands/Source/HairStrandsCore/Private/HairStrandsDatas.cpp
--- a/Engine/Plugins/Runtime/HairStrands/Source/HairStrandsCore/Private/HairStrandsDatas.cpp
+++ b/Engine/Plugins/Runtime/HairStrands/Source/HairStrandsCore/Private/HairStrandsDatas.cpp
@@ -269,7 +269,7 @@ void FHairStrandsBulkCommon::FQuery::Add(FHairBulkContainer& In, const TCHAR* In
 
 		// 2. Fill in actual DDC request
 		check(OutReadIO);
-		OutReadIO->Read(In.Data, InOffset, InSize, EAsyncIOPriorityAndFlags::AIOP_Normal, Chunk.Data_IO);
+		OutReadIO->Read(In.Data, InOffset, Chunk.Size, EAsyncIOPriorityAndFlags::AIOP_Normal, Chunk.Data_IO);
 	}
 	else if (Type == ReadWriteIO)
 	{
