Fix incorrect section binding for groom.

diff --git a/Engine/Plugins/Runtime/HairStrands/Source/HairStrandsCore/Private/GroomBindingBuilder.cpp b/Engine/Plugins/Runtime/HairStrands/Source/HairStrandsCore/Private/GroomBindingBuilder.cpp
--- a/Engine/Plugins/Runtime/HairStrands/Source/HairStrandsCore/Private/GroomBindingBuilder.cpp
+++ b/Engine/Plugins/Runtime/HairStrands/Source/HairStrandsCore/Private/GroomBindingBuilder.cpp
@@ -60,7 +60,7 @@ static FAutoConsoleVariableRef CVarHairStrandsBindingBuilderWarningEnable(TEXT("
 FString FGroomBindingBuilder::GetVersion()
 {
 	// Important to update the version when groom building changes
-	return TEXT("4d");
+	return TEXT("4e");
 }
 
 ///////////////////////////////////////////////////////////////////////////////////////////////////
@@ -1187,7 +1187,9 @@ namespace GroomBinding_RootProjection
 			}
 
 			// Only cull section if the 'matching section' exist
-			const bool bCullSection = InMatchingSection > 0 && InMatchingSection < SectionCount;
+			// For now, never cull section as this can cause binding issue. The 'section' can be shuffled/reindexed 
+			// based on project settings, causing the binding on a wrong surface
+			const bool bCullSection = false; //InMatchingSection > 0 && InMatchingSection < SectionCount;
 
 			float ClosestTrianglePoint = FLT_MAX;
 			check(SectionCount > 0);
