Add normal deformation support for hair manual skin cache.

diff --git a/Engine/Plugins/Runtime/HairStrands/Source/HairStrandsCore/Private/GroomManager.cpp b/Engine/Plugins/Runtime/HairStrands/Source/HairStrandsCore/Private/GroomManager.cpp
--- a/Engine/Plugins/Runtime/HairStrands/Source/HairStrandsCore/Private/GroomManager.cpp
+++ b/Engine/Plugins/Runtime/HairStrands/Source/HairStrandsCore/Private/GroomManager.cpp
@@ -221,8 +221,10 @@ struct FHairGeometryCache
 		FSkeletalMeshLODRenderData* LODData 		= nullptr;
 		FRDGBufferRef PositionBuffer				= nullptr;
 		FRDGBufferRef PreviousPositionBuffer		= nullptr;
+		FRDGBufferRef TangentBuffer					= nullptr;
 		FRDGBufferSRVRef PositionSRV				= nullptr;
 		FRDGBufferSRVRef PreviousPositionSRV		= nullptr;
+		FRDGBufferSRVRef TangentSRV					= nullptr;
 
 		FHairGeometryCacheKey Key;
 		uint32 Hash 				= 0;
@@ -255,7 +257,8 @@ struct FHairGeometryCache
 		uint32 InLODIndex, 
 		const TArray<uint32>& UniqueSections,
 		FRDGBufferSRVRef& Out, 
-		FRDGBufferSRVRef& OutPrev)
+		FRDGBufferSRVRef& OutPrev,
+		FRDGBufferSRVRef& OutTangent)
 	{
 		check(InLODData);
 		check(InMeshObject);
@@ -289,6 +292,21 @@ struct FHairGeometryCache
 			Data.PositionSRV 			= GraphBuilder.CreateSRV(Data.PositionBuffer, PF_R32_FLOAT);
 			Data.PreviousPositionSRV	= bNeedPreviousPosition ? GraphBuilder.CreateSRV(Data.PreviousPositionBuffer, PF_R32_FLOAT) : nullptr;
 
+			// If available create and compute deformed tangents
+			Data.TangentBuffer = nullptr;
+			Data.TangentSRV = nullptr;
+			if (const FRHIShaderResourceView* TangentSRV = InLODData->StaticVertexBuffers.StaticMeshVertexBuffer.GetTangentsSRV())
+			{
+				const uint32 TangentStride = InLODData->StaticVertexBuffers.StaticMeshVertexBuffer.GetUseHighPrecisionTangentBasis() ? 
+					sizeof(TStaticMeshVertexTangentDatum<typename TStaticMeshVertexTangentTypeSelector<EStaticMeshVertexTangentBasisType::HighPrecision>::TangentTypeT>):
+					sizeof(TStaticMeshVertexTangentDatum<typename TStaticMeshVertexTangentTypeSelector<EStaticMeshVertexTangentBasisType::Default>::TangentTypeT>);
+				const EPixelFormat TangentFormat = TangentSRV->GetDesc().Buffer.SRV.Format;
+				const uint32 TangentElementCount = InLODData->StaticVertexBuffers.StaticMeshVertexBuffer.GetNumVertices();
+				check(TangentSRV->GetDesc().Buffer.SRV.BufferType == FRHIViewDesc::EBufferType::Typed);
+				Data.TangentBuffer = GraphBuilder.CreateBuffer(FRDGBufferDesc::CreateBufferDesc(TangentElementCount, TangentStride), TEXT("Hair.SkinnedDeformedTangents"));;
+				Data.TangentSRV = GraphBuilder.CreateSRV(Data.TangentBuffer, TangentFormat);
+			}
+
 			Data.RequestedSections.Reserve(Data.TotalSectionCount);
 			Data.RequestedSectionBits.Init(false, Data.TotalSectionCount);
 		}
@@ -312,8 +330,9 @@ struct FHairGeometryCache
 		}
 
 		// Initialized returned values
-		Out 	= Data.PositionSRV;
-		OutPrev = Data.PreviousPositionSRV;
+		Out 		= Data.PositionSRV;
+		OutPrev 	= Data.PreviousPositionSRV;
+		OutTangent	= Data.TangentSRV;
 	}
 
 	void AddDebug(const FHairGroupInstance* InInstance, const FPrimitiveSceneProxy* InProxy, const FCachedGeometry& InGeom, EHairPositionUpdateType InGeometryType, ECacheType InCacheType, uint32 InTotalSectionCount=0)
@@ -529,8 +548,8 @@ static void GetOrAllocateCachedGeometry(
 	// Create deformed position buffer (output)
 	FRDGBufferSRVRef DeformedPositionSRV = nullptr;
 	FRDGBufferSRVRef DeformedPreviousPositionSRV = nullptr;
-	FRDGBufferSRVRef DeformedTangentSRV = nullptr; // TODO
-	OutHairGeometryCache.GetOrAdd(GraphBuilder, SkeletalMeshObject, &LODData, LODIndex, UniqueSections, DeformedPositionSRV, DeformedPreviousPositionSRV);
+	FRDGBufferSRVRef DeformedTangentSRV = nullptr;
+	OutHairGeometryCache.GetOrAdd(GraphBuilder, SkeletalMeshObject, &LODData, LODIndex, UniqueSections, DeformedPositionSRV, DeformedPreviousPositionSRV, DeformedTangentSRV);
 
 	// Add reference to be sure the data are not streamed out while they are used
 	LODData.AddRef();
@@ -545,6 +564,7 @@ static void GetOrAllocateCachedGeometry(
 		OutSection.PositionBuffer 			= nullptr; // Do not use the SRV slot, but instead use the RDG buffer created above (DeformedPositionSRV)
 		OutSection.PreviousPositionBuffer 	= nullptr; // Do not use the SRV slot, but instead use the RDG buffer created above (DeformedPositionSRV)
 		OutSection.UVsBuffer 				= LODData.StaticVertexBuffers.StaticMeshVertexBuffer.GetTexCoordsSRV();
+		OutSection.TangentBuffer 			= LODData.StaticVertexBuffers.StaticMeshVertexBuffer.GetTangentsSRV();
 		OutSection.TotalVertexCount 		= LODData.StaticVertexBuffers.PositionVertexBuffer.GetNumVertices();
 		OutSection.IndexBuffer 				= LODData.MultiSizeIndexContainer.GetIndexBuffer()->GetSRV();
 		OutSection.TotalIndexCount 			= LODData.MultiSizeIndexContainer.GetIndexBuffer()->Num();
@@ -742,7 +762,8 @@ static void RunHairBindingSurfaceUpdate(
 			*Data.LODData, 
 			Data.RequestedSections, 
 			Data.PositionBuffer, 
-			Data.PreviousPositionBuffer);
+			Data.PreviousPositionBuffer,
+			Data.TangentBuffer);
 	}
 
 	AddHairSkinCacheDebugPass(GraphBuilder, ShaderMap, View, ShaderPrintData, Instances, HairGeometryCache);
diff --git a/Engine/Plugins/Runtime/HairStrands/Source/HairStrandsCore/Private/HairStrandsMeshProjection.cpp b/Engine/Plugins/Runtime/HairStrands/Source/HairStrandsCore/Private/HairStrandsMeshProjection.cpp
--- a/Engine/Plugins/Runtime/HairStrands/Source/HairStrandsCore/Private/HairStrandsMeshProjection.cpp
+++ b/Engine/Plugins/Runtime/HairStrands/Source/HairStrandsCore/Private/HairStrandsMeshProjection.cpp
@@ -48,6 +48,7 @@ class FSkinUpdateCS : public FGlobalShader
 	using FPermutationDomain = TShaderPermutationDomain<FUnlimitedBoneInfluence, FUseExtraInfluence, FBoneIndexUint16, FBoneWeightUint16, FPrevious>;
 
 	BEGIN_SHADER_PARAMETER_STRUCT(FParameters, )
+		SHADER_PARAMETER(uint32, bWriteTangent)
 		SHADER_PARAMETER(uint32, NumVertexToProcess)
 		SHADER_PARAMETER(uint32, NumTotalVertices)
 		SHADER_PARAMETER(uint32, SectionVertexBaseIndex)
@@ -58,8 +59,10 @@ class FSkinUpdateCS : public FGlobalShader
 		SHADER_PARAMETER_SRV(Buffer<float4>, PrevBoneMatrices)
 		SHADER_PARAMETER_SRV(Buffer<uint>, VertexWeights)
 		SHADER_PARAMETER_SRV(Buffer<float>, RestPositions)
+		SHADER_PARAMETER_SRV(Buffer<float4>, RestTangents)
 		SHADER_PARAMETER_RDG_BUFFER_UAV(RWBuffer<float>, DeformedPositions)
 		SHADER_PARAMETER_RDG_BUFFER_UAV(RWBuffer<float>, PrevDeformedPositions)
+		SHADER_PARAMETER_RDG_BUFFER_UAV(RWBuffer<float4>, DeformedTangents)
 	END_SHADER_PARAMETER_STRUCT()
 
 public:
@@ -81,13 +84,31 @@ void AddSkinUpdatePass(
 	FSkeletalMeshLODRenderData& RenderData,
 	const TArray<FSkinUpdateSection>& Sections,
 	FRDGBufferRef OutDeformedPositionBuffer,
-	FRDGBufferRef OutPrevDeformedPositionBuffer)
+	FRDGBufferRef OutPrevDeformedPositionBuffer,
+	FRDGBufferRef OutDeformedTangentBuffer)
 {
 	check(Sections.Num() > 0);
 	const FSkinWeightVertexBuffer* SkinWeight = &RenderData.SkinWeightVertexBuffer;
 	const bool bPrevPosition = OutPrevDeformedPositionBuffer != nullptr && Sections[0].BonePrevBuffer;
 	
+	bool bHasTangentSRV = false;
+	FRDGBufferUAVRef DeformedTangentsUAV = nullptr;
+	if (OutDeformedTangentBuffer)
+	{
+		const FRHIShaderResourceView* TangentSRV = RenderData.StaticVertexBuffers.StaticMeshVertexBuffer.GetTangentsSRV();
+		check(TangentSRV);
+		const EPixelFormat TangentFormat = TangentSRV->GetDesc().Buffer.SRV.Format;
+		DeformedTangentsUAV = GraphBuilder.CreateUAV(OutDeformedTangentBuffer, TangentFormat, ERDGUnorderedAccessViewFlags::SkipBarrier);
+		bHasTangentSRV = true;
+	}
+	else
+	{
+		// Dummy output
+		DeformedTangentsUAV = GraphBuilder.CreateUAV(GraphBuilder.CreateBuffer(FRDGBufferDesc::CreateBufferDesc(4u, 1), TEXT("DummyUAV")), PF_R8G8B8A8_SNORM, ERDGUnorderedAccessViewFlags::SkipBarrier);
+	}
+
 	FSkinUpdateCS::FParameters* Parameters = GraphBuilder.AllocParameters<FSkinUpdateCS::FParameters>();
+	Parameters->bWriteTangent			= bHasTangentSRV ? 1 : 0;
 	Parameters->WeightIndexSize 		= SkinWeight->GetBoneIndexByteSize() | (SkinWeight->GetBoneWeightByteSize() << 8);
 	Parameters->NumVertexToProcess 		= 0;
 	Parameters->NumTotalVertices		= RenderData.StaticVertexBuffers.PositionVertexBuffer.GetNumVertices();
@@ -98,7 +119,9 @@ void AddSkinUpdatePass(
 	Parameters->PrevBoneMatrices 		= Sections[0].BonePrevBuffer;
 	Parameters->VertexWeights 			= SkinWeight->GetDataVertexBuffer()->GetSRV();
 	Parameters->RestPositions 			= RenderData.StaticVertexBuffers.PositionVertexBuffer.GetSRV();
+	Parameters->RestTangents 			= RenderData.StaticVertexBuffers.StaticMeshVertexBuffer.GetTangentsSRV();
 	Parameters->DeformedPositions 		= GraphBuilder.CreateUAV(OutDeformedPositionBuffer, PF_R32_FLOAT, ERDGUnorderedAccessViewFlags::SkipBarrier);
+	Parameters->DeformedTangents 		= DeformedTangentsUAV;
 	if (bPrevPosition)
 	{
 		Parameters->PrevDeformedPositions = GraphBuilder.CreateUAV(OutPrevDeformedPositionBuffer, PF_R32_FLOAT, ERDGUnorderedAccessViewFlags::SkipBarrier);
@@ -159,7 +182,7 @@ private:
 		SHADER_PARAMETER_RDG_BUFFER_SRV(StructuredBuffer, MeshSectionBuffer)
 		SHADER_PARAMETER_RDG_BUFFER_SRV(Buffer, RDGMeshPositionBuffer)
 		SHADER_PARAMETER_RDG_BUFFER_SRV(Buffer, RDGMeshPreviousPositionBuffer)
-		SHADER_PARAMETER_RDG_BUFFER_SRV(Buffer, RDGTangentBuffer)
+		SHADER_PARAMETER_RDG_BUFFER_SRV(Buffer, RDGMeshTangentBuffer)
 		SHADER_PARAMETER_SRV(Buffer, MeshPositionBuffer)
 		SHADER_PARAMETER_SRV(Buffer, MeshPreviousPositionBuffer)
 		SHADER_PARAMETER_SRV(Buffer, MeshIndexBuffer)
@@ -209,7 +232,7 @@ static bool AddHairStrandUpdateMeshTrianglesPass(
 	CommonParameters.TangentFormat					= MeshLODData.Sections[0].TangentFormat;
 	CommonParameters.RDGMeshPositionBuffer			= MeshLODData.Sections[0].RDGPositionBuffer;
 	CommonParameters.RDGMeshPreviousPositionBuffer 	= MeshLODData.Sections[0].RDGPreviousPositionBuffer;
-	CommonParameters.RDGTangentBuffer				= MeshLODData.Sections[0].RDGTangentBuffer;
+	CommonParameters.RDGMeshTangentBuffer			= MeshLODData.Sections[0].RDGTangentBuffer;
 	CommonParameters.MeshPositionBuffer				= MeshLODData.Sections[0].PositionBuffer;
 	CommonParameters.MeshPreviousPositionBuffer		= MeshLODData.Sections[0].PreviousPositionBuffer;
 	CommonParameters.MeshIndexBuffer				= MeshLODData.Sections[0].IndexBuffer;
diff --git a/Engine/Plugins/Runtime/HairStrands/Source/HairStrandsCore/Public/HairStrandsMeshProjection.h b/Engine/Plugins/Runtime/HairStrands/Source/HairStrandsCore/Public/HairStrandsMeshProjection.h
--- a/Engine/Plugins/Runtime/HairStrands/Source/HairStrandsCore/Public/HairStrandsMeshProjection.h
+++ b/Engine/Plugins/Runtime/HairStrands/Source/HairStrandsCore/Public/HairStrandsMeshProjection.h
@@ -93,7 +93,8 @@ void AddSkinUpdatePass(
 	FSkeletalMeshLODRenderData& RenderData,
 	const TArray<FSkinUpdateSection>& Sections,
 	FRDGBufferRef OutDeformedPositionBuffer,
-	FRDGBufferRef OutPrevDeformedPositionBuffer);
+	FRDGBufferRef OutPrevDeformedPositionBuffer,
+	FRDGBufferRef OutTangentBuffer);
 
 void AddHairMeshesRBFInterpolationPass(
 	FRDGBuilder& GraphBuilder,
diff --git a/Engine/Shaders/Private/HairStrands/HairStrandsBindingCommon.ush b/Engine/Shaders/Private/HairStrands/HairStrandsBindingCommon.ush
--- a/Engine/Shaders/Private/HairStrands/HairStrandsBindingCommon.ush
+++ b/Engine/Shaders/Private/HairStrands/HairStrandsBindingCommon.ush
@@ -53,7 +53,7 @@ float3 GetVertexNormal(uint VId, Buffer<float4> InTangentBuffer, Buffer<uint> In
 }
 
 // Return a mesh triangle (position + normal)
-FHairMeshTriangle GetHairMeshTriangle(uint TriangleIndex, Buffer<float> InPositionBuffer, Buffer<float4> InTangentBuffer, Buffer<uint> InIndexBuffer, uint InIndexOffset, uint InMaxIndexCount, uint InMaxVertexCount, uint InTangentFormat)
+FHairMeshTriangle GetHairMeshTriangle(uint TriangleIndex, Buffer<float> InPositionBuffer, Buffer<float4> InTangentBuffer, Buffer<uint> InIndexBuffer, uint InIndexOffset, uint InMaxIndexCount, uint InMaxVertexCount, uint InTangentFormat, bool bUseVertexTangent)
 {
 	const uint BaseVertexId = TriangleIndex * 3;
 	const uint VId0 = BaseVertexId;
@@ -65,9 +65,12 @@ FHairMeshTriangle GetHairMeshTriangle(uint TriangleIndex, Buffer<float> InPositi
 	Out.P1 = GetVertexPosition(VId1, InPositionBuffer, InIndexBuffer, InIndexOffset, InMaxIndexCount, InMaxVertexCount);
 	Out.P2 = GetVertexPosition(VId2, InPositionBuffer, InIndexBuffer, InIndexOffset, InMaxIndexCount, InMaxVertexCount);
 
-	Out.N0 = GetVertexNormal(VId0, InTangentBuffer, InIndexBuffer, InIndexOffset, InMaxIndexCount, InMaxVertexCount, InTangentFormat);
-	Out.N1 = GetVertexNormal(VId1, InTangentBuffer, InIndexBuffer, InIndexOffset, InMaxIndexCount, InMaxVertexCount, InTangentFormat);
-	Out.N2 = GetVertexNormal(VId2, InTangentBuffer, InIndexBuffer, InIndexOffset, InMaxIndexCount, InMaxVertexCount, InTangentFormat);
+	if (bUseVertexTangent)
+	{
+		Out.N0 = GetVertexNormal(VId0, InTangentBuffer, InIndexBuffer, InIndexOffset, InMaxIndexCount, InMaxVertexCount, InTangentFormat);
+		Out.N1 = GetVertexNormal(VId1, InTangentBuffer, InIndexBuffer, InIndexOffset, InMaxIndexCount, InMaxVertexCount, InTangentFormat);
+		Out.N2 = GetVertexNormal(VId2, InTangentBuffer, InIndexBuffer, InIndexOffset, InMaxIndexCount, InMaxVertexCount, InTangentFormat);
+	}
 	
 	return Out;
 }
diff --git a/Engine/Shaders/Private/HairStrands/HairStrandsMesh.usf b/Engine/Shaders/Private/HairStrands/HairStrandsMesh.usf
--- a/Engine/Shaders/Private/HairStrands/HairStrandsMesh.usf
+++ b/Engine/Shaders/Private/HairStrands/HairStrandsMesh.usf
@@ -148,7 +148,6 @@ struct FMeshSectionData
 uint TangentFormat;
 uint MaxSectionCount;
 uint MaxUniqueTriangleCount;
-Buffer<float4> MeshTangentBuffer;
 Buffer<uint> MeshIndexBuffer;
 Buffer<uint> UniqueTriangleIndices;
 StructuredBuffer<uint4> MeshSectionBuffer;
@@ -158,11 +157,14 @@ StructuredBuffer<uint4> MeshSectionBuffer;
 // (while the skin cache system is not converted to RDG). If skin cache is disable, a manual skin compute 
 // update happens in the hair system, and its output are RDG buffers.
 #if PERMUTATION_POSITION_TYPE == 0
+Buffer<float4> MeshTangentBuffer;
 Buffer<float> MeshPositionBuffer;
 Buffer<float> MeshPreviousPositionBuffer;
 #else
+Buffer<float4> RDGMeshTangentBuffer;
 Buffer<float> RDGMeshPositionBuffer;
 Buffer<float> RDGMeshPreviousPositionBuffer;
+#define MeshTangentBuffer RDGMeshTangentBuffer
 #define MeshPositionBuffer RDGMeshPositionBuffer
 #define MeshPreviousPositionBuffer RDGMeshPreviousPositionBuffer
 #endif 
@@ -187,7 +189,7 @@ FMeshSectionData ReadMeshSectionData(uint InSectionIndex)
 	return Out;
 }
 
-FHairMeshTriangle GetHairMeshTriangle(uint InTriangleIndex, FMeshSectionData MeshSectionData, bool bPrevious)
+FHairMeshTriangle GetHairMeshTriangle(uint InTriangleIndex, FMeshSectionData MeshSectionData, bool bPrevious, bool bUseVertexTangent)
 {
 	const bool bSamplePrev = (MeshSectionData.bIsSwapped > 0 && !bPrevious) || (MeshSectionData.bIsSwapped == 0 && bPrevious);
 	if (bSamplePrev)
@@ -200,7 +202,8 @@ FHairMeshTriangle GetHairMeshTriangle(uint InTriangleIndex, FMeshSectionData Mes
 			MeshSectionData.IndexOffset,
 			MeshSectionData.MaxIndexCount,
 			MeshSectionData.MaxVertexCount,
-			TangentFormat);
+			TangentFormat,
+			bUseVertexTangent);
 	}
 	else
 	{
@@ -212,7 +215,8 @@ FHairMeshTriangle GetHairMeshTriangle(uint InTriangleIndex, FMeshSectionData Mes
 			MeshSectionData.IndexOffset,
 			MeshSectionData.MaxIndexCount,
 			MeshSectionData.MaxVertexCount,
-			TangentFormat);
+			TangentFormat,
+			bUseVertexTangent);
 	}
 }
 
@@ -239,7 +243,7 @@ void MainCS(uint3 DispatchThreadId : SV_DispatchThreadID)
 	
 		// Compute the slot index of the resources for the section index (this is done in multiple pass
 		{
-			FHairMeshTriangle Tri = GetHairMeshTriangle(TriangleIndex.TriangleIndex, MeshSectionData, false);
+			FHairMeshTriangle Tri = GetHairMeshTriangle(TriangleIndex.TriangleIndex, MeshSectionData, false, !MeshSectionData.bUseFaceNormal);
 			if (MeshSectionData.bUseFaceNormal)
 			{
 				Tri.N0 = GetTriangleFaceNormal(Tri);
@@ -252,7 +256,7 @@ void MainCS(uint3 DispatchThreadId : SV_DispatchThreadID)
 		}
 		#if PERMUTATION_PREVIOUS
 		{
-			FHairMeshTriangle Tri = GetHairMeshTriangle(TriangleIndex.TriangleIndex, MeshSectionData, true);
+			FHairMeshTriangle Tri = GetHairMeshTriangle(TriangleIndex.TriangleIndex, MeshSectionData, true, !MeshSectionData.bUseFaceNormal);
 			if (MeshSectionData.bUseFaceNormal)
 			{
 				Tri.N0 = GetTriangleFaceNormal(Tri);
diff --git a/Engine/Shaders/Private/HairStrands/HairStrandsSkinUpdate.usf b/Engine/Shaders/Private/HairStrands/HairStrandsSkinUpdate.usf
--- a/Engine/Shaders/Private/HairStrands/HairStrandsSkinUpdate.usf
+++ b/Engine/Shaders/Private/HairStrands/HairStrandsSkinUpdate.usf
@@ -8,19 +8,23 @@
 #if SHADER_SKIN_UPDATE
 //----------------------------------------------------------------------------------
 
+uint bWriteTangent;
 uint WeightIndexSize; // bone index size: 1 - 8bit or 2 - 16bit
 uint SectionVertexBaseIndex;
 uint NumTotalVertices;
 uint NumVertexToProcess;
 uint WeightStride;
+uint TangentFormat;
 
 Buffer<uint> WeightLookup;
 Buffer<uint> VertexWeights;
 Buffer<float> RestPositions;
+Buffer<SNORM float4> RestTangents;
 
 // Current Position
 Buffer<float4> BoneMatrices;
 RWBuffer<float> DeformedPositions;
+RWBuffer<TANGENT_RWBUFFER_FORMAT> DeformedTangents; // GL does not support SNORM images, so instead we write to SINT
 
 // Previous Position
 #if PERMUTATION_PREV
@@ -63,6 +67,18 @@ void UpdateSkinPositionCS(uint3 GroupID : SV_GroupID,
 			DeformedPositions[PositionOffset  ] = OutputPosition.x;
 			DeformedPositions[PositionOffset+1] = OutputPosition.y;
 			DeformedPositions[PositionOffset+2] = OutputPosition.z;
+
+			if (bWriteTangent)
+			{
+				const float4 InputTangentZ = TangentBias_SkinCache(RestTangents[2 * VertexIndex + GPUSKIN_RWBUFFER_OFFSET_TANGENT_Z]);
+				//const float4 InputTangentZ = TangentBias(RestTangents[2 * VertexIndex + GPUSKIN_RWBUFFER_OFFSET_TANGENT_Z]);
+
+				// Use similar code than GPU skin cache, but this should use the inverse transpose for correct normal stransformation
+				//const const float4x4 InvTransposeBlendedMatrix = Transpose(Inverse(BlendedMatrix));
+				const float3 OutputTangentZ = mul(BlendedMatrix, float4(InputTangentZ.xyz,0)).xyz;
+			
+				DeformedTangents[VertexIndex * 2 + GPUSKIN_RWBUFFER_OFFSET_TANGENT_Z] = TangentUnbias_SkinCache(float4(OutputTangentZ, InputTangentZ.w));
+			}
 		}
 
 		// Previous position
