Fix hair data issuing 0-byte requests to DDC.

This can cause the Zen server to cancel all batches requests, causing hair data to be missing. This fix ensures only non-0-byte requests are emitted.

diff --git a/Engine/Plugins/Runtime/HairStrands/Source/HairStrandsCore/Private/HairStrandsDatas.cpp b/Engine/Plugins/Runtime/HairStrands/Source/HairStrandsCore/Private/HairStrandsDatas.cpp
--- a/Engine/Plugins/Runtime/HairStrands/Source/HairStrandsCore/Private/HairStrandsDatas.cpp
+++ b/Engine/Plugins/Runtime/HairStrands/Source/HairStrandsCore/Private/HairStrandsDatas.cpp
@@ -224,26 +224,42 @@ void FHairStrandsBulkCommon::FQuery::Add(FHairBulkContainer& In, const TCHAR* In
 
 		check (InSize >= InOffset);
 
-		// 1. Add chunk request to the streaming request. The chunk will hold the request result.
-		FHairStreamingRequest::FChunk& Chunk = StreamingRequest->Chunks.AddDefaulted_GetRef();
-		Chunk.Status 	= FHairStreamingRequest::FChunk::EStatus::Pending;
-		Chunk.Container = &In;
-		Chunk.Size 		= InSize - InOffset;
-		Chunk.Offset 	= InOffset;
-		Chunk.TotalSize = InSize;
-		In.ChunkRequest = &Chunk;
-
-		// 2. Fill in actual DDC request
-		check(OutReadDDC);
-		using namespace UE::DerivedData;
-		FCacheGetChunkRequest& Out = OutReadDDC->AddDefaulted_GetRef();
-		Out.Id			= FValueId::Null; 	// HairStrands::HairStrandsValueId : This is only needed for cache record, not cache value.
-		Out.Key			= ConvertLegacyCacheKey(*DerivedDataKey + InSuffix);
-		Out.RawOffset	= InOffset;
-		Out.RawSize		= InSize != 0 ? InSize-InOffset : MAX_uint64;
-		Out.RawHash		= FIoHash();
-		Out.UserData	= (uint64)&Chunk;
-		if (Owner) { Out.Name = Owner->GetPathName(); }
+		// DDC does not support 0-bytes requests, so bypass them.
+		const bool bHasValidSize = int32(InSize) - int32(InOffset) > 0;
+		if (bHasValidSize)
+		{
+			// 1. Add chunk request to the streaming request. The chunk will hold the request result.
+			FHairStreamingRequest::FChunk& Chunk = StreamingRequest->Chunks.AddDefaulted_GetRef();
+			Chunk.Status 	= FHairStreamingRequest::FChunk::EStatus::Pending;
+			Chunk.Container = &In;
+			Chunk.Size 		= InSize - InOffset;
+			Chunk.Offset 	= InOffset;
+			Chunk.TotalSize = InSize;
+			In.ChunkRequest = &Chunk;
+
+			// 2. Fill in actual DDC request
+			check(OutReadDDC);
+			using namespace UE::DerivedData;
+			FCacheGetChunkRequest& Out = OutReadDDC->AddDefaulted_GetRef();
+			Out.Id			= FValueId::Null; 	// HairStrands::HairStrandsValueId : This is only needed for cache record, not cache value.
+			Out.Key			= ConvertLegacyCacheKey(*DerivedDataKey + InSuffix);
+			Out.RawOffset	= InOffset;
+			Out.RawSize		= InSize != 0 ? InSize-InOffset : MAX_uint64;
+			Out.RawHash		= FIoHash();
+			Out.UserData	= (uint64)&Chunk;
+			if (Owner) { Out.Name = Owner->GetPathName(); }
+		}
+		else
+		{
+			// Add a (dummy) chunk request, so that the buffer creation path will be identical to the streaming path
+			FHairStreamingRequest::FChunk& Chunk = StreamingRequest->Chunks.AddDefaulted_GetRef();
+			Chunk.Status 	= FHairStreamingRequest::FChunk::EStatus::Completed;
+			Chunk.Container = &In;
+			Chunk.Size 		= InSize - InOffset;
+			Chunk.Offset 	= InOffset;
+			Chunk.TotalSize = InSize;
+			In.ChunkRequest = &Chunk;
+		}
 	}
 	else 
 #endif
