Add hair show flags to disable hair rendering on certain view.

diff --git a/Engine/Plugins/Runtime/HairStrands/Source/HairStrandsCore/Private/GroomComponent.cpp b/Engine/Plugins/Runtime/HairStrands/Source/HairStrandsCore/Private/GroomComponent.cpp
--- a/Engine/Plugins/Runtime/HairStrands/Source/HairStrandsCore/Private/GroomComponent.cpp
+++ b/Engine/Plugins/Runtime/HairStrands/Source/HairStrandsCore/Private/GroomComponent.cpp
@@ -1075,6 +1075,7 @@ public:
 		// but visible in shadow 'hidden shadow') so that raytracing geometry is created/updated correctly
 		const bool bPathtracing = View->Family->EngineShowFlags.PathTracing;
 		const bool bForceDrawRelevance = bPathtracing && (!IsShown(View) && IsShadowCast(View));
+		const bool bVisible = View->Family->EngineShowFlags.Hair;
 
 		bool bUseCardsOrMesh = false;
 		for (const TRefCountPtr<FHairGroupInstance>& Instance : HairGroupInstances)
@@ -1087,7 +1088,7 @@ public:
 		FPrimitiveViewRelevance Result;
 
 		// Special pass for hair strands geometry (not part of the base pass, and shadowing is handlded in a custom fashion). When cards rendering is enabled we reusethe base pass
-		Result.bDrawRelevance		= IsShown(View) || bForceDrawRelevance;
+		Result.bDrawRelevance		= bVisible && (IsShown(View) || bForceDrawRelevance);
 		Result.bRenderInMainPass	= bUseCardsOrMesh && ShouldRenderInMainPass();
 		Result.bShadowRelevance		= IsShadowCast(View);
 		Result.bDynamicRelevance	= bUseCardsOrMesh;
diff --git a/Engine/Plugins/Runtime/HairStrands/Source/HairStrandsCore/Private/GroomManager.cpp b/Engine/Plugins/Runtime/HairStrands/Source/HairStrandsCore/Private/GroomManager.cpp
--- a/Engine/Plugins/Runtime/HairStrands/Source/HairStrandsCore/Private/GroomManager.cpp
+++ b/Engine/Plugins/Runtime/HairStrands/Source/HairStrandsCore/Private/GroomManager.cpp
@@ -3148,6 +3148,11 @@ void ProcessHairStrandsBookmark(
 		return;
 	}
 
+	if (!Parameters.View || (Parameters.View->Family && !Parameters.View->Family->EngineShowFlags.Hair))
+	{
+		return;
+	}
+
 	if (IsInstanceFrustumCullingEnable())
 	{
 		Instances = nullptr;
diff --git a/Engine/Source/Runtime/Engine/Public/ShowFlagsValues.inl b/Engine/Source/Runtime/Engine/Public/ShowFlagsValues.inl
--- a/Engine/Source/Runtime/Engine/Public/ShowFlagsValues.inl
+++ b/Engine/Source/Runtime/Engine/Public/ShowFlagsValues.inl
@@ -455,5 +455,8 @@ SHOWFLAG_ALWAYS_ACCESSIBLE(SceneCaptureCopySceneDepth, SFG_Hidden, NSLOCTEXT("Un
 /** If false, camera's are not shown in the viewport (but are enabled and available for use). */
 SHOWFLAG_ALWAYS_ACCESSIBLE(Cameras, SFG_Normal, NSLOCTEXT("UnrealEd", "CamerasSF", "Cameras"))
 
+/** If false, hair geometry won't be rendered. */
+SHOWFLAG_ALWAYS_ACCESSIBLE(Hair, SFG_Advanced, NSLOCTEXT("UnrealEd", "HairSF", "Hair"))
+
 #undef SHOWFLAG_ALWAYS_ACCESSIBLE
 #undef SHOWFLAG_FIXED_IN_SHIPPING
diff --git a/Engine/Source/Runtime/Renderer/Private/HairStrands/HairStrandsRendering.cpp b/Engine/Source/Runtime/Renderer/Private/HairStrands/HairStrandsRendering.cpp
--- a/Engine/Source/Runtime/Renderer/Private/HairStrands/HairStrandsRendering.cpp
+++ b/Engine/Source/Runtime/Renderer/Private/HairStrands/HairStrandsRendering.cpp
@@ -120,7 +120,7 @@ void RenderHairPrePass(
 	for (FViewInfo& View : Views)
 	{
 		const bool bIsViewCompatible = IsHairStrandsEnabled(EHairStrandsShaderType::Strands, View.GetShaderPlatform());
-		if (!View.Family || !bIsViewCompatible)
+		if (!View.Family || !bIsViewCompatible || !View.Family->EngineShowFlags.Hair)
 			continue;
 
 		// For stereo rendering, hair groups/voxelization/deep-shadow are only produced once
@@ -180,7 +180,7 @@ void RenderHairBasePass(
 	for (FViewInfo& View : Views)
 	{
 		const bool bIsViewCompatible = IsHairStrandsEnabled(EHairStrandsShaderType::Strands, View.GetShaderPlatform());
-		if (View.Family && bIsViewCompatible && View.HairStrandsViewData.MacroGroupDatas.Num() > 0)
+		if (View.Family && View.Family->EngineShowFlags.Hair && bIsViewCompatible && View.HairStrandsViewData.MacroGroupDatas.Num() > 0)
 		{
 			RenderHairStrandsVisibilityBuffer(
 				GraphBuilder, 
