Fixed validation error due to incorrect padding in shader struct.

diff --git a/Engine/Plugins/Runtime/HairStrands/Source/HairStrandsCore/Private/HairStrandsMeshProjection.cpp b/Engine/Plugins/Runtime/HairStrands/Source/HairStrandsCore/Private/HairStrandsMeshProjection.cpp
--- a/Engine/Plugins/Runtime/HairStrands/Source/HairStrandsCore/Private/HairStrandsMeshProjection.cpp
+++ b/Engine/Plugins/Runtime/HairStrands/Source/HairStrandsCore/Private/HairStrandsMeshProjection.cpp
@@ -252,7 +252,7 @@ static bool AddHairStrandUpdateMeshTrianglesPass(
 		uint32 UVsChannelCount : 8;
 		uint32 bIsSwapped : 1;
 		uint32 bUseFaceNormal : 1;
-		uint32 Pad = 14;
+		uint32 Pad : 14;
 	};
 
 	FSectionData Default;
