[GroomBinding]
 - Make compilation async and protect member access during async compilation
 - 7m47s to 0.000537s of GroomBinding postloads time in CitySample Big_City_LVL -PIE -DDC-All-MissTypes=Groom,SkeletalMesh,GroomBinding

diff --git a/Engine/Plugins/Runtime/HairStrands/Source/HairStrandsCore/Private/GroomAsset.cpp b/Engine/Plugins/Runtime/HairStrands/Source/HairStrandsCore/Private/GroomAsset.cpp
--- a/Engine/Plugins/Runtime/HairStrands/Source/HairStrandsCore/Private/GroomAsset.cpp
+++ b/Engine/Plugins/Runtime/HairStrands/Source/HairStrandsCore/Private/GroomAsset.cpp
@@ -16,6 +16,7 @@
 #include "HairStrandsVertexFactory.h"
 #include "HAL/LowLevelMemTracker.h"
 #include "Misc/App.h"
+#include "AssetCompilingManager.h"
 #include "Misc/Paths.h"
 #include "RHIStaticStates.h"
 #include "Serialization/LargeMemoryReader.h"
@@ -1565,6 +1566,16 @@ static bool IsMeshesAttributes(const FName PropertyName)
 	return PropertyName == UGroomAsset::GetHairGroupsMeshesMemberName();
 }
 
+bool UGroomAsset::Modify(bool bAlwaysMarkDirty)
+{
+	// We are not really async compilable at the moment, but this will
+	// cue the other compilers that depends on this asset (i.e. UGroomBindings)
+	// to finish their async compilation before we start any modification.
+	FAssetCompilingManager::Get().FinishCompilationForObjects({ this });
+
+	return Super::Modify(bAlwaysMarkDirty);
+}
+
 void UGroomAsset::PostEditChangeProperty(FPropertyChangedEvent& PropertyChangedEvent)
 {
 	Super::PostEditChangeProperty(PropertyChangedEvent);
diff --git a/Engine/Plugins/Runtime/HairStrands/Source/HairStrandsCore/Private/GroomBindingAsset.cpp b/Engine/Plugins/Runtime/HairStrands/Source/HairStrandsCore/Private/GroomBindingAsset.cpp
--- a/Engine/Plugins/Runtime/HairStrands/Source/HairStrandsCore/Private/GroomBindingAsset.cpp
+++ b/Engine/Plugins/Runtime/HairStrands/Source/HairStrandsCore/Private/GroomBindingAsset.cpp
@@ -6,13 +6,16 @@
 #include "GeometryCache.h"
 #include "GeometryCacheMeshData.h"
 #include "GroomAsset.h"
+#include "GroomBindingCompiler.h"
 #include "GroomBindingBuilder.h"
 #include "HAL/LowLevelMemTracker.h"
 #include "Serialization/MemoryReader.h"
+#include "UObject/UObjectIterator.h"
 #include "UObject/AnimObjectVersion.h"
 #include "UObject/ObjectSaveContext.h"
 #include "UObject/DevObjectVersion.h"
 #include "Misc/CoreMisc.h"
+#include "Interfaces/ITargetPlatform.h"
 #include "Interfaces/ITargetPlatformManagerModule.h"
 
 #include UE_INLINE_GENERATED_CPP_BY_NAME(GroomBindingAsset)
@@ -124,6 +127,108 @@ static void InternalSerializePlatformDatas(FArchive& Ar, UObject* Owner, TArray<
 	}
 }
 
+/*-----------------------------------------------------------------------------
+	FGroomBindingAsyncBuildScope
+-----------------------------------------------------------------------------*/
+thread_local const UGroomBindingAsset* FGroomBindingAsyncBuildScope::Asset = nullptr;
+
+FGroomBindingAsyncBuildScope::FGroomBindingAsyncBuildScope(const UGroomBindingAsset* InAsset)
+{
+	PreviousScope = Asset;
+	Asset = InAsset;
+}
+
+FGroomBindingAsyncBuildScope::~FGroomBindingAsyncBuildScope()
+{
+	check(Asset);
+	Asset = PreviousScope;
+}
+
+bool FGroomBindingAsyncBuildScope::ShouldWaitOnLockedProperties(const UGroomBindingAsset* InAsset)
+{
+	return Asset != InAsset;
+}
+
+/*-----------------------------------------------------------------------------
+	FGroomBindingAsyncBuildWorker
+-----------------------------------------------------------------------------*/
+void FGroomBindingAsyncBuildWorker::DoWork()
+{
+#if WITH_EDITOR
+	if (BuildContext.IsSet())
+	{
+		GroomBinding->ExecuteCacheDerivedDatas(*BuildContext);
+	}
+#endif
+}
+
+void UGroomBindingAsset::WaitUntilAsyncPropertyReleased(EGroomBindingAsyncProperties AsyncProperties, EGroomBindingAsyncPropertyLockType LockType) const
+{
+#if WITH_EDITOR
+	// We need to protect internal asset data from race conditions during async build
+	if (IsCompiling())
+	{
+		if (FGroomBindingAsyncBuildScope::ShouldWaitOnLockedProperties(this))
+		{
+			bool bIsLocked = true;
+			// We can remove the lock if we're accessing in read-only and there is no write-lock
+			if ((LockType & EGroomBindingAsyncPropertyLockType::ReadOnly) == EGroomBindingAsyncPropertyLockType::ReadOnly)
+			{
+				// Maintain the lock if the write-lock bit is non-zero
+				bIsLocked &= (ModifiedProperties & (uint64)AsyncProperties) != 0;
+			}
+
+			if (bIsLocked)
+			{
+				FString PropertyName = StaticEnum<EGroomBindingAsyncProperties>()->GetNameByValue((int64)AsyncProperties).ToString();
+				TRACE_CPUPROFILER_EVENT_SCOPE_TEXT(*FString::Printf(TEXT("GroomBindingCompilationStall %s"), *PropertyName));
+
+				if (IsInGameThread())
+				{
+					UE_LOG(
+						LogHairStrands,
+						Verbose,
+						TEXT("Accessing property %s of the GroomBindingAsset while it is still being built asynchronously will force it to be compiled before continuing. "
+							"For better performance, consider making the caller async aware so it can wait until the groom binding is ready to access this property."
+							"To better understand where those calls are coming from, you can use Editor.AsyncAssetDumpStallStacks on the console."),
+						*PropertyName
+					);
+
+					FGroomBindingCompilingManager::Get().FinishCompilation({ const_cast<UGroomBindingAsset*>(this) });
+				}
+				else
+				{
+					// Trying to access a property from another thread that cannot force finish the compilation is invalid
+					ensureMsgf(
+						false,
+						TEXT("Accessing property %s of the GroomBindingAsset while it is still being built asynchronously is only supported on the game-thread. "
+							"To avoid any race-condition, consider finishing the compilation before pushing tasks to other threads or making higher-level game-thread code async aware so it "
+							"schedules the task only when the groom binding's compilation is finished. If this is a blocker, you can disable async groom binding compilation from the editor experimental settings."),
+						*PropertyName
+					);
+				}
+			}
+		}
+		// If we're accessing this property from the async build thread, make sure the property is still protected from access from other threads.
+		else
+		{
+			bool bIsLocked = true;
+			if ((LockType & EGroomBindingAsyncPropertyLockType::ReadOnly) == EGroomBindingAsyncPropertyLockType::ReadOnly)
+			{
+				bIsLocked &= (AccessedProperties & (uint64)AsyncProperties) != 0;
+			}
+
+			if ((LockType & EGroomBindingAsyncPropertyLockType::WriteOnly) == EGroomBindingAsyncPropertyLockType::WriteOnly)
+			{
+				bIsLocked &= (ModifiedProperties & (uint64)AsyncProperties) != 0;
+			}
+			ensureMsgf(bIsLocked, TEXT("Property %s has not been locked properly for use by async build"), *StaticEnum<EGroomBindingAsyncProperties>()->GetNameByValue((int64)AsyncProperties).ToString());
+		}
+	}
+#endif
+}
+
+
 ////////////////////////////////////////////////////////////////////////////////////////////////////////////
 
 void UGroomBindingAsset::Serialize(FArchive& Ar)
@@ -134,6 +239,21 @@ void UGroomBindingAsset::Serialize(FArchive& Ar)
 	{
 		Flags = GetGroom()->GenerateClassStripFlags(Ar);
 	}
+
+	if (IsCompiling())
+	{
+		// Skip serialization during compilation if told to do so.
+		if (Ar.ShouldSkipCompilingAssets())
+		{
+			return;
+		}
+
+		// Since UPROPERTY are accessed directly by offset during serialization instead of using accessors,
+		// the protection put in place to automatically finish compilation if a locked property is accessed will not work.
+		// We have no choice but to force finish the compilation here to avoid potential race conditions between 
+		// async compilation and the serialization.
+		FGroomBindingCompilingManager::Get().FinishCompilation({ this });
+	}
 #endif
 
 	Super::Serialize(Ar);
@@ -330,11 +450,21 @@ void UGroomBindingAsset::Reset()
 #if WITH_EDITORONLY_DATA
 void UGroomBindingAsset::InvalidateBinding(class USkeletalMesh*)
 {
-	CacheDerivedDatas();
+	InvalidateBinding();
 }
 
 void UGroomBindingAsset::InvalidateBinding()
 {
+	const bool bDependenciesStillCaching =
+		(GetTargetSkeletalMesh() && GetTargetSkeletalMesh()->IsCompiling()) ||
+		(GetSourceSkeletalMesh() && GetSourceSkeletalMesh()->IsCompiling());
+
+	// Nothing else to do as we're going to be rebuilt when our dependencies finish caching anyway
+	if (bDependenciesStillCaching)
+	{
+		return;
+	}
+
 	CacheDerivedDatas();
 }
 
@@ -355,29 +485,39 @@ void UGroomBindingAsset::PostLoad()
 		LocalGroom->ConditionalPostLoad();
 
 	#if WITH_EDITOR
-		CacheDerivedDatas();
+		USkeletalMesh* LocalTargetSkeletalMesh = GetTargetSkeletalMesh();
+		USkeletalMesh* LocalSourceSkeletalMesh = GetSourceSkeletalMesh();
 
-		// Sanity check. This function will report back warnings/issues into the log for user.
-		UGroomBindingAsset::IsCompatible(LocalGroom, this, true);
-
-		if (USkeletalMesh* LocalTargetSkeletalMesh = GetTargetSkeletalMesh())
+		if (LocalTargetSkeletalMesh)
 		{
+			LocalTargetSkeletalMesh->ConditionalPostLoad();
 			LocalTargetSkeletalMesh->OnPostMeshCached().AddUObject(this, &UGroomBindingAsset::InvalidateBinding);
 			bRegisterTargetMeshCallback = true;
 		}
 
-		if (USkeletalMesh* LocalSourceSkeletalMesh = GetSourceSkeletalMesh())
+		// Do not register the same skeletal mesh twice as it will make us invalidate bindings twice in a row and cause a stall.
+		if (LocalSourceSkeletalMesh && LocalSourceSkeletalMesh != LocalTargetSkeletalMesh)
 		{
+			LocalSourceSkeletalMesh->ConditionalPostLoad();
 			LocalSourceSkeletalMesh->OnPostMeshCached().AddUObject(this, &UGroomBindingAsset::InvalidateBinding);
 			bRegisterSourceMeshCallback = true;
 		}
 
-		if (LocalGroom)
+		if (GetTargetGeometryCache())
+		{
+			GetTargetGeometryCache()->ConditionalPostLoad();
+		}
+		
+		if (GetSourceGeometryCache())
 		{
-			LocalGroom->GetOnGroomAssetChanged().AddUObject(this, &UGroomBindingAsset::InvalidateBinding);
-			LocalGroom->GetOnGroomAssetResourcesChanged().AddUObject(this, &UGroomBindingAsset::InvalidateBinding);
-			bRegisterGroomAssetCallback = true;
+			GetSourceGeometryCache()->ConditionalPostLoad();
 		}
+
+		LocalGroom->GetOnGroomAssetChanged().AddUObject(this, &UGroomBindingAsset::InvalidateBinding);
+		LocalGroom->GetOnGroomAssetResourcesChanged().AddUObject(this, &UGroomBindingAsset::InvalidateBinding);
+		bRegisterGroomAssetCallback = true;
+
+		InvalidateBinding();
 	#endif
 	}
 
@@ -786,7 +926,7 @@ namespace GroomBindingDerivedDataCacheUtils
 
 	FString BuildGroomBindingDerivedDataKey(const FString& KeySuffix)
 	{
-		return FDerivedDataCacheInterface::BuildCacheKey(*(TEXT("GROOM_BINDING_V") + FGroomBindingBuilder::GetVersion() + TEXT("_")), *GetGroomBindingDerivedDataVersion(), *KeySuffix);
+		return FDerivedDataCacheInterface::BuildCacheKey(*(TEXT("GROOMBINDING_V") + FGroomBindingBuilder::GetVersion() + TEXT("_")), *GetGroomBindingDerivedDataVersion(), *KeySuffix);
 	}
 }
 
@@ -843,6 +983,69 @@ void UGroomBindingAsset::CacheDerivedDatas()
 		return;
 	}
 
+	if (IsCompiling())
+	{
+		FGroomBindingCompilingManager::Get().FinishCompilation({ this });
+	}
+
+	bool bNeedFlushRenderingCommand = false;
+	for (TObjectIterator<UGroomComponent> HairStrandsComponentIt; HairStrandsComponentIt; ++HairStrandsComponentIt)
+	{
+		if (HairStrandsComponentIt->BindingAsset == this ||
+			HairStrandsComponentIt->BindingAssetBeingLoaded == this) // A GroomAsset was set on the component while it was still loading
+		{
+			if (HairStrandsComponentIt->IsRenderStateCreated())
+			{
+				HairStrandsComponentIt->DestroyRenderState_Concurrent();
+				bNeedFlushRenderingCommand = true;
+			}
+		}
+	}
+
+	if (bNeedFlushRenderingCommand)
+	{
+		// Flush the rendering commands generated by the detachments.
+		FlushRenderingCommands();
+	}
+
+	FGroomBindingBuildContext Context;
+#if WITH_EDITOR
+	// Acquire everything in readonly so that any thread trying to write a property locked for read-only will wait until the build is finished
+	AcquireAsyncProperty(MAX_uint64, EGroomBindingAsyncPropertyLockType::ReadOnly);
+
+	// Acquire the resource modified by the build so that any thread trying to read those will wait until the build is finished
+	// If we forget something here and the build uses it, we will get an assert on the build thread because we're writing to an unlocked property
+	AcquireAsyncProperty((uint64)EGroomBindingAsyncProperties::HairGroupResources, EGroomBindingAsyncPropertyLockType::WriteOnly);
+	AcquireAsyncProperty((uint64)EGroomBindingAsyncProperties::HairGroupPlatformData, EGroomBindingAsyncPropertyLockType::WriteOnly);
+	AcquireAsyncProperty((uint64)EGroomBindingAsyncProperties::GroupInfos, EGroomBindingAsyncPropertyLockType::WriteOnly);
+
+	// Do synchronous build until GeometryCache is thread-safe
+	const bool bIsAsyncBuildSupported = GetTargetGeometryCache() == nullptr && GetSourceGeometryCache() == nullptr;
+
+	// Inline reduction is not thread-safe, prevent async build until this is fixed.
+	if (bIsAsyncBuildSupported && FGroomBindingCompilingManager::Get().IsAsyncCompilationAllowed(this))
+	{
+		FQueuedThreadPool* ThreadPool = FGroomBindingCompilingManager::Get().GetThreadPool();
+		EQueuedWorkPriority BasePriority = FGroomBindingCompilingManager::Get().GetBasePriority(this);
+		check(AsyncTask == nullptr);
+		AsyncTask = MakeUnique<FGroomBindingAsyncBuildTask>(this, MoveTemp(Context));
+		AsyncTask->StartBackgroundTask(ThreadPool, BasePriority, EQueuedWorkFlags::DoNotRunInsideBusyWait);
+		FGroomBindingCompilingManager::Get().AddGroomBindings({ this });
+	}
+	else
+#endif
+	{
+		ExecuteCacheDerivedDatas(Context);
+		FinishCacheDerivedDatas(Context);
+	}
+}
+
+void UGroomBindingAsset::ExecuteCacheDerivedDatas(FGroomBindingBuildContext& Context)
+{
+	TRACE_CPUPROFILER_EVENT_SCOPE(UGroomBindingAsset::ExecuteCacheDerivedDatas);
+
+	FGroomBindingAsyncBuildScope BuildScope(this);
+
 	// 1. Set group count to the groom target
 	const uint32 GroupCount = GetGroom()->GetNumHairGroups();
 	GetHairGroupsPlatformData().SetNum(GroupCount);
@@ -864,8 +1067,6 @@ void UGroomBindingAsset::CacheDerivedDatas()
 	bool bReloadResource = false;
 	if (bAnyGroupNeedRebuild)
 	{
-		FGroomComponentRecreateRenderStateContext RecreateRenderContext(GetGroom());
-
 		for (uint32 GroupIndex = 0; GroupIndex < GroupCount; ++GroupIndex)
 		{
 			// 1. Build data
@@ -911,11 +1112,7 @@ void UGroomBindingAsset::CacheDerivedDatas()
 		//    use the binding infos to know if curve count match between GroomAsset and GroomBindingAsset
 		UpdateGroomBindingAssetInfos(this);
 
-		// 4. Reload resources if needed
-		if (bReloadResource)
-		{
-			InitResource();
-		}
+		Context.bReloadResource |= bReloadResource;
 	}
 	else
 	{
@@ -954,7 +1151,7 @@ static void CacheDerivedDatas(UGroomBindingAsset* In, const uint32 InGroupIndex,
 		bool bHasDataInCache = false;
 		if (Data)
 		{
-			UE_CLOG(IsHairStrandsDDCLogEnable(), LogHairStrands, Log, TEXT("[GroomBinding/DDC] Found (GroomBinding:%s)."), *In->GetName());
+			UE_CLOG(IsHairStrandsDDCLogEnable(), LogHairStrands, Log, TEXT("[GroomBinding/DDC] Found (GroomBinding:%s TargetPlatform:%s Async:%s)."), *In->GetName(), *TargetPlatform->PlatformName(), IsInGameThread() ? TEXT("No") : TEXT("Yes"));
 
 			// Header
 			FMemoryReaderView Ar(Data, /*bIsPersistent*/ true);
@@ -976,7 +1173,7 @@ static void CacheDerivedDatas(UGroomBindingAsset* In, const uint32 InGroupIndex,
 		}
 		if (!bHasDataInCache)
 		{
-			UE_CLOG(IsHairStrandsDDCLogEnable(), LogHairStrands, Log, TEXT("[GroomBinding/DDC] Not found (GroomBinding:%s)."), *In->GetName());
+			UE_CLOG(IsHairStrandsDDCLogEnable(), LogHairStrands, Log, TEXT("[GroomBinding/DDC] Not found (GroomBinding:%s TargetPlatform:%s Async:%s)."), *In->GetName(), *TargetPlatform->PlatformName(), IsInGameThread() ? TEXT("No") : TEXT("Yes"));
 
 			// Build groom binding data
 			FGroomBindingBuilder::FInput BuilderInput;
@@ -1022,6 +1219,30 @@ static void CacheDerivedDatas(UGroomBindingAsset* In, const uint32 InGroupIndex,
 	}
 }
 
+void UGroomBindingAsset::FinishCacheDerivedDatas(FGroomBindingBuildContext& Context)
+{
+	ReleaseAsyncProperty();
+
+	if (!IsTemplate() && Context.bReloadResource && IsValid())
+	{
+		InitResource();
+	}
+
+	for (TObjectIterator<UGroomComponent> HairStrandsComponentIt; HairStrandsComponentIt; ++HairStrandsComponentIt)
+	{
+		if (HairStrandsComponentIt->BindingAsset == this ||
+			HairStrandsComponentIt->BindingAssetBeingLoaded == this) // A GroomAsset was set on the component while it was still loading
+		{
+			HairStrandsComponentIt->PostCompilation();
+
+			if (HairStrandsComponentIt->IsRegistered())
+			{
+				HairStrandsComponentIt->RecreateRenderState_Concurrent();
+			}
+		}
+	}
+}
+
 static TArray<FString> GetGroupDerivedDataKeys(const UGroomBindingAsset* In, const ITargetPlatform* TargetPlatform)
 {
 	check(In);
@@ -1053,6 +1274,12 @@ static UGroomBindingAsset::FCachedCookedPlatformData* FindCachedCookedPlatformDa
 
 void UGroomBindingAsset::BeginCacheForCookedPlatformData(const ITargetPlatform* TargetPlatform)
 {
+	// Finish any async compilation of the editor target before building other platforms
+	if (IsCompiling())
+	{
+		FGroomBindingCompilingManager::Get().FinishCompilation({ this });
+	}
+
 	Super::BeginCacheForCookedPlatformData(TargetPlatform);
 
 	// 1. Build the key for each group
@@ -1167,6 +1394,7 @@ FName UGroomBindingAsset::GetAssetPathName(int32 LODIndex)
 	UFUNCTION(BlueprintGetter)\
 	GetConst Type Access UGroomBindingAsset::Get##Name() const\
 	{\
+		WaitUntilAsyncPropertyReleased(EGroomBindingAsyncProperties::Name, EGroomBindingAsyncPropertyLockType::ReadOnly); \
 		PRAGMA_DISABLE_DEPRECATION_WARNINGS\
 		return Name;\
 		PRAGMA_ENABLE_DEPRECATION_WARNINGS\
@@ -1174,6 +1402,7 @@ FName UGroomBindingAsset::GetAssetPathName(int32 LODIndex)
 	UFUNCTION(BlueprintSetter)\
 	void UGroomBindingAsset::Set##Name(SetConst Type Access In##Name)\
 	{\
+		WaitUntilAsyncPropertyReleased(EGroomBindingAsyncProperties::Name, EGroomBindingAsyncPropertyLockType::WriteOnly); \
 		PRAGMA_DISABLE_DEPRECATION_WARNINGS\
 		Name = In##Name;\
 		PRAGMA_ENABLE_DEPRECATION_WARNINGS\
@@ -1192,6 +1421,7 @@ DEFINE_GROOMBINDING_MEMBER_ACCESSOR(TArray<FGoomBindingGroupInfo>, &, GroupInfos
 
 TArray<FGoomBindingGroupInfo>& UGroomBindingAsset::GetGroupInfos()
 {
+	WaitUntilAsyncPropertyReleased(EGroomBindingAsyncProperties::GroupInfos, EGroomBindingAsyncPropertyLockType::ReadWrite);
 	PRAGMA_DISABLE_DEPRECATION_WARNINGS
 	return GroupInfos;
 	PRAGMA_ENABLE_DEPRECATION_WARNINGS
@@ -1220,6 +1450,7 @@ FName UGroomBindingAsset::GetHairGroupResourcesMemberName()
 
 UGroomBindingAsset::FHairGroupResources& UGroomBindingAsset::GetHairGroupResources()
 {
+	WaitUntilAsyncPropertyReleased(EGroomBindingAsyncProperties::HairGroupResources, EGroomBindingAsyncPropertyLockType::ReadWrite);
 	PRAGMA_DISABLE_DEPRECATION_WARNINGS
 	return HairGroupResources;
 	PRAGMA_ENABLE_DEPRECATION_WARNINGS
@@ -1227,6 +1458,7 @@ UGroomBindingAsset::FHairGroupResources& UGroomBindingAsset::GetHairGroupResourc
 
 const UGroomBindingAsset::FHairGroupResources& UGroomBindingAsset::GetHairGroupResources() const
 {
+	WaitUntilAsyncPropertyReleased(EGroomBindingAsyncProperties::HairGroupResources, EGroomBindingAsyncPropertyLockType::ReadOnly);
 	PRAGMA_DISABLE_DEPRECATION_WARNINGS
 	return HairGroupResources;
 	PRAGMA_ENABLE_DEPRECATION_WARNINGS
@@ -1234,6 +1466,7 @@ const UGroomBindingAsset::FHairGroupResources& UGroomBindingAsset::GetHairGroupR
 
 void UGroomBindingAsset::SetHairGroupResources(UGroomBindingAsset::FHairGroupResources InHairGroupResources)
 {
+	WaitUntilAsyncPropertyReleased(EGroomBindingAsyncProperties::HairGroupResources, EGroomBindingAsyncPropertyLockType::WriteOnly);
 	PRAGMA_DISABLE_DEPRECATION_WARNINGS
 	HairGroupResources = InHairGroupResources;
 	PRAGMA_ENABLE_DEPRECATION_WARNINGS
@@ -1248,6 +1481,7 @@ FName UGroomBindingAsset::GetHairGroupPlatformDataMemberName()
 
 const TArray<UGroomBindingAsset::FHairGroupPlatformData>& UGroomBindingAsset::GetHairGroupsPlatformData() const
 {
+	WaitUntilAsyncPropertyReleased(EGroomBindingAsyncProperties::HairGroupPlatformData, EGroomBindingAsyncPropertyLockType::ReadOnly);
 	PRAGMA_DISABLE_DEPRECATION_WARNINGS
 	return HairGroupsPlatformData;
 	PRAGMA_ENABLE_DEPRECATION_WARNINGS
@@ -1255,6 +1489,7 @@ const TArray<UGroomBindingAsset::FHairGroupPlatformData>& UGroomBindingAsset::Ge
 
 TArray<UGroomBindingAsset::FHairGroupPlatformData>& UGroomBindingAsset::GetHairGroupsPlatformData()
 {
+	WaitUntilAsyncPropertyReleased(EGroomBindingAsyncProperties::HairGroupPlatformData, EGroomBindingAsyncPropertyLockType::ReadWrite);
 	PRAGMA_DISABLE_DEPRECATION_WARNINGS
 	return HairGroupsPlatformData;
 	PRAGMA_ENABLE_DEPRECATION_WARNINGS
diff --git a/Engine/Plugins/Runtime/HairStrands/Source/HairStrandsCore/Private/GroomBindingBuilder.cpp b/Engine/Plugins/Runtime/HairStrands/Source/HairStrandsCore/Private/GroomBindingBuilder.cpp
--- a/Engine/Plugins/Runtime/HairStrands/Source/HairStrandsCore/Private/GroomBindingBuilder.cpp
+++ b/Engine/Plugins/Runtime/HairStrands/Source/HairStrandsCore/Private/GroomBindingBuilder.cpp
@@ -10,7 +10,9 @@
 #include "HairStrandsMeshProjection.h"
 #include "Async/ParallelFor.h"
 #include "GlobalShader.h"
+#include "Misc/CoreMisc.h"
 #include "Misc/ScopedSlowTask.h"
+#include "Interfaces/ITargetPlatformManagerModule.h"
 #include "GroomRBFDeformer.h"
 #include "Engine/SkinnedAssetAsyncCompileUtils.h"
 #include "Interfaces/ITargetPlatform.h"
@@ -2225,6 +2227,9 @@ static bool InternalBuildBinding_CPU(const FGroomBindingBuilder::FInput& In, uin
 	// 1. Build groom root data
 	FHairRootGroupData OutData;
 	{
+		// If we're currently running on a worker thread, all this preloading
+		// stuff should have happened on the game-thread part of the build
+		// before going async so these should be no-ops in that case.
 		In.GroomAsset->ConditionalPostLoad();
 
 		// Ensure the skeletal meshes / geom caches are built
@@ -2247,13 +2252,19 @@ static bool InternalBuildBinding_CPU(const FGroomBindingBuilder::FInput& In, uin
 			}
 		}
 
+		// If a skeletal mesh build or preedit is called while we're async compiling groom bindings, the groom binding compiler
+		// will take care of finishing any groom binding that depends on the skeletal mesh being modified. So this is 
+		// safe to do asynchronously without locks on the render data. We assume that the skeletal mesh's render data is
+		// immutable once it has been built and can only be rebuilt throught a call to PreEditChange first.
+		const bool bAsyncCompiling = !IsInGameThread();
+
 		// * Only for SkeletalMesh: Take scoped lock on the skeletal render mesh data during the entire groom binding building
 		// * Then use an async build scope to allow accessing skeletal mesh property safely.
 		//   If skel.meshes are nullptr, this will act as a NOP
 		USkeletalMesh* InSourceSkeletalMesh = In.SourceSkeletalMesh == In.TargetSkeletalMesh ? nullptr : In.SourceSkeletalMesh;
 		PRAGMA_DISABLE_DEPRECATION_WARNINGS
-		FScopedSkeletalMeshRenderData SourceSkeletalMeshScopedData(InSourceSkeletalMesh);
-		FScopedSkeletalMeshRenderData TargetSkeletalMeshScopedData(In.TargetSkeletalMesh);
+		FScopedSkeletalMeshRenderData SourceSkeletalMeshScopedData(bAsyncCompiling ? nullptr : InSourceSkeletalMesh);
+		FScopedSkeletalMeshRenderData TargetSkeletalMeshScopedData(bAsyncCompiling ? nullptr : In.TargetSkeletalMesh);
 		PRAGMA_ENABLE_DEPRECATION_WARNINGS
 
 		TUniquePtr<GroomBinding_Mesh::IMeshData> SourceMeshData;
@@ -2262,11 +2273,20 @@ static bool InternalBuildBinding_CPU(const FGroomBindingBuilder::FInput& In, uin
 		{
 			if (InSourceSkeletalMesh)
 			{
-				FSkinnedAssetAsyncBuildScope AsyncBuildScope(InSourceSkeletalMesh);
-				USkeletalMesh::GetPlatformSkeletalMeshRenderData(TargetPlatform, SourceSkeletalMeshScopedData);
-				PRAGMA_DISABLE_DEPRECATION_WARNINGS
-				SourceMeshData = TUniquePtr<GroomBinding_Mesh::FSkeletalMeshData, TDefaultDelete<GroomBinding_Mesh::IMeshData>>(new GroomBinding_Mesh::FSkeletalMeshData(InSourceSkeletalMesh, SourceSkeletalMeshScopedData.GetData()));
-				PRAGMA_ENABLE_DEPRECATION_WARNINGS
+				if (bAsyncCompiling)
+				{
+					const ITargetPlatform* RunningPlatform = GetTargetPlatformManagerRef().GetRunningTargetPlatform();
+					checkf(TargetPlatform == RunningPlatform, TEXT("It is only safe to query the running platform's render data asynchronously from the skeletal mesh"));
+					SourceMeshData = TUniquePtr<GroomBinding_Mesh::FSkeletalMeshData, TDefaultDelete<GroomBinding_Mesh::IMeshData>>(new GroomBinding_Mesh::FSkeletalMeshData(InSourceSkeletalMesh, InSourceSkeletalMesh->GetResourceForRendering()));
+				}
+				else
+				{
+					FSkinnedAssetAsyncBuildScope AsyncBuildScope(InSourceSkeletalMesh);
+					USkeletalMesh::GetPlatformSkeletalMeshRenderData(TargetPlatform, SourceSkeletalMeshScopedData);
+					PRAGMA_DISABLE_DEPRECATION_WARNINGS
+					SourceMeshData = TUniquePtr<GroomBinding_Mesh::FSkeletalMeshData, TDefaultDelete<GroomBinding_Mesh::IMeshData>>(new GroomBinding_Mesh::FSkeletalMeshData(InSourceSkeletalMesh, SourceSkeletalMeshScopedData.GetData()));
+					PRAGMA_ENABLE_DEPRECATION_WARNINGS
+				}
 			}
 			else
 			{
@@ -2275,11 +2295,20 @@ static bool InternalBuildBinding_CPU(const FGroomBindingBuilder::FInput& In, uin
 
 			if (In.TargetSkeletalMesh)
 			{
-				FSkinnedAssetAsyncBuildScope AsyncBuildScope(In.TargetSkeletalMesh);
-				USkeletalMesh::GetPlatformSkeletalMeshRenderData(TargetPlatform, TargetSkeletalMeshScopedData);
-				PRAGMA_DISABLE_DEPRECATION_WARNINGS
-				TargetMeshData = TUniquePtr<GroomBinding_Mesh::FSkeletalMeshData, TDefaultDelete<GroomBinding_Mesh::IMeshData>>(new GroomBinding_Mesh::FSkeletalMeshData(In.TargetSkeletalMesh, TargetSkeletalMeshScopedData.GetData()));
-				PRAGMA_ENABLE_DEPRECATION_WARNINGS
+				if (bAsyncCompiling)
+				{
+					const ITargetPlatform* RunningPlatform = GetTargetPlatformManagerRef().GetRunningTargetPlatform();
+					checkf(TargetPlatform == RunningPlatform, TEXT("It is only safe to query the running platform's render data asynchronously from the skeletal mesh"));
+					TargetMeshData = TUniquePtr<GroomBinding_Mesh::FSkeletalMeshData, TDefaultDelete<GroomBinding_Mesh::IMeshData>>(new GroomBinding_Mesh::FSkeletalMeshData(In.TargetSkeletalMesh, In.TargetSkeletalMesh->GetResourceForRendering()));
+				}
+				else
+				{
+					FSkinnedAssetAsyncBuildScope AsyncBuildScope(In.TargetSkeletalMesh);
+					USkeletalMesh::GetPlatformSkeletalMeshRenderData(TargetPlatform, TargetSkeletalMeshScopedData);
+					PRAGMA_DISABLE_DEPRECATION_WARNINGS
+					TargetMeshData = TUniquePtr<GroomBinding_Mesh::FSkeletalMeshData, TDefaultDelete<GroomBinding_Mesh::IMeshData>>(new GroomBinding_Mesh::FSkeletalMeshData(In.TargetSkeletalMesh, TargetSkeletalMeshScopedData.GetData()));
+					PRAGMA_ENABLE_DEPRECATION_WARNINGS
+				}
 			}
 			else
 			{
diff --git a/Engine/Plugins/Runtime/HairStrands/Source/HairStrandsCore/Private/GroomBindingCompiler.cpp b/Engine/Plugins/Runtime/HairStrands/Source/HairStrandsCore/Private/GroomBindingCompiler.cpp
new file mode 100644
--- /dev/null
+++ b/Engine/Plugins/Runtime/HairStrands/Source/HairStrandsCore/Private/GroomBindingCompiler.cpp
@@ -0,0 +1,500 @@
+// Copyright Epic Games, Inc. All Rights Reserved.
+
+#include "GroomBindingCompiler.h"
+
+#if WITH_EDITOR
+
+#include "AssetCompilingManager.h"
+#include "GroomBindingAsset.h"
+#include "ObjectCacheContext.h"
+#include "Settings/EditorExperimentalSettings.h"
+#include "GroomComponent.h"
+#include "Misc/QueuedThreadPoolWrapper.h"
+#include "EngineModule.h"
+#include "Misc/ScopedSlowTask.h"
+#include "GroomComponent.h"
+#include "UObject/UObjectIterator.h"
+#include "Containers/Set.h"
+#include "SkeletalMeshCompiler.h"
+#include "Misc/IQueuedWork.h"
+#include "LevelEditor.h"
+#include "SLevelViewport.h"
+#include "Components/PrimitiveComponent.h"
+#include "ContentStreaming.h"
+#include "EngineUtils.h"
+#include "ProfilingDebugging/CountersTrace.h"
+
+#define LOCTEXT_NAMESPACE "GroomBindingCompiler"
+
+static AsyncCompilationHelpers::FAsyncCompilationStandardCVars CVarAsyncGroomBindingAssetStandard(
+	TEXT("GroomBinding"),
+	TEXT("groom bindings"),
+	FConsoleCommandDelegate::CreateLambda(
+		[]()
+		{
+			FGroomBindingCompilingManager::Get().FinishAllCompilation();
+		}
+	));
+
+namespace GroomBindingAssetCompilingManagerImpl
+{
+	static void EnsureInitializedCVars()
+	{
+		static bool bIsInitialized = false;
+
+		if (!bIsInitialized)
+		{
+			bIsInitialized = true;
+			
+			AsyncCompilationHelpers::EnsureInitializedCVars(
+				TEXT("groombinding"),
+				CVarAsyncGroomBindingAssetStandard.AsyncCompilation,
+				CVarAsyncGroomBindingAssetStandard.AsyncCompilationMaxConcurrency,
+				GET_MEMBER_NAME_CHECKED(UEditorExperimentalSettings, bEnableAsyncGroomBindingCompilation));
+		}
+	}
+}
+
+FGroomBindingCompilingManager::FGroomBindingCompilingManager()
+	: Notification(GetAssetNameFormat())
+{
+	GroomBindingAssetCompilingManagerImpl::EnsureInitializedCVars();
+
+	PostReachabilityAnalysisHandle = FCoreUObjectDelegates::PostReachabilityAnalysis.AddRaw(this, &FGroomBindingCompilingManager::OnPostReachabilityAnalysis);
+}
+
+void FGroomBindingCompilingManager::AttachDependencies(UGroomBindingAsset* GroomBindingAsset)
+{
+	if (USkeletalMesh* TargetSkeletalMesh = GroomBindingAsset->GetTargetSkeletalMesh())
+	{
+		RegisteredSkeletalMeshes.Add(TargetSkeletalMesh, GroomBindingAsset);
+	}
+
+	if (USkeletalMesh* SourceSkeletalMesh = GroomBindingAsset->GetSourceSkeletalMesh())
+	{
+		RegisteredSkeletalMeshes.Add(SourceSkeletalMesh, GroomBindingAsset);
+	}
+
+	if (UGroomAsset* GroomAsset = GroomBindingAsset->GetGroom())
+	{
+		RegisteredGroomAssets.Add(GroomAsset, GroomBindingAsset);
+	}
+}
+
+void FGroomBindingCompilingManager::DetachDependencies(UGroomBindingAsset* GroomBindingAsset)
+{
+	if (USkeletalMesh* TargetSkeletalMesh = GroomBindingAsset->GetTargetSkeletalMesh())
+	{
+		RegisteredSkeletalMeshes.Remove(TargetSkeletalMesh, GroomBindingAsset);
+	}
+
+	if (USkeletalMesh* SourceSkeletalMesh = GroomBindingAsset->GetSourceSkeletalMesh())
+	{
+		RegisteredSkeletalMeshes.Remove(SourceSkeletalMesh, GroomBindingAsset);
+	}
+
+	if (UGroomAsset* GroomAsset = GroomBindingAsset->GetGroom())
+	{
+		RegisteredGroomAssets.Remove(GroomAsset, GroomBindingAsset);
+	}
+}
+
+void FGroomBindingCompilingManager::OnPostReachabilityAnalysis()
+{
+	if (GetNumRemainingJobs())
+	{
+		TRACE_CPUPROFILER_EVENT_SCOPE(FGroomBindingCompilingManager::CancelUnreachableGroomBindings);
+
+		TArray<UGroomBindingAsset*> PendingAssets;
+		PendingAssets.Reserve(GetNumRemainingJobs());
+
+		for (auto Iterator = RegisteredGroomBindingAssets.CreateIterator(); Iterator; ++Iterator)
+		{
+			UGroomBindingAsset* GroomBinding = *Iterator;
+			if (GroomBinding->IsUnreachable())
+			{
+				UE_LOG(LogHairStrands, Verbose, TEXT("Cancelling groom binding %s async compilation because it's being garbage collected"), *GroomBinding->GetName());
+
+				if (GroomBinding->TryCancelAsyncTasks())
+				{
+					Iterator.RemoveCurrent();
+					DetachDependencies(GroomBinding);
+				}
+				else
+				{
+					PendingAssets.Add(GroomBinding);
+				}
+			}
+		}
+
+		FinishCompilation(PendingAssets);
+	}
+}
+
+FName FGroomBindingCompilingManager::GetAssetTypeName() const
+{
+	return TEXT("UE-GroomBinding");
+}
+
+FTextFormat FGroomBindingCompilingManager::GetAssetNameFormat() const
+{
+	return LOCTEXT("GroomBindingNameFormat", "{0}|plural(one=Groom Binding,other=Groom Bindings)");
+}
+
+TArrayView<FName> FGroomBindingCompilingManager::GetDependentTypeNames() const
+{
+	// GroomBindings may have dependencies to SkeletalMesh so we prefer processing them before we get called
+	// so a single FinishAllCompilation is able to finish everything in a single pass.
+	static FName DependentTypeNames[] = 
+	{
+		FSkinnedAssetCompilingManager::GetStaticAssetTypeName(),
+	};
+	return TArrayView<FName>(DependentTypeNames);
+}
+
+int32 FGroomBindingCompilingManager::GetNumRemainingAssets() const
+{
+	return GetNumRemainingJobs();
+}
+
+EQueuedWorkPriority FGroomBindingCompilingManager::GetBasePriority(UGroomBindingAsset* InGroomBindingAsset) const
+{
+	return EQueuedWorkPriority::Low;
+}
+
+FQueuedThreadPool* FGroomBindingCompilingManager::GetThreadPool() const
+{
+	static FQueuedThreadPoolDynamicWrapper* GGroomBindingAssetThreadPool = nullptr;
+	if (GGroomBindingAssetThreadPool == nullptr && FAssetCompilingManager::Get().GetThreadPool() != nullptr)
+	{
+		TFunction<EQueuedWorkPriority(EQueuedWorkPriority)> PriorityMapper = [](EQueuedWorkPriority) { return EQueuedWorkPriority::Low; };
+
+		// GroomBinding assets will be scheduled on the asset thread pool, where concurrency limits might by dynamically adjusted depending on memory constraints.
+		GGroomBindingAssetThreadPool = new FQueuedThreadPoolDynamicWrapper(FAssetCompilingManager::Get().GetThreadPool(), -1, PriorityMapper);
+
+		AsyncCompilationHelpers::BindThreadPoolToCVar(
+			GGroomBindingAssetThreadPool,
+			CVarAsyncGroomBindingAssetStandard.AsyncCompilation,
+			CVarAsyncGroomBindingAssetStandard.AsyncCompilationResume,
+			CVarAsyncGroomBindingAssetStandard.AsyncCompilationMaxConcurrency
+		);
+	}
+
+	return GGroomBindingAssetThreadPool;
+}
+
+void FGroomBindingCompilingManager::Shutdown()
+{
+	bHasShutdown = true;
+	if (GetNumRemainingJobs())
+	{
+		check(IsInGameThread());
+		TRACE_CPUPROFILER_EVENT_SCOPE(FGroomBindingCompilingManager::Shutdown);
+
+		for (auto Iterator = RegisteredGroomBindingAssets.CreateIterator(); Iterator; ++Iterator)
+		{
+			UGroomBindingAsset* GroomBindingAsset = *Iterator;
+			if (!GroomBindingAsset->IsAsyncTaskComplete())
+			{
+				if (GroomBindingAsset->AsyncTask->Cancel())
+				{
+					GroomBindingAsset->AsyncTask.Reset();
+				}
+			}
+
+			if (!GroomBindingAsset->AsyncTask)
+			{
+				Iterator.RemoveCurrent();
+				DetachDependencies(GroomBindingAsset);
+			}
+		}
+
+		FinishCompilation(RegisteredGroomBindingAssets.Array());
+	}
+}
+
+bool FGroomBindingCompilingManager::IsAsyncCompilationEnabled() const
+{
+	if (bHasShutdown)
+	{
+		return false;
+	}
+
+	return CVarAsyncGroomBindingAssetStandard.AsyncCompilation.GetValueOnAnyThread() != 0;
+}
+
+TRACE_DECLARE_INT_COUNTER(QueuedGroomBindingAssetCompilation, TEXT("AsyncCompilation/QueuedGroomBinding"));
+void FGroomBindingCompilingManager::UpdateCompilationNotification()
+{
+	TRACE_COUNTER_SET(QueuedGroomBindingAssetCompilation, GetNumRemainingJobs());
+	Notification.Update(GetNumRemainingJobs());
+}
+
+void FGroomBindingCompilingManager::PostCompilation(TArrayView<UGroomBindingAsset* const> InGroomBindingAssets)
+{
+	if (InGroomBindingAssets.Num())
+	{
+		TRACE_CPUPROFILER_EVENT_SCOPE(OnAssetPostCompileEvent);
+
+		TSet<UGroomBindingAsset*> Set;
+		TArray<FAssetCompileData> AssetsData;
+		AssetsData.Reserve(InGroomBindingAssets.Num());
+
+		for (UGroomBindingAsset* GroomBindingAsset : InGroomBindingAssets)
+		{
+			Set.Emplace(GroomBindingAsset);
+			AssetsData.Emplace(GroomBindingAsset);
+		}
+
+		FAssetCompilingManager::Get().OnAssetPostCompileEvent().Broadcast(AssetsData);
+	}
+}
+
+void FGroomBindingCompilingManager::PostCompilation(UGroomBindingAsset* GroomBindingAsset)
+{
+	using namespace GroomBindingAssetCompilingManagerImpl;
+	
+	// If AsyncTask is null here, the task got canceled so we don't need to do anything
+	if (GroomBindingAsset->AsyncTask)
+	{
+		check(IsInGameThread());
+		TRACE_CPUPROFILER_EVENT_SCOPE(FGroomBindingCompilingManager::PostCompilation);
+
+		UE_LOG(LogHairStrands, Verbose, TEXT("Refreshing groom binding asset %s because it is ready"), *GroomBindingAsset->GetName());
+
+		FObjectCacheContextScope ObjectCacheScope;
+
+		// The scope is important here to destroy the FGroomBindingAsyncBuildScope before broadcasting events
+		{
+			// Acquire the async task locally to protect against re-entrance
+			TUniquePtr<FGroomBindingAsyncBuildTask> LocalAsyncTask = MoveTemp(GroomBindingAsset->AsyncTask);
+			LocalAsyncTask->EnsureCompletion();
+
+			FGroomBindingAsyncBuildScope AsyncBuildScope(GroomBindingAsset);
+
+			if (LocalAsyncTask->GetTask().BuildContext.IsSet())
+			{
+				GroomBindingAsset->FinishCacheDerivedDatas(*LocalAsyncTask->GetTask().BuildContext);
+
+				LocalAsyncTask->GetTask().BuildContext.Reset();
+			}
+		}
+
+		// Calling this delegate during app exit might be quite dangerous and lead to crash
+		// if the content browser wants to refresh a thumbnail it might try to load a package
+		// which will then fail due to various reasons related to the editor shutting down.
+		// Triggering this callback while garbage collecting can also result in listeners trying to look up objects
+		if (!GExitPurge && !IsGarbageCollecting())
+		{
+			// Generate an empty property changed event, to force the asset registry tag
+			// to be refreshed now that RenderData is available.
+			FPropertyChangedEvent EmptyPropertyChangedEvent(nullptr);
+			FCoreUObjectDelegates::OnObjectPropertyChanged.Broadcast(GroomBindingAsset, EmptyPropertyChangedEvent);
+		}
+	}
+}
+
+bool FGroomBindingCompilingManager::IsAsyncCompilationAllowed(UGroomBindingAsset* GroomBindingAsset) const
+{
+	return IsAsyncCompilationEnabled();
+}
+
+FGroomBindingCompilingManager& FGroomBindingCompilingManager::Get()
+{
+	static FGroomBindingCompilingManager Singleton;
+	return Singleton;
+}
+
+int32 FGroomBindingCompilingManager::GetNumRemainingJobs() const
+{
+	return RegisteredGroomBindingAssets.Num();
+}
+
+void FGroomBindingCompilingManager::AddGroomBindings(TArrayView<UGroomBindingAsset* const> InGroomBindingAssets)
+{
+	TRACE_CPUPROFILER_EVENT_SCOPE(FGroomBindingCompilingManager::AddGroomBindings)
+	check(IsInGameThread());
+
+	for (UGroomBindingAsset* GroomBindingAsset : InGroomBindingAssets)
+	{
+		check(GroomBindingAsset->AsyncTask != nullptr);
+		RegisteredGroomBindingAssets.Emplace(GroomBindingAsset);
+
+		AttachDependencies(GroomBindingAsset);
+	}
+
+	UpdateCompilationNotification();
+}
+
+void FGroomBindingCompilingManager::FinishCompilation(TArrayView<UGroomBindingAsset* const> InGroomBindingAssets)
+{
+	TRACE_CPUPROFILER_EVENT_SCOPE(FGroomBindingCompilingManager::FinishCompilation);
+
+	check(IsInGameThread());
+
+	TArray<UGroomBindingAsset*> PendingGroomBindingAssets;
+	PendingGroomBindingAssets.Reserve(InGroomBindingAssets.Num());
+
+	for (UGroomBindingAsset* GroomBindingAsset : InGroomBindingAssets)
+	{
+		if (RegisteredGroomBindingAssets.Contains(GroomBindingAsset))
+		{
+			PendingGroomBindingAssets.Emplace(GroomBindingAsset);
+		}
+	}
+
+	if (PendingGroomBindingAssets.Num())
+	{
+		class FCompilableGroomBindingAsset : public AsyncCompilationHelpers::TCompilableAsyncTask<FGroomBindingAsyncBuildTask>
+		{
+		public:
+			FCompilableGroomBindingAsset(UGroomBindingAsset* InGroomBindingAsset)
+				: GroomBindingAsset(InGroomBindingAsset)
+			{
+			}
+
+			FGroomBindingAsyncBuildTask* GetAsyncTask() override
+			{
+				return GroomBindingAsset->AsyncTask.Get();
+			}
+
+			TStrongObjectPtr<UGroomBindingAsset> GroomBindingAsset;
+			FName GetName() override { return GroomBindingAsset->GetFName(); }
+		};
+
+		TArray<FCompilableGroomBindingAsset> CompilableGroomBindingAsset(PendingGroomBindingAssets);
+
+		FObjectCacheContextScope ObjectCacheScope;
+		AsyncCompilationHelpers::FinishCompilation(
+			[&CompilableGroomBindingAsset](int32 Index) -> AsyncCompilationHelpers::ICompilable& { return CompilableGroomBindingAsset[Index]; },
+			CompilableGroomBindingAsset.Num(),
+			LOCTEXT("GroomBindings", "Groom Bindings"),
+			LogHairStrands,
+			[this](AsyncCompilationHelpers::ICompilable* Object)
+			{
+				UGroomBindingAsset* GroomBindingAsset = static_cast<FCompilableGroomBindingAsset*>(Object)->GroomBindingAsset.Get();
+				PostCompilation(GroomBindingAsset);
+				RegisteredGroomBindingAssets.Remove(GroomBindingAsset);
+				DetachDependencies(GroomBindingAsset);
+			}
+		);
+
+		PostCompilation(PendingGroomBindingAssets);
+
+		UpdateCompilationNotification();
+	}
+}
+
+void FGroomBindingCompilingManager::Reschedule()
+{
+	// Implement if we ever want to prioritize nearest visible grooms first
+}
+
+void FGroomBindingCompilingManager::FinishCompilationsForGame()
+{
+	// Implement if grooms ever become necessary for gameplay to work correctly
+}
+
+void FGroomBindingCompilingManager::FinishAllCompilation()
+{
+	check(IsInGameThread());
+	TRACE_CPUPROFILER_EVENT_SCOPE(FGroomBindingCompilingManager::FinishAllCompilation);
+
+	if (GetNumRemainingJobs())
+	{
+		FinishCompilation(RegisteredGroomBindingAssets.Array());
+	}
+}
+
+void FGroomBindingCompilingManager::FinishCompilationForObjects(TArrayView<UObject* const> InObjects)
+{
+	TRACE_CPUPROFILER_EVENT_SCOPE(FGroomBindingCompilingManager::FinishCompilationForObjects);
+
+	TSet<UGroomBindingAsset*> AssetToFinish;
+	for (UObject* Object : InObjects)
+	{
+		if (UGroomBindingAsset* GroomBindingAsset = Cast<UGroomBindingAsset>(Object))
+		{
+			AssetToFinish.Add(GroomBindingAsset);
+		}
+		else if (UGroomComponent* GroomComponent = Cast<UGroomComponent>(Object))
+		{
+			if (GroomComponent->BindingAsset)
+			{
+				AssetToFinish.Add(GroomComponent->BindingAsset);
+			}
+		}
+		else if (USkeletalMesh* SkeletalMesh = Cast<USkeletalMesh>(Object))
+		{
+			for (auto It = RegisteredSkeletalMeshes.CreateKeyIterator(SkeletalMesh); It; ++It)
+			{
+				AssetToFinish.Add(It.Value());
+			}
+		}
+		else if (UGroomAsset* GroomAsset = Cast<UGroomAsset>(Object))
+		{
+			for (auto It = RegisteredGroomAssets.CreateKeyIterator(GroomAsset); It; ++It)
+			{
+				AssetToFinish.Add(It.Value());
+			}
+		}
+	}
+
+	if (AssetToFinish.Num())
+	{
+		FinishCompilation(AssetToFinish.Array());
+	}
+}
+
+void FGroomBindingCompilingManager::ProcessGroomBindingAssets(bool bLimitExecutionTime, int32 MinBatchSize)
+{
+	using namespace GroomBindingAssetCompilingManagerImpl;
+	TRACE_CPUPROFILER_EVENT_SCOPE(FGroomBindingCompilingManager::ProcessGroomBindingAssets);
+	const int32 NumRemainingJobs = GetNumRemainingJobs();
+	// Spread out the load over multiple frames but if too many assets, convergence is more important than frame time
+	const int32 MaxMeshUpdatesPerFrame = bLimitExecutionTime ? FMath::Max(64, NumRemainingJobs / 10) : INT32_MAX;
+
+	FObjectCacheContextScope ObjectCacheScope;
+	if (NumRemainingJobs && NumRemainingJobs >= MinBatchSize)
+	{
+		TRACE_CPUPROFILER_EVENT_SCOPE(ProcessFinishedGroomBindingAssets);
+
+		const double TickStartTime = FPlatformTime::Seconds();
+
+		TArray<UGroomBindingAsset*> ProcessedGroomBindingAssets;
+		int32 ProcessedCount = 0;
+		for (auto It = RegisteredGroomBindingAssets.CreateIterator(); It; ++It)
+		{
+			UGroomBindingAsset* GroomBindingAsset = *It;
+			if (GroomBindingAsset->IsAsyncTaskComplete())
+			{
+				PostCompilation(GroomBindingAsset);
+				ProcessedGroomBindingAssets.Add(GroomBindingAsset);
+				It.RemoveCurrent();
+
+				if (++ProcessedCount > MaxMeshUpdatesPerFrame)
+				{
+					break;
+				}
+			}
+		}
+
+		PostCompilation(ProcessedGroomBindingAssets);
+	}
+}
+
+void FGroomBindingCompilingManager::ProcessAsyncTasks(bool bLimitExecutionTime)
+{
+	FObjectCacheContextScope ObjectCacheScope;
+	FinishCompilationsForGame();
+
+	Reschedule();
+
+	ProcessGroomBindingAssets(bLimitExecutionTime);
+
+	UpdateCompilationNotification();
+}
+
+#endif // #if WITH_EDITOR
+
+#undef LOCTEXT_NAMESPACE
\ No newline at end of file
diff --git a/Engine/Plugins/Runtime/HairStrands/Source/HairStrandsCore/Private/GroomComponent.cpp b/Engine/Plugins/Runtime/HairStrands/Source/HairStrandsCore/Private/GroomComponent.cpp
--- a/Engine/Plugins/Runtime/HairStrands/Source/HairStrandsCore/Private/GroomComponent.cpp
+++ b/Engine/Plugins/Runtime/HairStrands/Source/HairStrandsCore/Private/GroomComponent.cpp
@@ -1554,6 +1554,26 @@ void UGroomComponent::SwitchSimulationLOD(const int32 PreviousLOD, const int32 C
 	}
 }
 
+#if WITH_EDITOR
+
+bool UGroomComponent::IsCompiling() const
+{
+	return BindingAsset && BindingAsset->IsCompiling();
+}
+
+void UGroomComponent::PostCompilation()
+{
+	if (GroomAssetBeingLoaded && GroomAssetBeingLoaded->IsValid())
+	{
+		// Re-set the assets on the component now that they are loaded
+		SetGroomAsset(GroomAssetBeingLoaded, BindingAssetBeingLoaded);
+	}
+
+	InitIfDependenciesReady();
+}
+
+#endif // WITH_EDITOR
+
 void UGroomComponent::SetGroomAsset(UGroomAsset* Asset)
 {
 	SetGroomAsset(Asset, BindingAsset);
@@ -1568,7 +1588,7 @@ void UGroomComponent::SetGroomAsset(UGroomAsset* Asset, UGroomBindingAsset* InBi
 		LODForcedIndex = FMath::Clamp(LODForcedIndex, -1, GroomAsset->GetLODCount() - 1);
 
 #if WITH_EDITORONLY_DATA
-		if (InBinding && !InBinding->IsValid())
+		if (InBinding && !InBinding->IsCompiling() && !InBinding->IsValid())
 		{
 			// The binding could be invalid if the groom asset was previously invalid.
 			// This will re-fetch the binding data from the DDC to make it valid
@@ -1600,6 +1620,18 @@ void UGroomComponent::SetGroomAsset(UGroomAsset* Asset, UGroomBindingAsset* InBi
 		BindingAsset = InBinding;
 	}
 
+	InitIfDependenciesReady(bUpdateSimulation);
+}
+
+void UGroomComponent::InitIfDependenciesReady(const bool bUpdateSimulation)
+{
+#if WITH_EDITOR
+	if (BindingAsset && BindingAsset->IsCompiling())
+	{
+		return;
+	}
+#endif
+
 	if (!UGroomBindingAsset::IsBindingAssetValid(BindingAsset, false, bValidationEnable) || !UGroomBindingAsset::IsCompatible(GroomAsset, BindingAsset, bValidationEnable))
 	{
 		BindingAsset = nullptr;
@@ -1611,7 +1643,23 @@ void UGroomComponent::SetGroomAsset(UGroomAsset* Asset, UGroomBindingAsset* InBi
 		return;
 	}
 	InitResources();
-	if(bUpdateSimulation) UpdateHairSimulation();
+
+	if (bUpdateSimulation)
+	{
+		UpdateHairSimulation();
+	}
+
+	// Registration can be skipped when BindingAsset is compiling, we need to take care of it here
+	// if we're already registered.
+	if (IsRegistered())
+	{
+		if (GUseGroomCacheStreaming)
+		{
+			IGroomCacheStreamingManager::Get().RegisterComponent(this);
+		}
+
+		MeshDeformerInstance = (MeshDeformer != nullptr) ? MeshDeformer->CreateInstance(this, MeshDeformerInstanceSettings) : nullptr;
+	}
 }
 
 void UGroomComponent::SetStableRasterization(bool bEnable)
@@ -1907,7 +1955,17 @@ void UGroomComponent::UpdateHairGroupsDescAndInvalidateRenderState(bool bInvalid
 FPrimitiveSceneProxy* UGroomComponent::CreateSceneProxy()
 {
 	if (!GroomAsset || GroomAsset->GetNumHairGroups() == 0 || HairGroupInstances.Num() == 0)
+	{
 		return nullptr;
+	}
+
+#if WITH_EDITOR
+	// This will be recreated when the binding asset compilation finishes
+	if (BindingAsset && BindingAsset->IsCompiling())
+	{
+		return nullptr;
+	}
+#endif
 
 	if (CheckPSOPrecachingAndBoostPriority() && GetPSOPrecacheProxyCreationStrategy() == EPSOPrecacheProxyCreationStrategy::DelayUntilPSOPrecached)
 	{
@@ -3217,6 +3275,14 @@ void UGroomComponent::InvalidateAndRecreate()
 void UGroomComponent::OnRegister()
 {
 	Super::OnRegister();
+
+#if WITH_EDITOR
+	if (BindingAsset && BindingAsset->IsCompiling())
+	{
+		return;
+	}
+#endif
+
 	UpdateHairGroupsDesc();
 
 	if (GUseGroomCacheStreaming)
@@ -3297,6 +3363,15 @@ void UGroomComponent::OnComponentDestroyed(bool bDestroyingHierarchy)
 void UGroomComponent::OnAttachmentChanged()
 {
 	Super::OnAttachmentChanged();
+
+#if WITH_EDITOR
+	// InitResources will be called when the binding finishes compiling
+	if (BindingAsset && BindingAsset->IsCompiling())
+	{
+		return;
+	}
+#endif
+
 	if (GroomAsset && !IsBeingDestroyed() && HasBeenCreated() && IsValidChecked(this))
 	{
 		UMeshComponent* NewMeshComponent = Cast<UMeshComponent>(GetAttachParent());
diff --git a/Engine/Plugins/Runtime/HairStrands/Source/HairStrandsCore/Private/HairStrandsCore.cpp b/Engine/Plugins/Runtime/HairStrands/Source/HairStrandsCore/Private/HairStrandsCore.cpp
--- a/Engine/Plugins/Runtime/HairStrands/Source/HairStrandsCore/Private/HairStrandsCore.cpp
+++ b/Engine/Plugins/Runtime/HairStrands/Source/HairStrandsCore/Private/HairStrandsCore.cpp
@@ -8,7 +8,9 @@
 #include "GroomManager.h"
 #include "Engine/StaticMesh.h"
 #include "GroomAsset.h"
+#include "GroomBindingCompiler.h"
 #include "GroomCreateBindingOptions.h"
+#include "AssetCompilingManager.h"
 #include "UObject/Package.h"
 
 IMPLEMENT_MODULE(FHairStrandsCore, HairStrandsCore);
@@ -23,6 +25,10 @@ void FHairStrandsCore::StartupModule()
 {
 	RegisterBookmarkFunction(ProcessHairStrandsBookmark);
 
+#if WITH_EDITOR
+	FAssetCompilingManager::Get().RegisterManager(&FGroomBindingCompilingManager::Get());
+#endif
+
 	// Maps virtual shader source directory /Plugin/FX/Niagara to the plugin's actual Shaders directory.
 	FString PluginShaderDir = FPaths::Combine(IPluginManager::Get().FindPlugin(TEXT("HairStrands"))->GetBaseDir(), TEXT("Shaders"));
 	AddShaderSourceDirectoryMapping(TEXT("/Plugin/Runtime/HairStrands"), PluginShaderDir);
@@ -31,6 +37,10 @@ void FHairStrandsCore::StartupModule()
 
 void FHairStrandsCore::ShutdownModule()
 {
+#if WITH_EDITOR
+	FAssetCompilingManager::Get().UnregisterManager(&FGroomBindingCompilingManager::Get());
+#endif
+
 	SetGroomEnabled(false);
 }
 
diff --git a/Engine/Plugins/Runtime/HairStrands/Source/HairStrandsCore/Public/GroomAsset.h b/Engine/Plugins/Runtime/HairStrands/Source/HairStrandsCore/Public/GroomAsset.h
--- a/Engine/Plugins/Runtime/HairStrands/Source/HairStrandsCore/Public/GroomAsset.h
+++ b/Engine/Plugins/Runtime/HairStrands/Source/HairStrandsCore/Public/GroomAsset.h
@@ -642,6 +642,7 @@ public:
 	FOnGroomAsyncLoadFinished& GetOnGroomAsyncLoadFinished() { return OnGroomAsyncLoadFinished; }
 
 	/**  Part of Uobject interface  */
+	virtual bool Modify(bool bAlwaysMarkDirty = true) override;
 	virtual void PostEditChangeProperty(FPropertyChangedEvent& PropertyChangedEvent) override;
 
 #endif // WITH_EDITOR
diff --git a/Engine/Plugins/Runtime/HairStrands/Source/HairStrandsCore/Public/GroomBindingAsset.h b/Engine/Plugins/Runtime/HairStrands/Source/HairStrandsCore/Public/GroomBindingAsset.h
--- a/Engine/Plugins/Runtime/HairStrands/Source/HairStrandsCore/Public/GroomBindingAsset.h
+++ b/Engine/Plugins/Runtime/HairStrands/Source/HairStrandsCore/Public/GroomBindingAsset.h
@@ -10,6 +10,7 @@
 #include "RenderResource.h"
 #include "GroomResources.h"
 #include "HairStrandsInterface.h"
+#include "Async/AsyncWork.h"
 #include "Engine/SkeletalMesh.h"
 #include "GroomBindingAsset.generated.h"
 
@@ -46,11 +47,111 @@ enum class EGroomBindingMeshType : uint8
 	GeometryCache
 };
 
+/*-----------------------------------------------------------------------------
+	Async GroomBinding Compilation
+-----------------------------------------------------------------------------*/
+
+UENUM()
+enum class EGroomBindingAsyncProperties : uint64
+{
+	None = 0,
+	GroomBindingType = 1 << 0,
+	Groom = 1 << 1,
+	SourceSkeletalMesh = 1 << 2,
+	TargetSkeletalMesh = 1 << 3,
+	SourceGeometryCache = 1 << 4,
+	TargetGeometryCache = 1 << 5,
+	NumInterpolationPoints = 1 << 6,
+	MatchingSection = 1 << 7,
+	GroupInfos = 1 << 8,
+	HairGroupResources = 1 << 9,
+	HairGroupPlatformData = 1 << 10,
+	All = MAX_uint64
+};
+
+ENUM_CLASS_FLAGS(EGroomBindingAsyncProperties);
+
+enum class EGroomBindingAsyncPropertyLockType
+{
+	None = 0,
+	ReadOnly = 1,
+	WriteOnly = 2,
+	ReadWrite = 3
+};
+
+ENUM_CLASS_FLAGS(EGroomBindingAsyncPropertyLockType);
+
+// Any thread implicated in the build must have a valid scope to be granted access to protected properties without causing any stalls.
+class FGroomBindingAsyncBuildScope
+{
+public:
+	HAIRSTRANDSCORE_API FGroomBindingAsyncBuildScope(const UGroomBindingAsset* Asset);
+	HAIRSTRANDSCORE_API ~FGroomBindingAsyncBuildScope();
+	HAIRSTRANDSCORE_API static bool ShouldWaitOnLockedProperties(const UGroomBindingAsset* Asset);
+
+private:
+	const UGroomBindingAsset* PreviousScope = nullptr;
+	// Only the thread(s) compiling this asset will have full access to protected properties without causing any stalls.
+	static thread_local const UGroomBindingAsset* Asset;
+};
+
+struct FGroomBindingBuildContext
+{
+	FGroomBindingBuildContext() = default;
+	// Non-copyable
+	FGroomBindingBuildContext(const FGroomBindingBuildContext&) = delete;
+	FGroomBindingBuildContext& operator=(const FGroomBindingBuildContext&) = delete;
+	// Movable
+	FGroomBindingBuildContext(FGroomBindingBuildContext&&) = default;
+	FGroomBindingBuildContext& operator=(FGroomBindingBuildContext&&) = default;
+
+	bool bReloadResource = false;
+};
+
+/**
+ * Worker used to perform async compilation.
+ */
+class FGroomBindingAsyncBuildWorker : public FNonAbandonableTask
+{
+public:
+	UGroomBindingAsset* GroomBinding;
+	TOptional<FGroomBindingBuildContext> BuildContext;
+
+	/** Initialization constructor. */
+	FGroomBindingAsyncBuildWorker(
+		UGroomBindingAsset* InGroomBinding,
+		FGroomBindingBuildContext&& InBuildContext)
+		: GroomBinding(InGroomBinding)
+		, BuildContext(MoveTemp(InBuildContext))
+	{
+	}
+
+	FORCEINLINE TStatId GetStatId() const
+	{
+		RETURN_QUICK_DECLARE_CYCLE_STAT(FGroomBindingAsyncBuildWorker, STATGROUP_ThreadPoolAsyncTasks);
+	}
+
+	void DoWork();
+};
+
+struct FGroomBindingAsyncBuildTask : public FAsyncTask<FGroomBindingAsyncBuildWorker>
+{
+	FGroomBindingAsyncBuildTask(
+		UGroomBindingAsset* InGroomBinding,
+		FGroomBindingBuildContext&& InBuildContext)
+		: FAsyncTask<FGroomBindingAsyncBuildWorker>(InGroomBinding, MoveTemp(InBuildContext))
+		, GroomBinding(InGroomBinding)
+	{
+	}
+
+	const UGroomBindingAsset* GroomBinding;
+};
+
 /**
  * Implements an asset that can be used to store binding information between a groom and a skeletal mesh
  */
 UCLASS(BlueprintType, hidecategories = (Object))
-class HAIRSTRANDSCORE_API UGroomBindingAsset : public UObject
+class HAIRSTRANDSCORE_API UGroomBindingAsset : public UObject, public IInterface_AsyncCompilation
 {
 	GENERATED_BODY()
 
@@ -251,15 +352,83 @@ public:
 	/** Return true if the binding asset is valid, i.e., correctly built and loaded. */
 	bool IsValid() const { return bIsValid;  }
 
-	//private :
 #if WITH_EDITOR
+private:
+	/** Used as a bit-field indicating which properties are read by async compilation. */
+	std::atomic<uint64> AccessedProperties;
+	/** Used as a bit-field indicating which properties are written to by async compilation. */
+	std::atomic<uint64> ModifiedProperties;
+	/** Holds the pointer to an async task if one exists. */
+	TUniquePtr<FGroomBindingAsyncBuildTask> AsyncTask;
+
+	bool IsAsyncTaskComplete() const
+	{
+		return AsyncTask == nullptr || AsyncTask->IsWorkDone();
+	}
+
+	bool TryCancelAsyncTasks()
+	{
+		if (AsyncTask)
+		{
+			if (AsyncTask->IsDone() || AsyncTask->Cancel())
+			{
+				AsyncTask.Reset();
+			}
+		}
+
+		return AsyncTask == nullptr;
+	}
+
+	void ExecuteCacheDerivedDatas(FGroomBindingBuildContext& Context);
+	void FinishCacheDerivedDatas(FGroomBindingBuildContext& Context);
+
+public:
+	/** IInterface_AsyncCompilation begin*/
+	virtual bool IsCompiling() const override
+	{
+		return AsyncTask != nullptr || AccessedProperties.load(std::memory_order_relaxed) != 0;
+	}
+	/** IInterface_AsyncCompilation end*/
+
 	FOnGroomBindingAssetChanged OnGroomBindingAssetChanged;
 
 	void RecreateResources();
 	void ChangeFeatureLevel(ERHIFeatureLevel::Type PendingFeatureLevel);
 	void ChangePlatformLevel(ERHIFeatureLevel::Type PendingFeatureLevel);
 #endif
+private:
+	void WaitUntilAsyncPropertyReleased(EGroomBindingAsyncProperties AsyncProperties, EGroomBindingAsyncPropertyLockType LockType) const;
+	void AcquireAsyncProperty(uint64 AsyncProperties = MAX_uint64, EGroomBindingAsyncPropertyLockType LockType = EGroomBindingAsyncPropertyLockType::ReadWrite)
+	{
+#if WITH_EDITOR
+		if ((LockType & EGroomBindingAsyncPropertyLockType::ReadOnly) == EGroomBindingAsyncPropertyLockType::ReadOnly)
+		{
+			AccessedProperties |= AsyncProperties;
+		}
+
+		if ((LockType & EGroomBindingAsyncPropertyLockType::WriteOnly) == EGroomBindingAsyncPropertyLockType::WriteOnly)
+		{
+			ModifiedProperties |= AsyncProperties;
+		}
+#endif
+	}
 
+	void ReleaseAsyncProperty(uint64 AsyncProperties = MAX_uint64, EGroomBindingAsyncPropertyLockType LockType = EGroomBindingAsyncPropertyLockType::ReadWrite)
+	{
+#if WITH_EDITOR
+		if ((LockType & EGroomBindingAsyncPropertyLockType::ReadOnly) == EGroomBindingAsyncPropertyLockType::ReadOnly)
+		{
+			AccessedProperties &= ~AsyncProperties;
+		}
+
+		if ((LockType & EGroomBindingAsyncPropertyLockType::WriteOnly) == EGroomBindingAsyncPropertyLockType::WriteOnly)
+		{
+			ModifiedProperties &= ~AsyncProperties;
+		}
+#endif
+	}
+
+public:
 #if WITH_EDITORONLY_DATA
 	/** Build/rebuild a binding asset */
 	void Build();
@@ -292,6 +461,9 @@ private:
 #endif
 	bool bIsValid = false;
 	uint32 AssetNameHash = 0;
+
+	friend class FGroomBindingCompilingManager;
+	friend class FGroomBindingAsyncBuildWorker;
 };
 
 UCLASS(BlueprintType, hidecategories = (Object))
diff --git a/Engine/Plugins/Runtime/HairStrands/Source/HairStrandsCore/Public/GroomBindingCompiler.h b/Engine/Plugins/Runtime/HairStrands/Source/HairStrandsCore/Public/GroomBindingCompiler.h
new file mode 100644
--- /dev/null
+++ b/Engine/Plugins/Runtime/HairStrands/Source/HairStrandsCore/Public/GroomBindingCompiler.h
@@ -0,0 +1,102 @@
+// Copyright Epic Games, Inc. All Rights Reserved.
+#pragma once
+
+#include "CoreMinimal.h"
+#include "Containers/Map.h"
+#include "Containers/Set.h"
+#include "Tickable.h"
+#include "Stats/Stats.h"
+#include "IAssetCompilingManager.h"
+#include "AsyncCompilationHelpers.h"
+
+#if WITH_EDITOR
+
+class UGroomBindingAsset;
+class UPrimitiveComponent;
+class FQueuedThreadPool;
+struct FAssetCompileContext;
+enum class EQueuedWorkPriority : uint8;
+
+class FGroomBindingCompilingManager : public IAssetCompilingManager
+{
+public:
+	HAIRSTRANDSCORE_API static FGroomBindingCompilingManager& Get();
+
+	/**
+	 * Returns true if the feature is currently activated.
+	 */
+	HAIRSTRANDSCORE_API bool IsAsyncCompilationEnabled() const;
+
+	/** 
+	 * Returns the number of outstanding compilations.
+	 */
+	HAIRSTRANDSCORE_API int32 GetNumRemainingJobs() const;
+
+	/** 
+	 * Adds groom binding assets compiled asynchronously so they are monitored. 
+	 */
+	HAIRSTRANDSCORE_API void AddGroomBindings(TArrayView<UGroomBindingAsset* const> InGroomBindingAssets);
+
+	/** 
+	 * Blocks until completion of the requested groom binding assets.
+	 */
+	HAIRSTRANDSCORE_API void FinishCompilation(TArrayView<UGroomBindingAsset* const> InGroomBindingAssets);
+
+	/** 
+	 * Blocks until completion of all async groom binding asset compilation.
+	 */
+	HAIRSTRANDSCORE_API void FinishAllCompilation() override;
+
+	/**
+	 * Returns if asynchronous compilation is allowed for this groom binding asset.
+	 */
+	HAIRSTRANDSCORE_API bool IsAsyncCompilationAllowed(UGroomBindingAsset* InGroomBindingAsset) const;
+
+	/**
+	 * Returns the priority at which the given groom binding asset should be scheduled.
+	 */
+	HAIRSTRANDSCORE_API EQueuedWorkPriority GetBasePriority(UGroomBindingAsset* InGroomBindingAsset) const;
+
+	/**
+	 * Returns the threadpool where groom binding asset compilation should be scheduled.
+	 */
+	HAIRSTRANDSCORE_API FQueuedThreadPool* GetThreadPool() const;
+
+	/**
+	 * Cancel any pending work and blocks until it is safe to shut down.
+	 */
+	HAIRSTRANDSCORE_API void Shutdown() override;
+
+private:
+	FGroomBindingCompilingManager();
+
+	FName GetAssetTypeName() const override;
+	FTextFormat GetAssetNameFormat() const override;
+	TArrayView<FName> GetDependentTypeNames() const override;
+	int32 GetNumRemainingAssets() const override;
+	void ProcessAsyncTasks(bool bLimitExecutionTime = false) override;
+
+	void FinishCompilationForObjects(TArrayView<UObject* const> InObjects) override;
+
+	friend class FAssetCompilingManager;
+	
+	bool bHasShutdown = false;
+	TSet<UGroomBindingAsset*> RegisteredGroomBindingAssets;
+	TMultiMap<class USkeletalMesh*, UGroomBindingAsset*> RegisteredSkeletalMeshes;
+	TMultiMap<class UGroomAsset*, UGroomBindingAsset*> RegisteredGroomAssets;
+	FAsyncCompilationNotification Notification;
+	void FinishCompilationsForGame();
+	void Reschedule();
+	void ProcessGroomBindingAssets(bool bLimitExecutionTime, int32 MinBatchSize = 1);
+	void UpdateCompilationNotification();
+	
+	void AttachDependencies(UGroomBindingAsset* GroomBindingAsset);
+	void DetachDependencies(UGroomBindingAsset* GroomBindingAsset);
+
+	void PostCompilation(UGroomBindingAsset* GroomBindingAsset);
+	void PostCompilation(TArrayView<UGroomBindingAsset* const> InGroomBindingAssets);
+	void OnPostReachabilityAnalysis();
+	FDelegateHandle PostReachabilityAnalysisHandle;
+};
+
+#endif // #if WITH_EDITOR
\ No newline at end of file
diff --git a/Engine/Plugins/Runtime/HairStrands/Source/HairStrandsCore/Public/GroomComponent.h b/Engine/Plugins/Runtime/HairStrands/Source/HairStrandsCore/Public/GroomComponent.h
--- a/Engine/Plugins/Runtime/HairStrands/Source/HairStrandsCore/Public/GroomComponent.h
+++ b/Engine/Plugins/Runtime/HairStrands/Source/HairStrandsCore/Public/GroomComponent.h
@@ -166,7 +166,6 @@ public:
 	FHairStrandsRestRootResource* GetGuideStrandsRestRootResource(uint32 GroupIndex);
 	FHairStrandsDeformedRootResource* GetGuideStrandsDeformedRootResource(uint32 GroupIndex);
 
-
 #if WITH_EDITOR
 	virtual void CheckForErrors() override;
 	virtual void PreEditChange(FProperty* PropertyAboutToChange) override;
@@ -175,6 +174,11 @@ public:
 	void ValidateMaterials(bool bMapCheck) const;
 	void Invalidate();
 	void InvalidateAndRecreate();
+	void PostCompilation();
+	
+	//~ Begin IInterface_AsyncCompilation Interface.
+	virtual bool IsCompiling() const override;
+	//~ End IInterface_AsyncCompilation Interface.
 #endif
 
 #if WITH_EDITOR
@@ -307,6 +311,7 @@ public:
 	//~ End INiagaraPhysicsAssetDICollectorInterface Interface
 
 private:
+	void InitIfDependenciesReady(const bool bUpdateSimulation = true);
 	void UpdateGroomCache(float Time);
 
 	UPROPERTY(EditAnywhere, Category = GroomCache)
@@ -358,6 +363,7 @@ private:
 	void InitResources(bool bIsBindingReloading=false);
 	void ReleaseResources();
 
+	friend class UGroomBindingAsset;
 	friend class FGroomComponentRecreateRenderStateContext;
 	friend class FHairStrandsSceneProxy;
 	friend class FHairCardsSceneProxy;
diff --git a/Engine/Source/Editor/UnrealEd/Classes/Settings/EditorExperimentalSettings.h b/Engine/Source/Editor/UnrealEd/Classes/Settings/EditorExperimentalSettings.h
--- a/Engine/Source/Editor/UnrealEd/Classes/Settings/EditorExperimentalSettings.h
+++ b/Engine/Source/Editor/UnrealEd/Classes/Settings/EditorExperimentalSettings.h
@@ -41,6 +41,10 @@ public:
 	UPROPERTY(EditAnywhere, config, Category = Performance, meta = (DisplayName = "Enable async sound compilation and loading"))
 	bool bEnableAsyncSoundWaveCompilation;
 
+	/** Enable async groom binding compilation to improve import and map load time performance when compilation is required */
+	UPROPERTY(EditAnywhere, config, Category = Performance, meta = (DisplayName = "Enable async groom binding compilation and loading"))
+	bool bEnableAsyncGroomBindingCompilation;
+
 	/** Allows the editor to run on HDR monitors on Windows 10 */
 	UPROPERTY(EditAnywhere, config, Category = HDR, meta = (ConfigRestartRequired = true, DisplayName = "Enable Editor Support for HDR Monitors"))
 	bool bHDREditor;
diff --git a/Engine/Source/Editor/UnrealEd/Private/Settings/SettingsClasses.cpp b/Engine/Source/Editor/UnrealEd/Private/Settings/SettingsClasses.cpp
--- a/Engine/Source/Editor/UnrealEd/Private/Settings/SettingsClasses.cpp
+++ b/Engine/Source/Editor/UnrealEd/Private/Settings/SettingsClasses.cpp
@@ -188,6 +188,7 @@ UEditorExperimentalSettings::UEditorExperimentalSettings( const FObjectInitializ
 	, bEnableAsyncSkeletalMeshCompilation(true)	// This was false and set to True in /Engine/Config/BaseEditorPerProjectUserSettings.ini. The setting is removed from .ini so change it to default True.
 	, bEnableAsyncSkinnedAssetCompilation(false)
 	, bEnableAsyncSoundWaveCompilation(false)
+	, bEnableAsyncGroomBindingCompilation(false)
 	, bHDREditor(false)
 	, HDREditorNITLevel(160.0f)
 	, bUseOpenCLForConvexHullDecomp(false)
diff --git a/Engine/Source/Runtime/Engine/Classes/Engine/SkeletalMesh.h b/Engine/Source/Runtime/Engine/Classes/Engine/SkeletalMesh.h
--- a/Engine/Source/Runtime/Engine/Classes/Engine/SkeletalMesh.h
+++ b/Engine/Source/Runtime/Engine/Classes/Engine/SkeletalMesh.h
@@ -2451,6 +2451,8 @@ private:
 public:
 	bool IsTransacting() const;
 	
+	ENGINE_API virtual void PreEditChange(FProperty* PropertyAboutToChange) override;
+
 	ENGINE_API virtual void PreEditUndo() override;
 	ENGINE_API virtual void PostEditUndo() override;
 	ENGINE_API virtual void GetAssetRegistryTagMetadata(TMap<FName, FAssetRegistryTagMetadata>& OutMetadata) const override;
diff --git a/Engine/Source/Runtime/Engine/Private/SkeletalMesh.cpp b/Engine/Source/Runtime/Engine/Private/SkeletalMesh.cpp
--- a/Engine/Source/Runtime/Engine/Private/SkeletalMesh.cpp
+++ b/Engine/Source/Runtime/Engine/Private/SkeletalMesh.cpp
@@ -32,6 +32,7 @@
 #include "PrimitiveViewRelevance.h"
 #include "UObject/Package.h"
 #include "RenderUtils.h"
+#include "AssetCompilingManager.h"
 #include "UObject/UE5MainStreamObjectVersion.h"
 #include "SceneInterface.h"
 #include "EngineUtils.h"
@@ -465,8 +466,6 @@ USkeletalMesh::USkeletalMesh(FVTableHelper& Helper)
 
 PRAGMA_DISABLE_DEPRECATION_WARNINGS
 USkeletalMesh::~USkeletalMesh() = default;
-
-
 PRAGMA_ENABLE_DEPRECATION_WARNINGS
 
 FSkeletalMeshRenderData* USkeletalMesh::GetSkeletalMeshRenderData() const
@@ -1397,9 +1396,27 @@ bool USkeletalMesh::IsTransacting() const
 	return bTransacting;
 }
 
+void USkeletalMesh::PreEditChange(FProperty* PropertyAboutToChange)
+{
+	TRACE_CPUPROFILER_EVENT_SCOPE(USkeletalMesh::PreEditChange);
+
+	// Tell the compiler to finish compiling us if we have a pending
+	// compilation ongoing plus any dependency (i.e. UGroomBindings).
+	FAssetCompilingManager::Get().FinishCompilationForObjects({ this });
+
+	Super::PreEditChange(PropertyAboutToChange);
+}
+
 void USkeletalMesh::PreEditUndo()
 {
+	TRACE_CPUPROFILER_EVENT_SCOPE(USkeletalMesh::PreEditUndo);
+
+	// Tell the compiler to finish compiling us if we have a pending
+	// compilation ongoing plus any dependency (i.e. UGroomBindings).
+	FAssetCompilingManager::Get().FinishCompilationForObjects({ this });
+
 	bTransacting = true;
+
 	Super::PreEditUndo();
 }
 
@@ -2000,10 +2017,9 @@ void USkeletalMesh::Build()
 {
 	TRACE_CPUPROFILER_EVENT_SCOPE(USkeletalMesh::Build);
 
-	if (IsCompiling())
-	{
-		FSkinnedAssetCompilingManager::Get().FinishCompilation({this});
-	}
+	// Tell the compiler to finish compiling us if we have a pending
+	// compilation ongoing plus any dependency (i.e. UGroomBindings).
+	FAssetCompilingManager::Get().FinishCompilationForObjects({this});
 
 	FSkinnedAssetAsyncBuildScope AsyncBuildScope(this);
 
@@ -2113,10 +2129,9 @@ FEvent* USkeletalMesh::LockPropertiesUntil()
 	FEvent* Event = FPlatformProcess::GetSynchEventFromPool();
 	check(Event);
 
-	if (IsCompiling())
-	{
-		FSkinnedAssetCompilingManager::Get().FinishCompilation({ this });
-	}
+	// Tell the compiler to finish compiling us if we have a pending
+	// compilation ongoing plus any dependency (i.e. UGroomBindings).
+	FAssetCompilingManager::Get().FinishCompilationForObjects({ this });
 
 	//Use the async task compile to lock the properties
 	FSkinnedAsyncTaskContext Context(Event);
diff --git a/Engine/Source/Runtime/Engine/Private/SkinnedAssetCompiler.cpp b/Engine/Source/Runtime/Engine/Private/SkinnedAssetCompiler.cpp
--- a/Engine/Source/Runtime/Engine/Private/SkinnedAssetCompiler.cpp
+++ b/Engine/Source/Runtime/Engine/Private/SkinnedAssetCompiler.cpp
@@ -59,11 +59,16 @@ FSkinnedAssetCompilingManager::FSkinnedAssetCompilingManager()
 	PreGarbageCollectHandle = FCoreUObjectDelegates::GetPreGarbageCollectDelegate().AddRaw(this, &FSkinnedAssetCompilingManager::OnPreGarbageCollect);
 }
 
-FName FSkinnedAssetCompilingManager::GetAssetTypeName() const
+FName FSkinnedAssetCompilingManager::GetStaticAssetTypeName()
 {
 	return TEXT("UE-SkinnedAsset");
 }
 
+FName FSkinnedAssetCompilingManager::GetAssetTypeName() const
+{
+	return GetStaticAssetTypeName();
+}
+
 FTextFormat FSkinnedAssetCompilingManager::GetAssetNameFormat() const
 {
 	return LOCTEXT("SkinnedAssetNameFormat", "{0}|plural(one=Skinned Asset,other=Skinned Assets)");
diff --git a/Engine/Source/Runtime/Engine/Public/SkinnedAssetCompiler.h b/Engine/Source/Runtime/Engine/Public/SkinnedAssetCompiler.h
--- a/Engine/Source/Runtime/Engine/Public/SkinnedAssetCompiler.h
+++ b/Engine/Source/Runtime/Engine/Public/SkinnedAssetCompiler.h
@@ -72,6 +72,9 @@ public:
 	 */
 	ENGINE_API void Shutdown() override;
 
+	/** Get the name of the asset type this compiler handles */
+	ENGINE_API static FName GetStaticAssetTypeName();
+
 private:
 	FSkinnedAssetCompilingManager();
 
