Fix missing Preview Scene Settings in Groom editor.

diff --git a/Engine/Plugins/Runtime/HairStrands/Source/HairStrandsEditor/Private/GroomCustomAssetEditorToolkit.cpp b/Engine/Plugins/Runtime/HairStrands/Source/HairStrandsEditor/Private/GroomCustomAssetEditorToolkit.cpp
--- a/Engine/Plugins/Runtime/HairStrands/Source/HairStrandsEditor/Private/GroomCustomAssetEditorToolkit.cpp
+++ b/Engine/Plugins/Runtime/HairStrands/Source/HairStrandsEditor/Private/GroomCustomAssetEditorToolkit.cpp
@@ -21,6 +21,7 @@
 #include "GroomMaterialDetails.h"
 #include "PropertyCustomizationHelpers.h"
 #include "AssetThumbnail.h"
+#include "AdvancedPreviewSceneModule.h"
 
 #include "Misc/AssetRegistryInterface.h"
 #include "AssetRegistry/AssetRegistryModule.h"
@@ -42,6 +43,7 @@ const FName FGroomCustomAssetEditorToolkit::TabId_MaterialProperties(TEXT("Groom
 const FName FGroomCustomAssetEditorToolkit::TabId_PhysicsProperties(TEXT("GroomCustomAssetEditor_PhysicsProperties"));
 const FName FGroomCustomAssetEditorToolkit::TabId_PreviewGroomComponent(TEXT("GroomCustomAssetEditor_PreviewGroomComponent"));
 const FName FGroomCustomAssetEditorToolkit::TabId_BindingProperties(TEXT("GroomCustomAssetEditor_BindingProperties"));
+const FName FGroomCustomAssetEditorToolkit::TabId_PreviewSceneProperties(TEXT("GroomCustomAssetEditor_PreviewSceneProperties"));
 
 void FGroomCustomAssetEditorToolkit::RegisterTabSpawners(const TSharedRef<class FTabManager>& InTabManager)
 {
@@ -100,6 +102,11 @@ void FGroomCustomAssetEditorToolkit::RegisterTabSpawners(const TSharedRef<class
 		.SetDisplayName(LOCTEXT("BindingPropertiesTab", "Binding"))
 		.SetGroup(WorkspaceMenuCategory.ToSharedRef())
 		.SetIcon(FSlateIcon(FAppStyle::GetAppStyleSetName(), "Icons.SkeletalMesh"));
+
+	InTabManager->RegisterTabSpawner(TabId_PreviewSceneProperties, FOnSpawnTab::CreateSP(this, &FGroomCustomAssetEditorToolkit::SpawnTab_PreviewSettings))
+		.SetDisplayName(LOCTEXT("PreviewSettingsTab", "Preview Scene Settings"))
+		.SetGroup(WorkspaceMenuCategory.ToSharedRef())
+		.SetIcon(FSlateIcon(FAppStyle::GetAppStyleSetName(), "LevelEditor.Tabs.Details"));
 }
 
 void FGroomCustomAssetEditorToolkit::UnregisterTabSpawners(const TSharedRef<class FTabManager>& InTabManager)
@@ -118,6 +125,7 @@ void FGroomCustomAssetEditorToolkit::UnregisterTabSpawners(const TSharedRef<clas
 	InTabManager->UnregisterTabSpawner(TabId_PreviewGroomComponent);
 #endif
 	InTabManager->UnregisterTabSpawner(TabId_BindingProperties);
+	InTabManager->UnregisterTabSpawner(TabId_PreviewSceneProperties);
 }
 
 void FGroomCustomAssetEditorToolkit::DocPropChanged(UObject *InObject, FPropertyChangedEvent &Property)
@@ -343,6 +351,7 @@ void FGroomCustomAssetEditorToolkit::OnClose()
 	DetailView_PreviewGroomComponent.Reset();
 #endif
 	DetailView_BindingProperties.Reset();
+	DetailView_PreviewSceneProperties.Reset();
 }
 
 static void ListAllBindingAssets(const UGroomAsset* InGroomAsset, TWeakObjectPtr<UGroomBindingAssetList>& Out)
@@ -484,6 +493,7 @@ void FGroomCustomAssetEditorToolkit::InitCustomAssetEditor(const EToolkitMode::T
 	DetailView_PreviewGroomComponent	= PropertyEditorModule.CreateDetailView(DetailsViewArgs);
 #endif
 	DetailView_BindingProperties		= PropertyEditorModule.CreateDetailView(DetailsViewArgs);
+	DetailView_PreviewSceneProperties	= PropertyEditorModule.CreateDetailView(DetailsViewArgs); 
 
 	// Customization
 	DetailView_CardsProperties->SetGenericLayoutDetailsDelegate(FOnGetDetailCustomizationInstance::CreateStatic(&FGroomRenderingDetails::MakeInstance, (IGroomCustomAssetEditorToolkit*)this, EMaterialPanelType::Cards));
@@ -530,6 +540,7 @@ void FGroomCustomAssetEditorToolkit::InitCustomAssetEditor(const EToolkitMode::T
 					->AddTab(TabId_PreviewGroomComponent,	ETabState::OpenedTab)
 				#endif
 					->AddTab(TabId_BindingProperties,		ETabState::OpenedTab)
+					->AddTab(TabId_PreviewSceneProperties,	ETabState::OpenedTab)
 				)
 			)
 		);
@@ -637,6 +648,11 @@ void FGroomCustomAssetEditorToolkit::InitCustomAssetEditor(const EToolkitMode::T
 		DetailView_BindingProperties->SetObject(Cast<UObject>(GroomBindingAssetList));
 	}
 
+	if (DetailView_PreviewSceneProperties.IsValid())
+	{
+		DetailView_PreviewSceneProperties->SetObject(Cast<UObject>(GroomAsset));
+	}
+
 	ExtendToolbar();
 	RegenerateMenusAndToolbars();
 
@@ -676,6 +692,7 @@ void FGroomCustomAssetEditorToolkit::InitCustomAssetEditor(const EToolkitMode::T
 			if (LocalToolKit->DetailView_MeshesProperties)			{ LocalToolKit->DetailView_MeshesProperties->ForceRefresh(); }
 			if (LocalToolKit->DetailView_MaterialProperties)		{ LocalToolKit->DetailView_MaterialProperties->ForceRefresh(); }
 			if (LocalToolKit->DetailView_BindingProperties)			{ LocalToolKit->DetailView_BindingProperties->ForceRefresh(); }
+			if (LocalToolKit->DetailView_PreviewSceneProperties)	{ LocalToolKit->DetailView_PreviewSceneProperties->ForceRefresh(); }
 		};
 
 		PropertyListenDelegatesResourceChanged.Add(GroomAsset->GetOnGroomAssetResourcesChanged().AddLambda(InvalidateDetailViews));
@@ -943,6 +960,30 @@ TSharedRef<SDockTab> FGroomCustomAssetEditorToolkit::SpawnTab_BindingProperties(
 	return DockTab;
 }
 
+TSharedRef<SDockTab> FGroomCustomAssetEditorToolkit::SpawnTab_PreviewSettings(const FSpawnTabArgs& Args)
+{
+	FAdvancedPreviewSceneModule& AdvancedPreviewSceneModule = FModuleManager::LoadModuleChecked<FAdvancedPreviewSceneModule>(TEXT("AdvancedPreviewScene"));
+
+	TSharedRef<SWidget> PreviewSceneSettingsWidget = SNullWidget::NullWidget;
+	TSharedPtr<class FAdvancedPreviewScene> PreviewScene = ViewportTab->GetAdvancedPreviewScene();
+	if (PreviewScene.IsValid())
+	{
+		PreviewSceneSettingsWidget = AdvancedPreviewSceneModule.CreateAdvancedPreviewSceneSettingsWidget(PreviewScene.ToSharedRef());
+	}
+
+	TSharedRef<SDockTab> SpawnedTab =
+		SNew(SDockTab)
+	.Label(LOCTEXT("PreviewSceneSettingsTab", "Preview Scene Settings"))
+	[
+		SNew(SBox)
+		[
+			PreviewSceneSettingsWidget
+		]
+	];
+
+	return SpawnedTab;
+}
+
 UGroomComponent *FGroomCustomAssetEditorToolkit::GetPreview_GroomComponent() const
 {	
 	return PreviewGroomComponent.Get();
diff --git a/Engine/Plugins/Runtime/HairStrands/Source/HairStrandsEditor/Public/GroomCustomAssetEditorToolkit.h b/Engine/Plugins/Runtime/HairStrands/Source/HairStrandsEditor/Public/GroomCustomAssetEditorToolkit.h
--- a/Engine/Plugins/Runtime/HairStrands/Source/HairStrandsEditor/Public/GroomCustomAssetEditorToolkit.h
+++ b/Engine/Plugins/Runtime/HairStrands/Source/HairStrandsEditor/Public/GroomCustomAssetEditorToolkit.h
@@ -147,6 +147,7 @@ private:
 	TSharedRef<SDockTab> SpawnTab_PhysicsProperties(const FSpawnTabArgs& Args);
 	TSharedRef<SDockTab> SpawnTab_PreviewGroomComponent(const FSpawnTabArgs& Args);
 	TSharedRef<SDockTab> SpawnTab_BindingProperties(const FSpawnTabArgs& Args);
+	TSharedRef<SDockTab> SpawnTab_PreviewSettings(const FSpawnTabArgs& Args);
 private:
 
 	/** Dockable tab for properties */
@@ -162,6 +163,7 @@ private:
 	TSharedPtr<class IDetailsView> DetailView_PhysicsProperties;
 	TSharedPtr<class IDetailsView> DetailView_PreviewGroomComponent;
 	TSharedPtr<class IDetailsView> DetailView_BindingProperties;
+	TSharedPtr<class IDetailsView> DetailView_PreviewSceneProperties;
 
 	static const FName ToolkitFName;
 	static const FName TabId_Viewport;
@@ -175,6 +177,7 @@ private:
 	static const FName TabId_PhysicsProperties;
 	static const FName TabId_PreviewGroomComponent;
 	static const FName TabId_BindingProperties;
+	static const FName TabId_PreviewSceneProperties;
 
 
 	bool bIsTabManagerInitialized = false;
diff --git a/Engine/Plugins/Runtime/HairStrands/Source/HairStrandsEditor/Public/Widgets/SGroomEditorViewport.h b/Engine/Plugins/Runtime/HairStrands/Source/HairStrandsEditor/Public/Widgets/SGroomEditorViewport.h
--- a/Engine/Plugins/Runtime/HairStrands/Source/HairStrandsEditor/Public/Widgets/SGroomEditorViewport.h
+++ b/Engine/Plugins/Runtime/HairStrands/Source/HairStrandsEditor/Public/Widgets/SGroomEditorViewport.h
@@ -60,6 +60,8 @@ public:
 	bool IsLODModelSelected(int32 InLODSelection) const;
 	void OnSetLODModel(int32 InLODSelection);
 
+	TSharedPtr<class FAdvancedPreviewScene> GetAdvancedPreviewScene() { return AdvancedPreviewScene; }
+
 protected:
 
 	/** SEditorViewport interface */
