Fix groom binding no longer being used.

* Fix groom component InitResources() being called with reentrance (while not supported). This can also be caused by the groom binding validation code causing skel. mesh to finish compiling, calling subsequently InitResources(), again while already inside this function. This needs to be fixed in a follow up CL, but this option is not enabled by default.
* Fix incorrect early out when the groom binding initialization is not completed.

diff --git a/Engine/Plugins/Runtime/HairStrands/Source/HairStrandsCore/Private/GroomComponent.cpp b/Engine/Plugins/Runtime/HairStrands/Source/HairStrandsCore/Private/GroomComponent.cpp
--- a/Engine/Plugins/Runtime/HairStrands/Source/HairStrandsCore/Private/GroomComponent.cpp
+++ b/Engine/Plugins/Runtime/HairStrands/Source/HairStrandsCore/Private/GroomComponent.cpp
@@ -88,6 +88,11 @@ bool IsHairManualSkinCacheEnabled();
 
 ////////////////////////////////////////////////////////////////////////////////////////////////////////////////////////
 
+bool IsGroomBindingValidationEnabled()
+{
+	return GHairBindingValidationEnable > 0;
+}
+
 template<typename T>
 void InternalResourceRelease(T*& In)
 {
@@ -1382,7 +1387,7 @@ UGroomComponent::UGroomComponent(const FObjectInitializer& ObjectInitializer)
 	NiagaraComponents.Empty();
 	PhysicsAsset = nullptr;
 	bCanEverAffectNavigation = false;
-	bValidationEnable = GHairBindingValidationEnable > 0;
+	bValidationEnable = IsGroomBindingValidationEnabled();
 	bRunning = true;
 	bLooping = true;
 	bManualTick = false;
@@ -1632,11 +1637,6 @@ void UGroomComponent::InitIfDependenciesReady(const bool bUpdateSimulation)
 	}
 #endif
 
-	if (!UGroomBindingAsset::IsBindingAssetValid(BindingAsset, false, bValidationEnable) || !UGroomBindingAsset::IsCompatible(GroomAsset, BindingAsset, bValidationEnable))
-	{
-		BindingAsset = nullptr;
-	}
-
 	UpdateHairGroupsDesc();
 	if (!GroomAsset || !GroomAsset->IsValid())
 	{
@@ -1875,8 +1875,14 @@ void UGroomComponent::SetBindingAsset(UGroomBindingAsset* InBinding)
 {
 	if (BindingAsset != InBinding)
 	{
-		const bool bIsValid = InBinding != nullptr ? UGroomBindingAsset::IsBindingAssetValid(BindingAsset, false, bValidationEnable) : true;
-		if (bIsValid && UGroomBindingAsset::IsCompatible(GroomAsset, InBinding, bValidationEnable))
+		bool bIsValid = InBinding != nullptr;
+		if (bIsValid && bValidationEnable)
+		{
+			bIsValid = 
+				UGroomBindingAsset::IsBindingAssetValid(InBinding, false, bValidationEnable) &&
+				UGroomBindingAsset::IsCompatible(GroomAsset, InBinding, bValidationEnable);
+		}
+		if (bIsValid)
 		{
 			BindingAsset = InBinding;
 			InitResources();
@@ -2381,20 +2387,11 @@ void UGroomComponent::UpdateSimulatedGroups()
 	}
 }
 
-void UGroomComponent::OnChildDetached(USceneComponent* ChildComponent)
-{}
-
-void UGroomComponent::OnChildAttached(USceneComponent* ChildComponent)
-{
-
-}
-
-static UGeometryCacheComponent* ValidateBindingAsset(
+static UGeometryCacheComponent* ValidateBindingAsset_GeometryCache(
 	UGroomAsset* GroomAsset, 
 	UGroomBindingAsset* BindingAsset, 
 	UGeometryCacheComponent* GeometryCacheComponent, 
 	bool bIsBindingReloading, 
-	bool bValidationEnable, 
 	const USceneComponent* Component)
 {
 	if (!GroomAsset || !BindingAsset || !GeometryCacheComponent)
@@ -2418,20 +2415,19 @@ static UGeometryCacheComponent* ValidateBindingAsset(
 	}
 
 	const bool bIsBindingCompatible =
-		UGroomBindingAsset::IsCompatible(GeometryCacheComponent ? GeometryCacheComponent->GeometryCache : nullptr, BindingAsset, bValidationEnable) &&
-		UGroomBindingAsset::IsCompatible(GroomAsset, BindingAsset, bValidationEnable) &&
-		UGroomBindingAsset::IsBindingAssetValid(BindingAsset, bIsBindingReloading, bValidationEnable);
+		UGroomBindingAsset::IsCompatible(GeometryCacheComponent ? GeometryCacheComponent->GeometryCache : nullptr, BindingAsset, true /*bValidationEnable*/) &&
+		UGroomBindingAsset::IsCompatible(GroomAsset, BindingAsset, true /*bValidationEnable*/) &&
+		UGroomBindingAsset::IsBindingAssetValid(BindingAsset, bIsBindingReloading, true /*bValidationEnable*/);
 
 	return bIsBindingCompatible ? GeometryCacheComponent : nullptr;
 }
 
 // Return a non-null skeletal mesh Component if the binding asset is compatible with the current component
-static USkeletalMeshComponent* ValidateBindingAsset(
+static USkeletalMeshComponent* ValidateBindingAsset_SkeletalMesh(
 	UGroomAsset* GroomAsset,
 	UGroomBindingAsset* BindingAsset,
 	USkeletalMeshComponent* SkeletalMeshComponent,
 	bool bIsBindingReloading,
-	bool bValidationEnable,
 	const USceneComponent* Component)
 {
 	if (!GroomAsset || !BindingAsset || !SkeletalMeshComponent)
@@ -2470,9 +2466,9 @@ static USkeletalMeshComponent* ValidateBindingAsset(
 	}
 
 	const bool bIsBindingCompatible =
-		UGroomBindingAsset::IsCompatible(SkeletalMeshComponent ? SkeletalMeshComponent->GetSkeletalMeshAsset() : nullptr, BindingAsset, bValidationEnable) &&
-		UGroomBindingAsset::IsCompatible(GroomAsset, BindingAsset, bValidationEnable) &&
-		UGroomBindingAsset::IsBindingAssetValid(BindingAsset, bIsBindingReloading, bValidationEnable);
+		UGroomBindingAsset::IsCompatible(SkeletalMeshComponent ? SkeletalMeshComponent->GetSkeletalMeshAsset() : nullptr, BindingAsset, true /*bValidationEnable*/) &&
+		UGroomBindingAsset::IsCompatible(GroomAsset, BindingAsset, true /*bValidationEnable*/) &&
+		UGroomBindingAsset::IsBindingAssetValid(BindingAsset, bIsBindingReloading, true /*bValidationEnable*/);
 
 	if (!bIsBindingCompatible)
 	{
@@ -2507,6 +2503,41 @@ static USkeletalMeshComponent* ValidateBindingAsset(
 		}
 	}
 
+	// Validate that if we are bound to a skel. mesh, and we have binding, and the binding type is set to skinning, that the skin cache is enabled for these
+	if (BindingAsset)
+	{
+		// Extract if skin cache or mesh. deformer is enabled on at least on LOD.
+		// Since there is 1:1 mapping between groom LOD and mesh LOD, only use this has a hint.
+		bool bSupportSkinCache = SkeletalMeshComponent->HasMeshDeformer();
+		if (!bSupportSkinCache)
+		{
+			for (uint32 SkelLODIt = 0, SkelLODCount = SkeletalMeshComponent->GetNumLODs(); SkelLODIt < SkelLODCount; ++SkelLODIt)
+			{
+				bSupportSkinCache = bSupportSkinCache || SkeletalMeshComponent->IsSkinCacheAllowed(SkelLODIt);
+			}
+		}
+
+		for (int32 GroupIt = 0, GroupCount = GroomAsset->GetHairGroupsPlatformData().Num(); GroupIt < GroupCount; ++GroupIt)
+		{
+			for (uint32 LODIt = 0, LODCount = GroomAsset->GetLODCount(); LODIt < LODCount; ++LODIt)
+			{
+				const EGroomBindingType BindingType = GroomAsset->GetBindingType(GroupIt, LODIt);
+				const bool bIsVisible = GroomAsset->IsVisible(GroupIt, LODIt);
+
+				if (BindingType == EGroomBindingType::Skinning && (!bSupportSkinCache && !IsHairManualSkinCacheEnabled()) && bIsVisible)
+				{
+					UE_LOG(LogHairStrands, Warning, TEXT("[Groom] Groom asset (Group:%d/%d) is set to use Skinning at LOD %d/%d while the parent skel. mesh does not support skin. cache at this LOD - Groom:%s - Skel.Mesh:%s"),
+						GroupIt,
+						GroupCount,
+						LODIt,
+						LODCount,
+						*GroomAsset->GetPathName(),
+						*SkeletalMesh->GetPathName());
+				}
+			}
+		}
+	}
+
 	return SkeletalMeshComponent;
 }
 
@@ -2525,9 +2556,26 @@ static UMeshComponent* ValidateBindingAsset(
 
 	if (BindingAsset->GetGroomBindingType() == EGroomBindingMeshType::SkeletalMesh)
 	{
-		return ValidateBindingAsset(GroomAsset, BindingAsset, Cast<USkeletalMeshComponent>(MeshComponent), bIsBindingReloading, bValidationEnable, Component);
+		USkeletalMeshComponent* SkelMeshComponent = Cast<USkeletalMeshComponent>(MeshComponent);
+		if (bValidationEnable)
+		{
+			SkelMeshComponent = ValidateBindingAsset_SkeletalMesh(GroomAsset, BindingAsset, SkelMeshComponent, bIsBindingReloading, Component);
+		}
+		return SkelMeshComponent;
+	}
+	else if (BindingAsset->GetGroomBindingType() == EGroomBindingMeshType::GeometryCache)
+	{
+		UGeometryCacheComponent* GeometryCacheComponent = Cast<UGeometryCacheComponent>(MeshComponent);
+		if (bValidationEnable)
+		{
+			GeometryCacheComponent = ValidateBindingAsset_GeometryCache(GroomAsset, BindingAsset, GeometryCacheComponent, bIsBindingReloading, Component);
+		}
+		return GeometryCacheComponent;
+	}
+	else
+	{
+		return nullptr;
 	}
-	return ValidateBindingAsset(GroomAsset, BindingAsset, Cast<UGeometryCacheComponent>(MeshComponent), bIsBindingReloading, bValidationEnable, Component);
 }
 
 static EGroomGeometryType GetEffectiveGeometryType(EGroomGeometryType Type, bool bUseCards, EShaderPlatform InPlatform)
@@ -2675,47 +2723,6 @@ void UGroomComponent::InitResources(bool bIsBindingReloading)
 		}
 	}
 
-	// Validate that if we are bound to a skel. mesh, and we have binding, and the binding type is set to skinning, that the skin cache is enabled for these
-	if (bHasNeedSkinningBinding && ParentMeshComponent)
-	{
-		if (USkeletalMeshComponent* ParentSkelMeshComponent = Cast<USkeletalMeshComponent>(ParentMeshComponent))
-		{
-			if (BindingAsset)
-			{
-				// Extract if skin cache or mesh. deformer is enabled on at least on LOD.
-				// Since there is 1:1 mapping between groom LOD and mesh LOD, only use this has a hint.
-				bool bSupportSkinCache = ParentSkelMeshComponent->HasMeshDeformer();
-				if (!bSupportSkinCache)
-				{
-					for (uint32 SkelLODIt = 0, SkelLODCount = ParentSkelMeshComponent->GetNumLODs(); SkelLODIt < SkelLODCount; ++SkelLODIt)
-					{
-						bSupportSkinCache = bSupportSkinCache || ParentSkelMeshComponent->IsSkinCacheAllowed(SkelLODIt);
-					}
-				}
-
-				for (int32 GroupIt = 0, GroupCount = GroomAsset->GetHairGroupsPlatformData().Num(); GroupIt < GroupCount; ++GroupIt)
-				{
-					for (uint32 LODIt = 0, LODCount = GroomAsset->GetLODCount(); LODIt < LODCount; ++LODIt)
-					{
-						const EGroomBindingType BindingType = GroomAsset->GetBindingType(GroupIt, LODIt);
-						const bool bIsVisible = GroomAsset->IsVisible(GroupIt, LODIt);
-
-						if (BindingType == EGroomBindingType::Skinning && (!bSupportSkinCache && !IsHairManualSkinCacheEnabled()) && bIsVisible)
-						{
-							UE_LOG(LogHairStrands, Warning, TEXT("[Groom] Groom asset (Group:%d/%d) is set to use Skinning at LOD %d/%d while the parent skel. mesh does not support skin. cache at this LOD - Groom:%s - Skel.Mesh:%s"),
-								GroupIt,
-								GroupCount,
-								LODIt,
-								LODCount,
-								*GroomAsset->GetPathName(),
-								*ParentMeshComponent->GetPathName());
-						}
-					}
-				}
-			}
-		}
-	}
-
 	// Insure the ticking of the Groom component always happens after the skeletalMeshComponent.
 	UGroomBindingAsset* LocalBindingAsset = nullptr;
 	if (ValidatedMeshComponent)
diff --git a/Engine/Plugins/Runtime/HairStrands/Source/HairStrandsCore/Public/GroomComponent.h b/Engine/Plugins/Runtime/HairStrands/Source/HairStrandsCore/Public/GroomComponent.h
--- a/Engine/Plugins/Runtime/HairStrands/Source/HairStrandsCore/Public/GroomComponent.h
+++ b/Engine/Plugins/Runtime/HairStrands/Source/HairStrandsCore/Public/GroomComponent.h
@@ -339,11 +339,6 @@ private:
 	TObjectPtr<UGroomBindingAsset> BindingAssetBeingLoaded;
 #endif
 
-protected:
-	// Used for tracking if a Niagara component is attached or not
-	virtual void OnChildAttached(USceneComponent* ChildComponent) override;
-	virtual void OnChildDetached(USceneComponent* ChildComponent) override;
-
 private:
 	void DeleteDeferredHairGroupInstances();
 	void* InitializedResources;
