Remove unused code.

diff --git a/Engine/Plugins/Runtime/HairStrands/Source/HairStrandsCore/Private/GroomComponent.cpp b/Engine/Plugins/Runtime/HairStrands/Source/HairStrandsCore/Private/GroomComponent.cpp
--- a/Engine/Plugins/Runtime/HairStrands/Source/HairStrandsCore/Private/GroomComponent.cpp
+++ b/Engine/Plugins/Runtime/HairStrands/Source/HairStrandsCore/Private/GroomComponent.cpp
@@ -1378,7 +1378,6 @@ UGroomComponent::UGroomComponent(const FObjectInitializer& ObjectInitializer)
 	bAutoActivate = true;
 	bSelectable = true;
 	RegisteredMeshComponent = nullptr;
-	SkeletalPreviousPositionOffset = FVector::ZeroVector;
 	InitializedResources = nullptr;
 	Mobility = EComponentMobility::Movable;
 	bIsGroomAssetCallbackRegistered = false;
@@ -3175,7 +3174,6 @@ void UGroomComponent::ReleaseResources()
 	{
 		RemoveTickPrerequisiteComponent(RegisteredMeshComponent);
 	}
-	SkeletalPreviousPositionOffset = FVector::ZeroVector;
 	RegisteredMeshComponent = nullptr;
 
 	GroomCacheBuffers.Reset();
diff --git a/Engine/Plugins/Runtime/HairStrands/Source/HairStrandsCore/Public/GroomComponent.h b/Engine/Plugins/Runtime/HairStrands/Source/HairStrandsCore/Public/GroomComponent.h
--- a/Engine/Plugins/Runtime/HairStrands/Source/HairStrandsCore/Public/GroomComponent.h
+++ b/Engine/Plugins/Runtime/HairStrands/Source/HairStrandsCore/Public/GroomComponent.h
@@ -344,7 +344,6 @@ private:
 	void* InitializedResources;
 	class UMeshComponent* RegisteredMeshComponent;
 	class UMeshComponent* DeformedMeshComponent;
-	FVector SkeletalPreviousPositionOffset;
 	bool bIsGroomAssetCallbackRegistered;
 	bool bIsGroomBindingAssetCallbackRegistered;
 	bool bValidationEnable = true;
