* Fix scene preview profile in groom editor.
* Add widget in groom editor for switching scene profile.

diff --git a/Engine/Plugins/Runtime/HairStrands/Source/HairStrandsEditor/Private/GroomEditorViewportToolBar.cpp b/Engine/Plugins/Runtime/HairStrands/Source/HairStrandsEditor/Private/GroomEditorViewportToolBar.cpp
--- a/Engine/Plugins/Runtime/HairStrands/Source/HairStrandsEditor/Private/GroomEditorViewportToolBar.cpp
+++ b/Engine/Plugins/Runtime/HairStrands/Source/HairStrandsEditor/Private/GroomEditorViewportToolBar.cpp
@@ -17,6 +17,7 @@
 #include "Widgets/Text/STextBlock.h"
 #include "GroomVisualizationMenuCommands.h"
 #include "EditorViewportCommands.h"
+#include "PreviewProfileController.h"
 
 #define LOCTEXT_NAMESPACE "GroomEditorViewportToolbar"
 
@@ -26,7 +27,7 @@
 
 void SGroomEditorViewportToolbar::Construct(const FArguments& InArgs, TSharedPtr<class ICommonEditorViewportToolbarInfoProvider> InInfoProvider)
 {
-	SCommonEditorViewportToolbarBase::Construct(SCommonEditorViewportToolbarBase::FArguments(), InInfoProvider);
+	SCommonEditorViewportToolbarBase::Construct(SCommonEditorViewportToolbarBase::FArguments().PreviewProfileController(MakeShared<FPreviewProfileController>()), InInfoProvider);
 }
 
 // SCommonEditorViewportToolbarBase interface
diff --git a/Engine/Plugins/Runtime/HairStrands/Source/HairStrandsEditor/Private/Widgets/SGroomEditorViewport.cpp b/Engine/Plugins/Runtime/HairStrands/Source/HairStrandsEditor/Private/Widgets/SGroomEditorViewport.cpp
--- a/Engine/Plugins/Runtime/HairStrands/Source/HairStrandsEditor/Private/Widgets/SGroomEditorViewport.cpp
+++ b/Engine/Plugins/Runtime/HairStrands/Source/HairStrandsEditor/Private/Widgets/SGroomEditorViewport.cpp
@@ -21,6 +21,8 @@
 #include "GroomEditorViewportToolBar.h"
 #include "GroomVisualizationMenuCommands.h"
 #include "EditorViewportCommands.h"
+#include "AssetViewerSettings.h"
+#include "Editor/EditorPerProjectUserSettings.h"
 
 #define LOCTEXT_NAMESPACE "SGroomEditorViewport"
 
@@ -31,7 +33,10 @@ class FGroomEditorViewportClient : public FEditorViewportClient
 {
 public:
 	FGroomEditorViewportClient(FAdvancedPreviewScene& InPreviewScene, const TSharedRef<SGroomEditorViewport>& InGroomEditorViewport);
-	
+	virtual ~FGroomEditorViewportClient() override;
+
+	void OnAssetViewerSettingsChanged(const FName& InPropertyName);
+
 	// FEditorViewportClient interface
 	virtual FLinearColor GetBackgroundColor() const override;
 	virtual void Tick(float DeltaSeconds) override;
@@ -77,6 +82,29 @@ FGroomEditorViewportClient::FGroomEditorViewportClient(FAdvancedPreviewScene& In
 
 	//This seems to be needed to get the correct world time in the preview.
 	SetIsSimulateInEditorViewport(true);	
+
+	// Register delegate to update the show flags when the post-processing is turned on or off
+	UAssetViewerSettings::Get()->OnAssetViewerSettingsChanged().AddRaw(this, &FGroomEditorViewportClient::OnAssetViewerSettingsChanged);
+
+	UEditorPerProjectUserSettings* PerProjectSettings = GetMutableDefault<UEditorPerProjectUserSettings>();
+	UAssetViewerSettings* DefaultSettings = UAssetViewerSettings::Get();
+	PerProjectSettings->AssetViewerProfileIndex = DefaultSettings->Profiles.IsValidIndex(PerProjectSettings->AssetViewerProfileIndex) ? PerProjectSettings->AssetViewerProfileIndex : 0;
+
+	FAdvancedPreviewScene* PreviewSceneCasted = static_cast<FAdvancedPreviewScene*>(PreviewScene);
+	PreviewSceneCasted->SetProfileIndex(PerProjectSettings->AssetViewerProfileIndex);
+}
+
+FGroomEditorViewportClient::~FGroomEditorViewportClient()
+{
+	UAssetViewerSettings::Get()->OnAssetViewerSettingsChanged().RemoveAll(this);
+}
+
+void FGroomEditorViewportClient::OnAssetViewerSettingsChanged(const FName& InPropertyName)
+{
+	if (GetPreviewScene() == nullptr)
+	{
+		return;
+	}
 }
 
 void FGroomEditorViewportClient::Tick(float DeltaSeconds)
