[GroomAsset/GroomBinding]
 - Fix crash when multiple threads query hair description

diff --git a/Engine/Plugins/Runtime/HairStrands/Source/HairStrandsCore/Private/GroomAsset.cpp b/Engine/Plugins/Runtime/HairStrands/Source/HairStrandsCore/Private/GroomAsset.cpp
--- a/Engine/Plugins/Runtime/HairStrands/Source/HairStrandsCore/Private/GroomAsset.cpp
+++ b/Engine/Plugins/Runtime/HairStrands/Source/HairStrandsCore/Private/GroomAsset.cpp
@@ -3,6 +3,7 @@
 #include "GroomAsset.h"
 
 #include "Async/Async.h"
+#include "Async/UniqueLock.h"
 #include "Engine/Texture2D.h"
 #include "EngineUtils.h"
 #include "GroomAssetImportData.h"
@@ -1148,6 +1149,11 @@ void UGroomAsset::ReleaseMeshesResource(uint32 GroupIndex)
 
 void UGroomAsset::UpdateHairGroupsInfo()
 {
+#if WITH_EDITORONLY_DATA
+	TRACE_CPUPROFILER_EVENT_SCOPE(UGroomAsset::UpdateHairGroupsInfo);
+	UE::TUniqueLock LockScope(InternalLock);
+#endif
+
 	const uint32 GroupCount = GetNumHairGroups();
 	const bool bForceReset = GetHairGroupsInfo().Num() != GroupCount;
 	GetHairGroupsInfo().SetNum(GroupCount);
@@ -2331,8 +2337,12 @@ FString UGroomAsset::GetDerivedDataKeyForStrands(uint32 GroupIndex)
 
 void UGroomAsset::CommitHairDescription(FHairDescription&& InHairDescription, EHairDescriptionType Type)
 {
+	TRACE_CPUPROFILER_EVENT_SCOPE(UGroomAsset::CommitHairDescription);
+
 	HairDescriptionType = Type;
 
+	UE::TUniqueLock LockScope(InternalLock);
+
 	CachedHairDescription[HairDescriptionType] = MakeUnique<FHairDescription>(InHairDescription);
 
 	if (!HairDescriptionBulkData[HairDescriptionType])
@@ -2349,17 +2359,35 @@ void UGroomAsset::CommitHairDescription(FHairDescription&& InHairDescription, EH
 
 const FHairDescriptionGroups& UGroomAsset::GetHairDescriptionGroups()
 {
-	check(HairDescriptionBulkData[HairDescriptionType]);
+	TRACE_CPUPROFILER_EVENT_SCOPE(UGroomAsset::GetHairDescriptionGroups);
 
-	if (!CachedHairDescription[HairDescriptionType])
-	{
-		CachedHairDescription[HairDescriptionType] = MakeUnique<FHairDescription>();
-		HairDescriptionBulkData[HairDescriptionType]->LoadHairDescription(*CachedHairDescription[HairDescriptionType]);
-	}
-	if (!CachedHairDescriptionGroups[HairDescriptionType])
+	check(HairDescriptionBulkData[HairDescriptionType]);
+	
+	if (!CachedHairDescription[HairDescriptionType] || !CachedHairDescriptionGroups[HairDescriptionType])
 	{
-		CachedHairDescriptionGroups[HairDescriptionType] = MakeUnique<FHairDescriptionGroups>();
-		FGroomBuilder::BuildHairDescriptionGroups(*CachedHairDescription[HairDescriptionType], *CachedHairDescriptionGroups[HairDescriptionType]);
+		UE::TUniqueLock LockScope(InternalLock);
+
+		if (!CachedHairDescription[HairDescriptionType])
+		{
+			TUniquePtr<FHairDescription> HairDescription = MakeUnique<FHairDescription>();
+			HairDescriptionBulkData[HairDescriptionType]->LoadHairDescription(*HairDescription);
+
+			// Since we don't take read-lock, we need to make sure all memory is visible inside newly created object before publishing it.
+			FPlatformMisc::MemoryBarrier();
+
+			CachedHairDescription[HairDescriptionType] = MoveTemp(HairDescription);
+		}
+
+		if (!CachedHairDescriptionGroups[HairDescriptionType])
+		{
+			TUniquePtr<FHairDescriptionGroups> HairDescriptionGroups = MakeUnique<FHairDescriptionGroups>();
+			FGroomBuilder::BuildHairDescriptionGroups(*CachedHairDescription[HairDescriptionType], *HairDescriptionGroups);
+			
+			// Since we don't take read-lock, we need to make sure all memory is visible inside newly created object before publishing it.
+			FPlatformMisc::MemoryBarrier();
+
+			CachedHairDescriptionGroups[HairDescriptionType] = MoveTemp(HairDescriptionGroups);
+		}
 	}
 
 	return *CachedHairDescriptionGroups[HairDescriptionType];
@@ -2772,7 +2800,7 @@ inline FString GetLODName(const UGroomAsset* Asset, uint32 LODIndex)
 }
 
 static bool InternalImportCardGeometry(
-	FHairGroupsCardsSourceDescription* InDesc, 
+	const FHairGroupsCardsSourceDescription* InDesc, 
 	const FHairDescriptionGroups& InHairDescriptionGroups, 
 	const TArray<FHairGroupsInterpolation>& InHairInterpolationGroups,
 	FHairCardsBulkData& OutBulkData, 
@@ -3676,6 +3704,7 @@ bool UGroomAsset::GetHairCardsGuidesDatas(
 	const int32 LODIndex,
 	FHairStrandsDatas& OutCardsGuidesData)
 {
+	TRACE_CPUPROFILER_EVENT_SCOPE(UGroomAsset::GetHairCardsGuidesDatas);
 	FHairStrandsDatas StrandsData;
 	FHairStrandsDatas GuidesData;
 	const bool bIsValid = GetHairStrandsDatas(GroupIndex, StrandsData, GuidesData);
diff --git a/Engine/Plugins/Runtime/HairStrands/Source/HairStrandsCore/Public/GroomAsset.h b/Engine/Plugins/Runtime/HairStrands/Source/HairStrandsCore/Public/GroomAsset.h
--- a/Engine/Plugins/Runtime/HairStrands/Source/HairStrandsCore/Public/GroomAsset.h
+++ b/Engine/Plugins/Runtime/HairStrands/Source/HairStrandsCore/Public/GroomAsset.h
@@ -18,6 +18,7 @@
 #include "Curves/CurveFloat.h"
 #include "HairStrandsInterface.h"
 #include "Engine/SkeletalMesh.h"
+#include "Async/RecursiveMutex.h"
 #include "Interfaces/Interface_AssetUserData.h"
 #include "UObject/PerPlatformProperties.h"
 #include "UObject/StrongObjectPtr.h"
@@ -830,6 +831,7 @@ private:
 	// Transient HairDescription & HairDescriptionGroups, which are built from HairDescriptionBulkData.
 	// All these data (bulk/desc/groups) needs to be in sync. I.e., when the HairDescription is updated, 
 	// HairDescriptionGroups needs to also be updated
+	UE::FRecursiveMutex InternalLock;
 	TUniquePtr<FHairDescription> CachedHairDescription[EHairDescriptionType::Count];
 	TUniquePtr<FHairDescriptionGroups> CachedHairDescriptionGroups[EHairDescriptionType::Count];
 
