Fix advanced postprocessing being disabled in groom editor, causing eye adaptation to be disabled.

Advanced postprocessing is now correctly enabled in groom editor if the scene profile enables it.

diff --git a/Engine/Plugins/Runtime/HairStrands/Source/HairStrandsEditor/Private/Widgets/SGroomEditorViewport.cpp b/Engine/Plugins/Runtime/HairStrands/Source/HairStrandsEditor/Private/Widgets/SGroomEditorViewport.cpp
--- a/Engine/Plugins/Runtime/HairStrands/Source/HairStrandsEditor/Private/Widgets/SGroomEditorViewport.cpp
+++ b/Engine/Plugins/Runtime/HairStrands/Source/HairStrandsEditor/Private/Widgets/SGroomEditorViewport.cpp
@@ -49,6 +49,7 @@ public:
 	virtual bool CanCycleWidgetMode() const override { return false; }
 
 	void SetShowGrid(bool bShowGrid);
+	void SetAdvancedShowFlagsForScene(const bool bAdvancedShowFlags);
 
 	virtual void SetIsSimulateInEditorViewport(bool bInIsSimulateInEditorViewport)override;
 	
@@ -92,6 +93,9 @@ FGroomEditorViewportClient::FGroomEditorViewportClient(FAdvancedPreviewScene& In
 
 	FAdvancedPreviewScene* PreviewSceneCasted = static_cast<FAdvancedPreviewScene*>(PreviewScene);
 	PreviewSceneCasted->SetProfileIndex(PerProjectSettings->AssetViewerProfileIndex);
+
+	// Set correct flags according to current profile settings
+	SetAdvancedShowFlagsForScene(UAssetViewerSettings::Get()->Profiles[PerProjectSettings->AssetViewerProfileIndex].bPostProcessingEnabled);
 }
 
 FGroomEditorViewportClient::~FGroomEditorViewportClient()
@@ -101,9 +105,27 @@ FGroomEditorViewportClient::~FGroomEditorViewportClient()
 
 void FGroomEditorViewportClient::OnAssetViewerSettingsChanged(const FName& InPropertyName)
 {
-	if (GetPreviewScene() == nullptr)
+	if (InPropertyName == GET_MEMBER_NAME_CHECKED(FPreviewSceneProfile, bPostProcessingEnabled) || InPropertyName == NAME_None)
+	{
+		FAdvancedPreviewScene* PreviewSceneCasted = static_cast<FAdvancedPreviewScene*>(PreviewScene);
+		UAssetViewerSettings* Settings = UAssetViewerSettings::Get();
+		const int32 ProfileIndex = PreviewSceneCasted->GetCurrentProfileIndex();
+		if (Settings->Profiles.IsValidIndex(ProfileIndex))
+		{
+			SetAdvancedShowFlagsForScene(Settings->Profiles[ProfileIndex].bPostProcessingEnabled);
+		}		
+	}
+}
+
+void FGroomEditorViewportClient::SetAdvancedShowFlagsForScene(const bool bAdvancedShowFlags)
+{	
+	if (bAdvancedShowFlags)
+	{
+		EngineShowFlags.EnableAdvancedFeatures();
+	}
+	else
 	{
-		return;
+		EngineShowFlags.DisableAdvancedFeatures();
 	}
 }
 
