Fix race condition between NiagaraDataInterface and GroomComponent.

The hair group instance data were incorrectly fetched and read from the game thread on the nigara side, while the could be recreated at the same time by groom component. This CL addresses this by reading the scene proxy's hair group instance data on the render thread.

diff --git a/Engine/Plugins/Runtime/HairStrands/Source/HairStrandsCore/Private/GroomComponent.cpp b/Engine/Plugins/Runtime/HairStrands/Source/HairStrandsCore/Private/GroomComponent.cpp
--- a/Engine/Plugins/Runtime/HairStrands/Source/HairStrandsCore/Private/GroomComponent.cpp
+++ b/Engine/Plugins/Runtime/HairStrands/Source/HairStrandsCore/Private/GroomComponent.cpp
@@ -1136,6 +1136,17 @@ private:
 	UMaterialInterface* Strands_DebugMaterial = nullptr;
 };
 
+FHairGroupInstance* GetHairGroupInstance(UGroomComponent* In, int32 InGroupIndex)
+{
+	check(In);
+	FHairStrandsSceneProxy* Proxy = (FHairStrandsSceneProxy*)In->GetSceneProxy();
+	if (Proxy && Proxy->HairGroupInstances.IsValidIndex(InGroupIndex))
+	{
+		return Proxy->HairGroupInstances[InGroupIndex];
+	}
+	return nullptr;
+}
+
 /** GroomCacheBuffers implementation that hold copies of the GroomCacheAnimationData needed for playback */
 class FGroomCacheBuffers : public IGroomCacheBuffers
 {
diff --git a/Engine/Plugins/Runtime/HairStrands/Source/HairStrandsCore/Private/Niagara/NiagaraDataInterfaceHairStrands.cpp b/Engine/Plugins/Runtime/HairStrands/Source/HairStrandsCore/Private/Niagara/NiagaraDataInterfaceHairStrands.cpp
--- a/Engine/Plugins/Runtime/HairStrands/Source/HairStrandsCore/Private/Niagara/NiagaraDataInterfaceHairStrands.cpp
+++ b/Engine/Plugins/Runtime/HairStrands/Source/HairStrandsCore/Private/Niagara/NiagaraDataInterfaceHairStrands.cpp
@@ -397,7 +397,7 @@ void FNDIHairStrandsData::Update(UNiagaraDataInterfaceHairStrands* Interface, FN
 			FHairGroupsPhysics& HairPhysics = GroomAsset->GetHairGroupsPhysics()[GroupIndex];
 			StrandsSize = static_cast<uint8>(HairPhysics.StrandsParameters.StrandsSize);
 
-			HairGroupInstance = Interface->IsComponentValid() ? Interface->SourceComponent->GetGroupInstance(GroupIndex) : nullptr;
+			HairGroupIndex = Interface->IsComponentValid() ? GroupIndex : -1;
 			HairGroupInstSource = Interface->IsComponentValid() ? Interface->SourceComponent : nullptr;
 
 			SubSteps = HairPhysics.SolverSettings.SubSteps;
diff --git a/Engine/Plugins/Runtime/HairStrands/Source/HairStrandsCore/Public/Niagara/NiagaraDataInterfaceHairStrands.h b/Engine/Plugins/Runtime/HairStrands/Source/HairStrandsCore/Public/Niagara/NiagaraDataInterfaceHairStrands.h
--- a/Engine/Plugins/Runtime/HairStrands/Source/HairStrandsCore/Public/Niagara/NiagaraDataInterfaceHairStrands.h
+++ b/Engine/Plugins/Runtime/HairStrands/Source/HairStrandsCore/Public/Niagara/NiagaraDataInterfaceHairStrands.h
@@ -21,6 +21,8 @@ static const int32 ThicknessOffset = 3;
 
 struct FNDIHairStrandsData;
 
+FHairGroupInstance* GetHairGroupInstance(UGroomComponent* In, int32 InGroupIndex);
+
 /** Render buffers that will be used in hlsl functions */
 struct FNDIHairStrandsBuffer : public FRenderResource
 {
@@ -122,6 +124,7 @@ struct FNDIHairStrandsData
 		bSkinningTransfer = false;
 		HairGroupInstSource = nullptr;
 		HairGroupInstance = nullptr;
+		HairGroupIndex = -1;
 
 		TickCount = 0;
 		ForceReset = true;
@@ -191,7 +194,13 @@ struct FNDIHairStrandsData
 			bSkinningTransfer = OtherDatas->bSkinningTransfer;
 			BindingType = OtherDatas->BindingType;
 			HairGroupInstSource = OtherDatas->HairGroupInstSource;
-			HairGroupInstance = OtherDatas->HairGroupInstance;
+			HairGroupInstance = nullptr;
+			HairGroupIndex = OtherDatas->HairGroupIndex;
+
+			if (HairGroupInstSource != nullptr)
+			{
+				HairGroupInstance = GetHairGroupInstance(HairGroupInstSource.Get(), HairGroupIndex);
+			}
 
 			TickCount = OtherDatas->TickCount;
 			ForceReset = OtherDatas->ForceReset;
@@ -286,6 +295,9 @@ struct FNDIHairStrandsData
 	/** Strands Gpu buffer */
 	FNDIHairStrandsBuffer* HairStrandsBuffer;
 
+	/** Hair group index */
+	int32 HairGroupIndex = -1;
+
 	/** Hair group instance */
 	FHairGroupInstance* HairGroupInstance;
 
