Fix race condition in hair cards builder.

diff --git a/Engine/Plugins/Runtime/HairStrands/Source/HairStrandsCore/Private/HairCardsBuilder.cpp b/Engine/Plugins/Runtime/HairStrands/Source/HairStrandsCore/Private/HairCardsBuilder.cpp
--- a/Engine/Plugins/Runtime/HairStrands/Source/HairStrandsCore/Private/HairCardsBuilder.cpp
+++ b/Engine/Plugins/Runtime/HairStrands/Source/HairStrandsCore/Private/HairCardsBuilder.cpp
@@ -772,80 +772,93 @@ static bool InternalImportGeometry_WithGeneratedGuides(
 			}
 		}
 		
+		struct FRootUVData
+		{
+			FVector2f RootUV0;
+			FVector2f RootUV1;
+		};
+		TArray<FRootUVData> OutRootUVDatas;
+		OutRootUVDatas.SetNum(CardsRoots.Num());
+
 		// 3. Find cards root / curve root
+
+		// 3.1 Find closet root UV
+		// /!\ N^2 loop: the number of cards should be relatively small
 		ParallelFor(CardsRoots.Num(), 
 		[
 			&CardsRoots,
 			&StrandsRoots,
 			&Out,
-			&OutBulk
+			&OutBulk,
+			&OutRootUVDatas
 		] (uint32 CardRootIt) 
 		//for (const FCardsRootData& CardsRoot : CardsRoots)
 		{
 			const FCardsRootData& CardsRoot = CardsRoots[CardRootIt];
-			for (uint32 CardsRootPositionIndex=0; CardsRootPositionIndex <2; CardsRootPositionIndex++)
+
+			auto FindRootUV = [&](const FVector3f CardsRootPosition0, const FVector3f CardsRootPosition1, FVector2f& OutRootUV0, FVector2f& OutRootUV1)
 			{
-				// 3.1 Find closet root UV
-				// /!\ N^2 loop: the number of cards should be relatively small
-				auto FindRootUV = [&](const FVector3f CardsRootPosition0, const FVector3f CardsRootPosition1, FVector2f& OutRootUV0, FVector2f& OutRootUV1)
+				float ClosestDistance1 = FLT_MAX;
+				float ClosestDistance0 = FLT_MAX;
+				OutRootUV0 = FVector2f::ZeroVector;
+				OutRootUV1 = FVector2f::ZeroVector;
+				for (const FStrandsRootData& StrandsRoot : StrandsRoots)
 				{
-					float ClosestDistance1 = FLT_MAX;
-					float ClosestDistance0 = FLT_MAX;
-					OutRootUV0 = FVector2f::ZeroVector;
-					OutRootUV1 = FVector2f::ZeroVector;
-					for (const FStrandsRootData& StrandsRoot : StrandsRoots)
-					{
-						const float Distance0 = FVector3f::Distance(StrandsRoot.Position, CardsRootPosition0);
-						const float Distance1 = FVector3f::Distance(StrandsRoot.Position, CardsRootPosition1);
-
-						if (Distance0 < ClosestDistance0)
-						{
-							ClosestDistance0 = Distance0;
-							OutRootUV0 = StrandsRoot.RootUV;
-						}
+					const float Distance0 = FVector3f::Distance(StrandsRoot.Position, CardsRootPosition0);
+					const float Distance1 = FVector3f::Distance(StrandsRoot.Position, CardsRootPosition1);
 
-						if (Distance1 < ClosestDistance1)
-						{
-							ClosestDistance1 = Distance1;
-							OutRootUV1 = StrandsRoot.RootUV;
-						}
+					if (Distance0 < ClosestDistance0)
+					{
+						ClosestDistance0 = Distance0;
+						OutRootUV0 = StrandsRoot.RootUV;
 					}
 
-				};
-
-				FVector2f RootUV0, RootUV1;
-				FindRootUV(CardsRoot.Root0.Position, CardsRoot.Root1.Position, RootUV0, RootUV1);
-
-				// 3.2 Apply root UV to all cards vertices
-				{
-					const uint32 CardsIndexCount = Out.Cards.IndexCounts[CardsRoot.CardsIndex];
-					const uint32 CardsIndexOffset = Out.Cards.IndexOffsets[CardsRoot.CardsIndex];
-					for (uint32 IndexIt = 0; IndexIt < CardsIndexCount; ++IndexIt)
+					if (Distance1 < ClosestDistance1)
 					{
-						const uint32 VertexIndex = Out.Cards.Indices[CardsIndexOffset + IndexIt];
-
-						// Linearly interpolate between the two roots based on the local V coordinate 
-						// * LocalU is along the card 
-						// * LocalV is across the card
-						//  U
-						//  ^  ____
-						//  | |    |
-						//  | |____| Card
-						//  | |    |
-						//  | |____|
-						//     ----> V
-						// Root0  Root1
-						const float TexCoordV = Out.Cards.LocalUVs[VertexIndex].Y;
-						const FVector2f RootUV = FMath::Lerp(RootUV0, RootUV1, TexCoordV);
-
-						Out.Cards.UVs[VertexIndex].Z = RootUV.X;
-						Out.Cards.UVs[VertexIndex].W = RootUV.Y;
-						OutBulk.UVs[VertexIndex].Z   = RootUV.X;
-						OutBulk.UVs[VertexIndex].W   = RootUV.Y;
+						ClosestDistance1 = Distance1;
+						OutRootUV1 = StrandsRoot.RootUV;
 					}
 				}
-			}
+
+			};
+
+			FRootUVData& OutRootUVData = OutRootUVDatas[CardRootIt];
+			FindRootUV(CardsRoot.Root0.Position, CardsRoot.Root1.Position, OutRootUVData.RootUV0, OutRootUVData.RootUV1);
 		});
+
+		// 3.2 Apply root UV to all cards vertices
+		// Apply this in a serial manner to avoid race in data writing. Some cards might share certain vertices, for which different RootUV could be written to.
+		for (uint32 CardRootIt = 0, CardRootCount = CardsRoots.Num(); CardRootIt < CardRootCount; ++CardRootIt)
+		{
+			const FRootUVData& RootUVData = OutRootUVDatas[CardRootIt];
+			const FCardsRootData& CardsRoot = CardsRoots[CardRootIt];
+
+			const uint32 CardsIndexCount = Out.Cards.IndexCounts[CardsRoot.CardsIndex];
+			const uint32 CardsIndexOffset = Out.Cards.IndexOffsets[CardsRoot.CardsIndex];
+			for (uint32 IndexIt = 0; IndexIt < CardsIndexCount; ++IndexIt)
+			{
+				const uint32 VertexIndex = Out.Cards.Indices[CardsIndexOffset + IndexIt];
+
+				// Linearly interpolate between the two roots based on the local V coordinate 
+				// * LocalU is along the card 
+				// * LocalV is across the card
+				//  U
+				//  ^  ____
+				//  | |    |
+				//  | |____| Card
+				//  | |    |
+				//  | |____|
+				//     ----> V
+				// Root0  Root1
+				const float TexCoordV = Out.Cards.LocalUVs[VertexIndex].Y;
+				const FVector2f RootUV = FMath::Lerp(RootUVData.RootUV0, RootUVData.RootUV1, TexCoordV);
+
+				Out.Cards.UVs[VertexIndex].Z = RootUV.X;
+				Out.Cards.UVs[VertexIndex].W = RootUV.Y;
+				OutBulk.UVs[VertexIndex].Z   = RootUV.X;
+				OutBulk.UVs[VertexIndex].W   = RootUV.Y;
+			}
+		}
 	}
 
 	return bSuccess;
