Set niagara system in solver property to be registered for packaging

diff --git a/Engine/Plugins/Runtime/HairStrands/Source/HairStrandsCore/Private/GroomAsset.cpp b/Engine/Plugins/Runtime/HairStrands/Source/HairStrandsCore/Private/GroomAsset.cpp
--- a/Engine/Plugins/Runtime/HairStrands/Source/HairStrandsCore/Private/GroomAsset.cpp
+++ b/Engine/Plugins/Runtime/HairStrands/Source/HairStrandsCore/Private/GroomAsset.cpp
@@ -1482,6 +1482,8 @@ void UGroomAsset::PostLoad()
 			}
 		}
 	}
+	// Update the physics system based on the enum for picking the right asset while packaging 
+	UpdatePhysicsSystems();
 }
 
 void UGroomAsset::BeginDestroy()
@@ -1609,6 +1611,10 @@ void UGroomAsset::PostEditChangeProperty(FPropertyChangedEvent& PropertyChangedE
 			CachedHairGroupsMeshes.Init(Dirty, GroupCount);
 		}
 	}
+	if(PropertyName == GET_MEMBER_NAME_CHECKED(FHairSolverSettings, NiagaraSolver))
+	{
+		UpdatePhysicsSystems();
+	}
 
 	// Rebuild the groom cached data if interpolation or LODs have changed
 	const bool bNeedRebuildDerivedData = 
@@ -1728,6 +1734,21 @@ void UGroomAsset::PostInitProperties()
 }
 #endif
 
+void UGroomAsset::UpdatePhysicsSystems()
+{
+	for (auto& Group : GetHairGroupsPhysics())
+	{
+		if (Group.SolverSettings.NiagaraSolver == EGroomNiagaraSolvers::AngularSprings)
+		{
+			Group.SolverSettings.CustomSystem = LoadObject<UNiagaraSystem>(nullptr, TEXT("/HairStrands/Emitters/StableSpringsSystem.StableSpringsSystem"));
+		}
+		else if (Group.SolverSettings.NiagaraSolver == EGroomNiagaraSolvers::CosseratRods)
+		{
+			Group.SolverSettings.CustomSystem = LoadObject<UNiagaraSystem>(nullptr, TEXT("/HairStrands/Emitters/StableRodsSystem.StableRodsSystem"));
+		}
+	}
+}
+
 int32 UGroomAsset::GetNumHairGroups() const
 {
 	return GetHairGroupsPlatformData().Num();
diff --git a/Engine/Plugins/Runtime/HairStrands/Source/HairStrandsCore/Private/GroomAssetPhysics.cpp b/Engine/Plugins/Runtime/HairStrands/Source/HairStrandsCore/Private/GroomAssetPhysics.cpp
--- a/Engine/Plugins/Runtime/HairStrands/Source/HairStrandsCore/Private/GroomAssetPhysics.cpp
+++ b/Engine/Plugins/Runtime/HairStrands/Source/HairStrandsCore/Private/GroomAssetPhysics.cpp
@@ -1,6 +1,7 @@
 // Copyright Epic Games, Inc. All Rights Reserved.
 
 #include "GroomAssetPhysics.h"
+#include "NiagaraSystem.h"
 
 #include UE_INLINE_GENERATED_CPP_BY_NAME(GroomAssetPhysics)
 
@@ -12,6 +13,7 @@ FHairSolverSettings::FHairSolverSettings()
 	IterationCount = 5;
 	GravityPreloading = 0.0f;
 	bForceVisible = false;
+	CustomSystem =  LoadObject<UNiagaraSystem>(nullptr, TEXT("/HairStrands/Emitters/StableSpringsSystem.StableSpringsSystem"));
 }
 
 FHairExternalForces::FHairExternalForces()
diff --git a/Engine/Plugins/Runtime/HairStrands/Source/HairStrandsCore/Public/GroomAsset.h b/Engine/Plugins/Runtime/HairStrands/Source/HairStrandsCore/Public/GroomAsset.h
--- a/Engine/Plugins/Runtime/HairStrands/Source/HairStrandsCore/Public/GroomAsset.h
+++ b/Engine/Plugins/Runtime/HairStrands/Source/HairStrandsCore/Public/GroomAsset.h
@@ -770,6 +770,9 @@ public:
 private:
 	void ApplyStripFlags(uint8 StripFlags, const class ITargetPlatform* CookTarget);
 
+	/** Update the physics system based on the solver settings enum */
+	void UpdatePhysicsSystems();
+
 	// Functions allocating lazily/on-demand resources (guides, interpolation, RT geometry, ...)
 	FHairStrandsRestResource*			AllocateGuidesResources(uint32 GroupIndex);
 	FHairStrandsInterpolationResource*	AllocateInterpolationResources(uint32 GroupIndex);
