* Fix typo in groom cache causing groom cache's position to be overriden by groom cache's radius.
* Fix groom cache reimport causing max hair radius computation to be incorrect due to a memory move operation
* Fix incorrect groom cache's hair radius normalization when updating interpolated hair radius.

diff --git a/Engine/Plugins/Runtime/HairStrands/Source/HairStrandsCore/Private/GroomCacheData.cpp b/Engine/Plugins/Runtime/HairStrands/Source/HairStrandsCore/Private/GroomCacheData.cpp
--- a/Engine/Plugins/Runtime/HairStrands/Source/HairStrandsCore/Private/GroomCacheData.cpp
+++ b/Engine/Plugins/Runtime/HairStrands/Source/HairStrandsCore/Private/GroomCacheData.cpp
@@ -5,6 +5,8 @@
 
 #include UE_INLINE_GENERATED_CPP_BY_NAME(GroomCacheData)
 
+float GetHairStrandsMaxRadius(const TArray<float>& InPointsRadius);
+
 FGroomAnimationInfo::FGroomAnimationInfo()
 	: NumFrames(0)
 	, SecondsPerFrame(0.0f)
@@ -69,9 +71,11 @@ void FGroomCacheStrandData::Serialize(FArchive& Ar, int32 Version, EGroomCacheAt
 	Ar << CurvesLength;
 }
 
+// Use VertexData.PointsRadius, rather than GroupData, because GroupData.StrandsPoints has been 'moved' to 
+// VertexData by the time we reach the GetHairStrandsMaxRadius function
 FGroomCacheGroupData::FGroomCacheGroupData(FHairStrandsDatas&& GroupData)
 : VertexData(MoveTemp(GroupData.StrandsPoints))
-, StrandData(MoveTemp(GroupData.StrandsCurves), GetHairStrandsMaxLength(GroupData), GetHairStrandsMaxRadius(GroupData))
+, StrandData(MoveTemp(GroupData.StrandsCurves), GetHairStrandsMaxLength(GroupData), GetHairStrandsMaxRadius(VertexData.PointsRadius)) 
 , BoundingBox(GroupData.BoundingBox)
 {
 }
diff --git a/Engine/Plugins/Runtime/HairStrands/Source/HairStrandsCore/Private/GroomManager.cpp b/Engine/Plugins/Runtime/HairStrands/Source/HairStrandsCore/Private/GroomManager.cpp
--- a/Engine/Plugins/Runtime/HairStrands/Source/HairStrandsCore/Private/GroomManager.cpp
+++ b/Engine/Plugins/Runtime/HairStrands/Source/HairStrandsCore/Private/GroomManager.cpp
@@ -1055,6 +1055,7 @@ static void RunHairStrandsInterpolation_Guide(
 					InstanceData.Instance->RegisteredIndex,
 					InstanceData.Instance->Guides.RestResource->GetPointCount(),
 					InterpolationFactor,
+					1.f /*MaxHairRadius*/,
 					CacheResources0,
 					CacheResources1,
 					RegisterAsSRV(GraphBuilder, InstanceData.Instance->Guides.RestResource->PositionBuffer),
@@ -1537,6 +1538,7 @@ static void RunHairStrandsInterpolation_Strands(
 					InstanceData.Instance->RegisteredIndex,
 					InstanceData.ActivePointCount,
 					InterpolationFactor,
+					InstanceData.Instance->Strands.Modifier.HairWidth * 0.5f /*InMaxHairRadius*/,
 					CacheResources0,
 					CacheResources1,
 					RegisterAsSRV(GraphBuilder, InstanceData.Instance->Strands.RestResource->PositionBuffer),
diff --git a/Engine/Plugins/Runtime/HairStrands/Source/HairStrandsCore/Private/HairStrandsDatas.cpp b/Engine/Plugins/Runtime/HairStrands/Source/HairStrandsCore/Private/HairStrandsDatas.cpp
--- a/Engine/Plugins/Runtime/HairStrands/Source/HairStrandsCore/Private/HairStrandsDatas.cpp
+++ b/Engine/Plugins/Runtime/HairStrands/Source/HairStrandsCore/Private/HairStrandsDatas.cpp
@@ -26,16 +26,21 @@ float GetHairStrandsMaxLength(const FHairStrandsDatas& In)
 	return MaxLength;
 }
 
-float GetHairStrandsMaxRadius(const FHairStrandsDatas& In)
+float GetHairStrandsMaxRadius(const TArray<float>& InPointsRadius)
 {
 	float MaxRadius = 0;
-	for (float PointRadius : In.StrandsPoints.PointsRadius)
+	for (float PointRadius : InPointsRadius)
 	{
 		MaxRadius = FMath::Max(MaxRadius, PointRadius);
 	}
 	return MaxRadius;
 }
 
+float GetHairStrandsMaxRadius(const FHairStrandsDatas& In)
+{
+	return GetHairStrandsMaxRadius(In.StrandsPoints.PointsRadius);
+}
+
 ///////////////////////////////////////////////////////////////////////////////////////////////////
 // Curve/Point/Strands/Interpolation/Culling Data
 
diff --git a/Engine/Plugins/Runtime/HairStrands/Source/HairStrandsCore/Private/HairStrandsInterpolation.cpp b/Engine/Plugins/Runtime/HairStrands/Source/HairStrandsCore/Private/HairStrandsInterpolation.cpp
--- a/Engine/Plugins/Runtime/HairStrands/Source/HairStrandsCore/Private/HairStrandsInterpolation.cpp
+++ b/Engine/Plugins/Runtime/HairStrands/Source/HairStrandsCore/Private/HairStrandsInterpolation.cpp
@@ -153,7 +153,7 @@ FGroomCacheResources CreateGroomCacheBuffer(FRDGBuilder& GraphBuilder, FGroomCac
 				InVertexData.PointsRadius.GetData(),
 				RadiusDataSizeInBytes,
 				ERDGInitialDataFlags::None);
-			GraphBuilder.QueueBufferExtraction(RadiusBuffer, &InVertexData.PositionBuffer);
+			GraphBuilder.QueueBufferExtraction(RadiusBuffer, &InVertexData.RadiusBuffer);
 		}
 		else
 		{
@@ -237,6 +237,7 @@ class FGroomCacheUpdatePassCS : public FGlobalShader
 		SHADER_PARAMETER(uint32, ElementCount)
 		SHADER_PARAMETER(uint32, bHasRadiusData)
 		SHADER_PARAMETER(float, InterpolationFactor)
+		SHADER_PARAMETER(float, MaxHairRadius)
 		SHADER_PARAMETER_RDG_BUFFER_SRV(Buffer, InPosition0Buffer)
 		SHADER_PARAMETER_RDG_BUFFER_SRV(Buffer, InPosition1Buffer)
 		SHADER_PARAMETER_RDG_BUFFER_SRV(Buffer, InRadius0Buffer)
@@ -265,6 +266,7 @@ void AddGroomCacheUpdatePass(
 	uint32 InstanceRegisteredIndex,
 	uint32 PointCount,
 	float InterpolationFactor,
+	float InMaxHairRadius,
 	FGroomCacheResources CacheResources0,
 	FGroomCacheResources CacheResources1,
 	FRDGBufferSRVRef InBuffer,
@@ -286,6 +288,7 @@ void AddGroomCacheUpdatePass(
 	Parameters->InDeformedOffsetBuffer = InDeformedOffsetBuffer;
 	Parameters->OutDeformedBuffer = OutBuffer;
 	Parameters->InterpolationFactor = InterpolationFactor;
+	Parameters->MaxHairRadius = InMaxHairRadius;
 	Parameters->bHasRadiusData = CacheResources0.bHasRadiusData ? 1u : 0u;
 
 	const FIntVector DispatchCount = FIntVector(FMath::DivideAndRoundUp(PointCount, FGroomCacheUpdatePassCS::GetGroupSize()), 1, 1);
diff --git a/Engine/Plugins/Runtime/HairStrands/Source/HairStrandsCore/Public/HairStrandsInterpolation.h b/Engine/Plugins/Runtime/HairStrands/Source/HairStrandsCore/Public/HairStrandsInterpolation.h
--- a/Engine/Plugins/Runtime/HairStrands/Source/HairStrandsCore/Public/HairStrandsInterpolation.h
+++ b/Engine/Plugins/Runtime/HairStrands/Source/HairStrandsCore/Public/HairStrandsInterpolation.h
@@ -83,6 +83,7 @@ void AddGroomCacheUpdatePass(
 	uint32 InstanceRegisteredIndex,
 	uint32 PointCount,
 	float InterpolationFactor,
+	float MaxHairRadius,
 	FGroomCacheResources CacheResources0,
 	FGroomCacheResources CacheResources1,
 	FRDGBufferSRVRef InBuffer,
diff --git a/Engine/Shaders/Private/HairStrands/HairStrandsInterpolation.usf b/Engine/Shaders/Private/HairStrands/HairStrandsInterpolation.usf
--- a/Engine/Shaders/Private/HairStrands/HairStrandsInterpolation.usf
+++ b/Engine/Shaders/Private/HairStrands/HairStrandsInterpolation.usf
@@ -917,6 +917,7 @@ uint InstanceRegisteredIndex;
 uint ElementCount;
 uint bHasRadiusData;
 float InterpolationFactor;
+float MaxHairRadius;
 
 Buffer<float> InPosition0Buffer;
 Buffer<float> InPosition1Buffer;
@@ -958,7 +959,7 @@ void MainCS(uint DispatchThreadId : SV_DispatchThreadID)
 
 	const float3 DeformedPositionOffset = ReadRenPositionOffset(InDeformedOffsetBuffer, InstanceRegisteredIndex);
 
-	WriteHairControlPointPosition(OutDeformedBuffer, Index, CP, DeformedPositionOffset, 1 /*Radius*/);
+	WriteHairControlPointPosition(OutDeformedBuffer, Index, CP, DeformedPositionOffset, MaxHairRadius);
 }
 
 #endif // SHADER_GROOMCACHE_UPDATE
