Fix incorrect assert when forcing specific LOD on GroomComponent.

diff --git a/Engine/Plugins/Runtime/HairStrands/Source/HairStrandsCore/Private/GroomComponent.cpp b/Engine/Plugins/Runtime/HairStrands/Source/HairStrandsCore/Private/GroomComponent.cpp
--- a/Engine/Plugins/Runtime/HairStrands/Source/HairStrandsCore/Private/GroomComponent.cpp
+++ b/Engine/Plugins/Runtime/HairStrands/Source/HairStrandsCore/Private/GroomComponent.cpp
@@ -1546,17 +1546,19 @@ void UGroomComponent::UpdateHairSimulation()
 	UpdateSimulatedGroups();
 }
 
-void UGroomComponent::SwitchSimulationLOD(const int32 PreviousLOD, const int32 CurrentLOD)
+void UGroomComponent::SwitchSimulationLOD(const int32 PreviousLOD, const int32 CurrentLOD, const EHairLODSelectionType InLODSelectionType)
 {
 	if (GroomAsset && PreviousLOD != CurrentLOD)
 	{
 		bool bRequiresSimulationUpdate = false;
 		for (int32 GroupIt = 0, GroupCount = GroomAsset->GetHairGroupsPlatformData().Num(); GroupIt < GroupCount; ++GroupIt)
 		{
-			if ((IsSimulationEnable(GroupIt, PreviousLOD) != IsSimulationEnable(GroupIt, CurrentLOD)))
+			const bool bPrevSimEnable = IsSimulationEnable(GroupIt, PreviousLOD);
+			const bool bCurrSimEnable = IsSimulationEnable(GroupIt, CurrentLOD);
+			if (bPrevSimEnable != bCurrSimEnable)
 			{
 				// Sanity check. Simulation does not support immediate mode.
-				check(LODSelectionType != EHairLODSelectionType::Immediate);
+				check(InLODSelectionType != EHairLODSelectionType::Immediate);
 
 				CreateHairSimulation(GroupIt, CurrentLOD);
 				bRequiresSimulationUpdate = true;
@@ -1793,7 +1795,7 @@ void UGroomComponent::SetForcedLOD(int32 CurrLODIndex)
 	}
 
 	// Inform simulation about LOD switch.
-	SwitchSimulationLOD(PrevLODIndex, CurrLODIndex);
+	SwitchSimulationLOD(PrevLODIndex, CurrLODIndex, CurrLODSelectionType);
 
 	// Finally, update the forced LOD value and LOD selection type
 	const bool bHasLODSwitch = CurrLODIndex != PrevLODIndex;
@@ -3655,7 +3657,7 @@ void UGroomComponent::TickComponent(float DeltaTime, enum ELevelTick TickType, F
 		}
 
 		// 2. If there is a LOD change, update the simulation adequately 
-		SwitchSimulationLOD(LODPredictedIndex, EffectiveForceLOD);
+		SwitchSimulationLOD(LODPredictedIndex, EffectiveForceLOD, LODSelectionType);
 
 		// 3. Update global predicted index (used for SyncLOD API)
 		LODPredictedIndex = EffectiveForceLOD;
diff --git a/Engine/Plugins/Runtime/HairStrands/Source/HairStrandsCore/Public/GroomComponent.h b/Engine/Plugins/Runtime/HairStrands/Source/HairStrandsCore/Public/GroomComponent.h
--- a/Engine/Plugins/Runtime/HairStrands/Source/HairStrandsCore/Public/GroomComponent.h
+++ b/Engine/Plugins/Runtime/HairStrands/Source/HairStrandsCore/Public/GroomComponent.h
@@ -110,7 +110,7 @@ public:
 	void CreateHairSimulation(const int32 GroupIndex, const int32 LODIndex);
 
 	/** Enable/Disable hair simulation while transitioning from one LOD to another one */
-	void SwitchSimulationLOD(const int32 PreviousLOD, const int32 CurrentLOD);
+	void SwitchSimulationLOD(const int32 PreviousLOD, const int32 CurrentLOD, const EHairLODSelectionType InLODSelectionType);
 
 	/** Check if the simulation is enabled or not */
 	bool IsSimulationEnable(int32 GroupIndex, int32 LODIndex) const;
