Fix missing control point when enabling triangl strip geometry for hair strands.

diff --git a/Engine/Plugins/Runtime/HairStrands/Source/HairStrandsCore/Private/GroomBuilder.cpp b/Engine/Plugins/Runtime/HairStrands/Source/HairStrandsCore/Private/GroomBuilder.cpp
--- a/Engine/Plugins/Runtime/HairStrands/Source/HairStrandsCore/Private/GroomBuilder.cpp
+++ b/Engine/Plugins/Runtime/HairStrands/Source/HairStrandsCore/Private/GroomBuilder.cpp
@@ -51,7 +51,7 @@ bool DoesHairStrandsSupportCompressedPosition();
 
 FString FGroomBuilder::GetVersion()
 {
-	return TEXT("v16");
+	return TEXT("v16d");
 }
 
 namespace GroomBuilder_Voxelization
@@ -2124,7 +2124,8 @@ bool FGroomBuilder::BuildHairDescriptionGroups(const FHairDescription& HairDescr
 			continue;
 		}
 
-		CurrentHairStrandsDatas->StrandsCurves.CurvesCount.Add(FMath::Min(uint32(CurveNumVertices), HAIR_MAX_NUM_POINT_PER_CURVE));
+		const bool bAddEndingControlPoint = GetHairStrandsUsesTriangleStrips() && uint32(CurveNumVertices + 1) <= HAIR_MAX_NUM_POINT_PER_CURVE;
+		CurrentHairStrandsDatas->StrandsCurves.CurvesCount.Add(FMath::Min(uint32(CurveNumVertices + (bAddEndingControlPoint ? 1u : 0u)), HAIR_MAX_NUM_POINT_PER_CURVE));
 
 		if (bCanUseClosestGuidesAndWeights)
 		{
@@ -2230,6 +2231,26 @@ bool FGroomBuilder::BuildHairDescriptionGroups(const FHairDescription& HairDescr
 			}
 
 			CurrentHairStrandsDatas->StrandsPoints.PointsRadius.Add(VertexWidth * 0.5f);
+
+			// Add extra control point at the end of each curve when GetHairStrandsUsesTriangleStrips() 
+			// is enabled to avoid loosing the last segment of each curve
+			if (bAddEndingControlPoint && (VertexIndex == CurveNumVertices-1))
+			{
+				CurrentHairStrandsDatas->StrandsPoints.PointsPosition.Add(VertexPositions[VertexID]);
+				if (bHasBaseColorAttribute)
+				{
+					CurrentHairStrandsDatas->StrandsPoints.PointsBaseColor.Add(FLinearColor(VertexBaseColor[VertexID]));
+				}
+				if (bHasRoughnessAttribute)
+				{
+					CurrentHairStrandsDatas->StrandsPoints.PointsRoughness.Add(VertexRoughness[VertexID]);
+				}
+				if (bHasAOAttribute)
+				{
+					CurrentHairStrandsDatas->StrandsPoints.PointsAO.Add(VertexAO[VertexID]);
+				}
+				CurrentHairStrandsDatas->StrandsPoints.PointsRadius.Add(VertexWidth * 0.5f);
+			}
 		}
 	}
 
