Fix Niagara hair solver being loading on async thread, while it is only allowed to be loaded on the game thread.

diff --git a/Engine/Plugins/Runtime/HairStrands/Source/HairStrandsCore/Private/GroomAssetPhysics.cpp b/Engine/Plugins/Runtime/HairStrands/Source/HairStrandsCore/Private/GroomAssetPhysics.cpp
--- a/Engine/Plugins/Runtime/HairStrands/Source/HairStrandsCore/Private/GroomAssetPhysics.cpp
+++ b/Engine/Plugins/Runtime/HairStrands/Source/HairStrandsCore/Private/GroomAssetPhysics.cpp
@@ -13,7 +13,12 @@ FHairSolverSettings::FHairSolverSettings()
 	IterationCount = 5;
 	GravityPreloading = 0.0f;
 	bForceVisible = false;
-	CustomSystem =  LoadObject<UNiagaraSystem>(nullptr, TEXT("/HairStrands/Emitters/StableSpringsSystem.StableSpringsSystem"));
+
+	// Object loading is only allowed on the GameThread
+	if (IsInGameThread())
+	{
+		CustomSystem =  LoadObject<UNiagaraSystem>(nullptr, TEXT("/HairStrands/Emitters/StableSpringsSystem.StableSpringsSystem"));
+	}
 }
 
 FHairExternalForces::FHairExternalForces()
