Fix (tentative) for memory corruption inside hair card builder.

The mesh description retrieved from the static mesh can be shared between cards asset. Since the builder is modifying this asset, this could cause some race condition. The CL create a local copy of the mesh description to ensure there is no side effect.

diff --git a/Engine/Plugins/Runtime/HairStrands/Source/HairStrandsCore/Private/HairCardsBuilder.cpp b/Engine/Plugins/Runtime/HairStrands/Source/HairStrandsCore/Private/HairCardsBuilder.cpp
--- a/Engine/Plugins/Runtime/HairStrands/Source/HairStrandsCore/Private/HairCardsBuilder.cpp
+++ b/Engine/Plugins/Runtime/HairStrands/Source/HairStrandsCore/Private/HairCardsBuilder.cpp
@@ -66,7 +66,7 @@ namespace FHairCardsBuilder
 FString GetVersion()
 {
 	// Important to update the version when cards building or importing changes
-	return TEXT("10");
+	return TEXT("11");
 }
 
 bool InternalCreateCardsGuides(
@@ -457,8 +457,12 @@ static bool InternalImportGeometry_WithGeneratedGuides(
 {
 	const uint32 MeshLODIndex = 0;
 
-	// Note: if there are multiple section we only import the first one. Support for multiple section could be added later on. 
-	FMeshDescription* MeshDescription = StaticMesh->GetMeshDescription(0);
+	// Note: * if there are multiple section we only import the first one. Support for multiple section could be added later on. 
+	//       * use a local copy of the mesh description as multiple card assets referencing the same mesh description could be 
+	//         built in parallel. We later call SanitizeMeshDescription which modify the mesh description and is not thread safe
+	FMeshDescription LocalMeshDescription;
+	StaticMesh->CloneMeshDescription(0, LocalMeshDescription);
+	FMeshDescription* MeshDescription = &LocalMeshDescription;
 	const uint32 VertexCount = MeshDescription->Vertices().Num();
 	uint32 IndexCount  = MeshDescription->Triangles().Num() * 3;
 
