Add small position offset for the added control point when triangle strip optimization is enabled.

This fix the last control point's tangent orientation.

diff --git a/Engine/Plugins/Runtime/HairStrands/Source/HairStrandsCore/Private/GroomAsset.cpp b/Engine/Plugins/Runtime/HairStrands/Source/HairStrandsCore/Private/GroomAsset.cpp
--- a/Engine/Plugins/Runtime/HairStrands/Source/HairStrandsCore/Private/GroomAsset.cpp
+++ b/Engine/Plugins/Runtime/HairStrands/Source/HairStrandsCore/Private/GroomAsset.cpp
@@ -2137,11 +2137,13 @@ namespace GroomDerivedDataCacheUtils
 	void SerializeHairInterpolationSettingsForDDC(FArchive& Ar, uint32 GroupIndex, FHairGroupsInterpolation& InterpolationSettings, FHairGroupsLOD& LODSettings, bool bRequireInterpolationData)
 	{
 		bool bSupportCompressedPosition = DoesHairStrandsSupportCompressedPosition();
+		bool bSupportTriangleStrip = GetHairStrandsUsesTriangleStrips();
 
 		// Note: this serializer is only used to build the groom DDC key, no versioning is required
 		Ar << GroupIndex;
 		Ar << bRequireInterpolationData;
 		Ar << bSupportCompressedPosition;
+		Ar << bSupportTriangleStrip;
 
 		InterpolationSettings.BuildDDCKey(Ar);
 		LODSettings.BuildDDCKey(Ar);
diff --git a/Engine/Plugins/Runtime/HairStrands/Source/HairStrandsCore/Private/GroomBuilder.cpp b/Engine/Plugins/Runtime/HairStrands/Source/HairStrandsCore/Private/GroomBuilder.cpp
--- a/Engine/Plugins/Runtime/HairStrands/Source/HairStrandsCore/Private/GroomBuilder.cpp
+++ b/Engine/Plugins/Runtime/HairStrands/Source/HairStrandsCore/Private/GroomBuilder.cpp
@@ -51,7 +51,7 @@ bool DoesHairStrandsSupportCompressedPosition();
 
 FString FGroomBuilder::GetVersion()
 {
-	return TEXT("v16d");
+	return TEXT("v16e");
 }
 
 namespace GroomBuilder_Voxelization
@@ -2236,7 +2236,19 @@ bool FGroomBuilder::BuildHairDescriptionGroups(const FHairDescription& HairDescr
 			// is enabled to avoid loosing the last segment of each curve
 			if (bAddEndingControlPoint && (VertexIndex == CurveNumVertices-1))
 			{
-				CurrentHairStrandsDatas->StrandsPoints.PointsPosition.Add(VertexPositions[VertexID]);
+				const int32 CurrentIndex = CurrentHairStrandsDatas->StrandsPoints.PointsPosition.Num();
+				const int32 P0Index = CurrentIndex-1;
+				const int32 P1Index = CurrentIndex-2;
+				check(CurveNumVertices >= 2);
+				check(P0Index > 0);
+				check(P1Index > 0);
+
+				// Add a small position offset to avoid superposed positions of the extra position
+				const FVector3f P0 = CurrentHairStrandsDatas->StrandsPoints.PointsPosition[P0Index];
+				const FVector3f P1 = CurrentHairStrandsDatas->StrandsPoints.PointsPosition[P1Index];
+				const FVector3f PositionOffset = (P0 - P1) * 0.1f; // 10% of the distance
+
+				CurrentHairStrandsDatas->StrandsPoints.PointsPosition.Add(VertexPositions[VertexID] + PositionOffset);
 				if (bHasBaseColorAttribute)
 				{
 					CurrentHairStrandsDatas->StrandsPoints.PointsBaseColor.Add(FLinearColor(VertexBaseColor[VertexID]));
diff --git a/Engine/Source/Runtime/Core/Private/UObject/DevObjectVersion.cpp b/Engine/Source/Runtime/Core/Private/UObject/DevObjectVersion.cpp
--- a/Engine/Source/Runtime/Core/Private/UObject/DevObjectVersion.cpp
+++ b/Engine/Source/Runtime/Core/Private/UObject/DevObjectVersion.cpp
@@ -129,7 +129,7 @@ FGuid FDevSystemGuidRegistry::GetSystemGuid(FGuid System)
 FDevSystemGuids::FDevSystemGuids()
 	: GLOBALSHADERMAP_DERIVEDDATA_VER(0x7BB10A3C, 0xAC4E46F3, 0xAC78F4C1, 0xEFB3E34F)
 	, GROOM_BINDING_DERIVED_DATA_VERSION(0x30769E53, 0x0C574C7B, 0xA15C56F2, 0x24A64E32)
-	, GROOM_DERIVED_DATA_VERSION(0x05A37379, 0xF8A049D5, 0x986824E9, 0xAAA83F41)
+	, GROOM_DERIVED_DATA_VERSION(0x8D2064FA, 0x4EDF4DB7, 0x99B33838, 0xA79A8B79)
 	, LANDSCAPE_MOBILE_COOK_VERSION(0x0E9ADF72, 0xD6B64E0D, 0x81C4A92B, 0x081A37AB)
 	, MATERIALSHADERMAP_DERIVEDDATA_VER(0x5421F350, 0xA448491B, 0x95422D86, 0xB105837C)
 	, NANITE_DERIVEDDATA_VER(0xBEB0226A, 0x070E4ECA, 0x972C1E7D, 0xD8599E68)
