Fix RT groom geometry being rebuilt every frame with static groom.

This was causing OOM on certain platform with tight memory budget.

diff --git a/Engine/Plugins/Runtime/HairStrands/Source/HairStrandsCore/Private/GroomManager.cpp b/Engine/Plugins/Runtime/HairStrands/Source/HairStrandsCore/Private/GroomManager.cpp
--- a/Engine/Plugins/Runtime/HairStrands/Source/HairStrandsCore/Private/GroomManager.cpp
+++ b/Engine/Plugins/Runtime/HairStrands/Source/HairStrandsCore/Private/GroomManager.cpp
@@ -1682,12 +1682,14 @@ static void RunHairStrandsInterpolation_Strands(
 			if (InstanceData.bNeedRaytracing)
 			{
 				// Note: VFInput.Strands.Common.Radius already contains RadiusScale from discrete LOD setup, for backward compatibility.
-				const float CLODScale = InstanceData.Instance->HairGroupPublicData->ContinuousLODCoverageScale;
+				// Only used CLOD scale when the geometry is dynamic
+				const bool bIsDynamicGeometry = InstanceData.Instance->Strands.DeformedResource != nullptr;
+				const float CLODScale = bIsDynamicGeometry ? InstanceData.Instance->HairGroupPublicData->ContinuousLODCoverageScale : 1.f;
 				const float HairRadiusRT = InstanceData.Instance->HairGroupPublicData->VFInput.Strands.Common.RaytracingRadiusScale * InstanceData.Instance->HairGroupPublicData->VFInput.Strands.Common.Radius * CLODScale;
 				const float HairRootScaleRT = InstanceData.Instance->HairGroupPublicData->VFInput.Strands.Common.RootScale;
 				const float HairTipScaleRT = InstanceData.Instance->HairGroupPublicData->VFInput.Strands.Common.TipScale;
 
-				InstanceData.bNeedRaytracingUpdate = InstanceData.Instance->Strands.DeformedResource != nullptr ||
+				InstanceData.bNeedRaytracingUpdate = bIsDynamicGeometry ||
 					InstanceData.Instance->Strands.CachedHairScaledRadius != HairRadiusRT ||
 					InstanceData.Instance->Strands.CachedHairRootScale != HairRootScaleRT ||
 					InstanceData.Instance->Strands.CachedHairTipScale != HairTipScaleRT ||
