Fix groom per-CV attribute to be incorrect depending of the imported content.

This was due to an incorrect Log2 rounding of the CV chunk size, causing incorrect lookup at runtime.

diff --git a/Engine/Plugins/Runtime/HairStrands/Source/HairStrandsCore/Private/HairStrandsInterpolation.cpp b/Engine/Plugins/Runtime/HairStrands/Source/HairStrandsCore/Private/HairStrandsInterpolation.cpp
--- a/Engine/Plugins/Runtime/HairStrands/Source/HairStrandsCore/Private/HairStrandsInterpolation.cpp
+++ b/Engine/Plugins/Runtime/HairStrands/Source/HairStrandsCore/Private/HairStrandsInterpolation.cpp
@@ -95,12 +95,12 @@ inline uint32 ComputeGroupSize()
 void GetHairStrandsAttributeParameter(const FHairStrandsBulkData& In, FHairStrandsInstanceAttributeParameters& Out)
 {
 	check(FMath::IsPowerOfTwo(In.Header.Strides.CurveAttributeChunkElementCount));
-	Out.CurveAttributeIndexToChunkDivAsShift	= FMath::FloorToInt(FMath::Log2(float(In.Header.Strides.CurveAttributeChunkElementCount)));
+	Out.CurveAttributeIndexToChunkDivAsShift	= FPlatformMath::FloorLog2(In.Header.Strides.CurveAttributeChunkElementCount);
 	Out.CurveAttributeChunkElementCount			= In.Header.Strides.CurveAttributeChunkElementCount;
 	Out.CurveAttributeChunkStrideInBytes		= In.Header.Strides.CurveAttributeChunkStride;
 
 	check(FMath::IsPowerOfTwo(In.Header.Strides.PointAttributeChunkElementCount));
-	Out.PointAttributeIndexToChunkDivAsShift	= FMath::FloorToInt(FMath::Log2(float(In.Header.Strides.PointAttributeChunkElementCount)));
+	Out.PointAttributeIndexToChunkDivAsShift	= FPlatformMath::FloorLog2(In.Header.Strides.PointAttributeChunkElementCount);
 	Out.PointAttributeChunkElementCount     	= In.Header.Strides.PointAttributeChunkElementCount;
 	Out.PointAttributeChunkStrideInBytes    	= In.Header.Strides.PointAttributeChunkStride;
 
