[GroomBinding]
 - Fix the 'create binding' groom action menu doing nothing when async groom binding is active

#rnx

diff --git a/Engine/Plugins/Runtime/HairStrands/Source/HairStrandsEditor/Private/GroomActions.cpp b/Engine/Plugins/Runtime/HairStrands/Source/HairStrandsEditor/Private/GroomActions.cpp
--- a/Engine/Plugins/Runtime/HairStrands/Source/HairStrandsEditor/Private/GroomActions.cpp
+++ b/Engine/Plugins/Runtime/HairStrands/Source/HairStrandsEditor/Private/GroomActions.cpp
@@ -18,6 +18,7 @@
 #include "GroomCreateFollicleMaskOptionsWindow.h"
 #include "GroomCreateStrandsTexturesOptions.h"
 #include "GroomCreateStrandsTexturesOptionsWindow.h"
+#include "AssetCompilingManager.h"
 #include "HairStrandsImporter.h"
 #include "HairStrandsTranslator.h"
 #include "ToolMenuSection.h"
@@ -223,6 +224,9 @@ void ExecuteCreateBindingAsset(const FToolMenuContext& InContext)
 				if (BindingAsset)
 				{
 					BindingAsset->Build();
+#if WITH_EDITOR
+					FAssetCompilingManager::Get().FinishCompilationForObjects({BindingAsset});
+#endif
 					if (BindingAsset->IsValid())
 					{
 						TArray<UObject*> CreatedObjects;
