Fix another occurance of race condition in MeshCardBuilder due to non thread-safe MeshDescription access.

diff --git a/Engine/Plugins/Runtime/HairStrands/Source/HairStrandsCore/Private/HairCardsBuilder.cpp b/Engine/Plugins/Runtime/HairStrands/Source/HairStrandsCore/Private/HairCardsBuilder.cpp
--- a/Engine/Plugins/Runtime/HairStrands/Source/HairStrandsCore/Private/HairCardsBuilder.cpp
+++ b/Engine/Plugins/Runtime/HairStrands/Source/HairStrandsCore/Private/HairCardsBuilder.cpp
@@ -878,8 +878,12 @@ static bool InternalImportGeometry_WithImportedGuides(
 	FHairStrandsDatas& OutGuides,
 	FHairCardsInterpolationBulkData& OutInterpolationBulkData)
 {
-	// Note: if there are multiple section we only import the first one. Support for multiple section could be added later on. 
-	FMeshDescription* MeshDescription = StaticMesh->GetMeshDescription(0);
+	// Note: * if there are multiple section we only import the first one. Support for multiple section could be added later on. 
+	//       * use a local copy of the mesh description as multiple card assets referencing the same mesh description could be 
+	//         built in parallel. We later call SanitizeMeshDescription which modify the mesh description and is not thread safe
+	FMeshDescription LocalMeshDescription;
+	StaticMesh->CloneMeshDescription(0, LocalMeshDescription);
+	FMeshDescription* MeshDescription = &LocalMeshDescription;
 	const uint32 PointCount = MeshDescription->Vertices().Num();
 	const uint32 IndexCount  = MeshDescription->Triangles().Num() * 3;
 
