Fix hair strands binding when the skel. mesh does not have surface normal available on GPU.

diff --git a/Engine/Plugins/Runtime/HairStrands/Source/HairStrandsCore/Private/HairStrandsMeshProjection.cpp b/Engine/Plugins/Runtime/HairStrands/Source/HairStrandsCore/Private/HairStrandsMeshProjection.cpp
--- a/Engine/Plugins/Runtime/HairStrands/Source/HairStrandsCore/Private/HairStrandsMeshProjection.cpp
+++ b/Engine/Plugins/Runtime/HairStrands/Source/HairStrandsCore/Private/HairStrandsMeshProjection.cpp
@@ -301,19 +301,22 @@ static bool AddHairStrandUpdateMeshTrianglesPass(
 		CommonParameters.RDGMeshPreviousPositionBuffer = CommonParameters.RDGMeshPositionBuffer;
 	}
 
+	// On some platform the TangentSRV can be null as it is not create by default (requires CPU access flags).
+	// In such a case we bind the index buffer as a dummy data, and we use face normal.
+	if (CommonParameters.MeshTangentBuffer == nullptr)
+	{
+		CommonParameters.MeshTangentBuffer = CommonParameters.MeshIndexBuffer;
+		for (FSectionData& SectionData : SectionDatas)
+		{
+			SectionData.bUseFaceNormal = 1u;
+		}
+	}
+
 	FRDGBufferRef SectionBuffer = CreateStructuredBuffer(GraphBuilder, TEXT("Hair.SkelMeshSectionBuffer"), sizeof(FSectionData),  SectionDatas.Num(), SectionDatas.GetData(), sizeof(FSectionData) * SectionDatas.Num());
 	CommonParameters.MeshSectionBuffer = GraphBuilder.CreateSRV(SectionBuffer);
 
 	const bool bComputePreviousDeformedPosition = OutputPrevUAV != nullptr;
 	{
-		// On some platform the TangentSRV can be null as it is not create by default (requires CPU access flags).
-		// In such a case we bind the index buffer as a dummy data, and we use face normal.
-		if (CommonParameters.MeshTangentBuffer == nullptr)
-		{
-			CommonParameters.MeshTangentBuffer = CommonParameters.MeshIndexBuffer;
-			Default.bUseFaceNormal = 1u;
-		}
-
 		FHairUpdateMeshTriangleCS::FParameters* PassParameters = GraphBuilder.AllocParameters<FHairUpdateMeshTriangleCS::FParameters>();
 		*PassParameters = CommonParameters;
 
