Fix RT geometry initializer being reset prior to release the RT geometry.

The initializer value are checked during the destructor causing some assert.

diff --git a/Engine/Plugins/Runtime/HairStrands/Source/HairStrandsCore/Private/GroomResources.cpp b/Engine/Plugins/Runtime/HairStrands/Source/HairStrandsCore/Private/GroomResources.cpp
--- a/Engine/Plugins/Runtime/HairStrands/Source/HairStrandsCore/Private/GroomResources.cpp
+++ b/Engine/Plugins/Runtime/HairStrands/Source/HairStrandsCore/Private/GroomResources.cpp
@@ -1908,8 +1908,8 @@ void FHairStrandsRaytracingResource::InternalAllocate(FRDGBuilder& GraphBuilder,
 		InternalCreateStructuredBufferRDG<FHairStrandsIndexFormat>(GraphBuilder, MaxIndexCount, IndexBuffer, ToHairResourceDebugName(TEXT("Hair.StrandsRaytracing_IndexBuffer"), ResourceName), OwnerName, EHairResourceUsageType::Dynamic);
 
 		// Force reallocating and rebuilding the RT geometry
-		RayTracingGeometry.Initializer = {};
 		RayTracingGeometry.ReleaseRHI();
+		RayTracingGeometry.SetInitializer(FRayTracingGeometryInitializer{});
 		bIsRTGeometryInitialized = false;
 	}
 }
@@ -1918,8 +1918,8 @@ void FHairStrandsRaytracingResource::InternalRelease()
 {
 	PositionBuffer.Release();
 	IndexBuffer.Release();
-	RayTracingGeometry.Initializer = {};
 	RayTracingGeometry.ReleaseRHI();
+	RayTracingGeometry.SetInitializer(FRayTracingGeometryInitializer{});
 	bIsRTGeometryInitialized = false;
 }
 #endif
