Fixed a couple of crashes and asserts related to async groom compilation

diff --git a/Engine/Plugins/Runtime/HairStrands/Source/HairStrandsCore/Private/GroomComponent.cpp b/Engine/Plugins/Runtime/HairStrands/Source/HairStrandsCore/Private/GroomComponent.cpp
--- a/Engine/Plugins/Runtime/HairStrands/Source/HairStrandsCore/Private/GroomComponent.cpp
+++ b/Engine/Plugins/Runtime/HairStrands/Source/HairStrandsCore/Private/GroomComponent.cpp
@@ -1991,6 +1991,11 @@ FPrimitiveSceneProxy* UGroomComponent::CreateSceneProxy()
 		return nullptr;
 	}
 
+	if (GroomAsset->GetHairGroupsPlatformData().Num() != HairGroupInstances.Num())
+	{
+		return nullptr;
+	}
+
 	bool bIsValid = false;
 	for (const TRefCountPtr<FHairGroupInstance>& Instance : HairGroupInstances)
 	{
diff --git a/Engine/Plugins/Runtime/HairStrands/Source/HairStrandsCore/Private/Niagara/NiagaraDataInterfaceHairStrands.cpp b/Engine/Plugins/Runtime/HairStrands/Source/HairStrandsCore/Private/Niagara/NiagaraDataInterfaceHairStrands.cpp
--- a/Engine/Plugins/Runtime/HairStrands/Source/HairStrandsCore/Private/Niagara/NiagaraDataInterfaceHairStrands.cpp
+++ b/Engine/Plugins/Runtime/HairStrands/Source/HairStrandsCore/Private/Niagara/NiagaraDataInterfaceHairStrands.cpp
@@ -3355,7 +3355,7 @@ void UNiagaraDataInterfaceHairStrands::SetShaderParameters(const FNiagaraDataInt
 	// TODO: refactor all of that with if condition
 	const int32 MeshLODIndex = ProxyData && ProxyData->HairGroupInstance && ProxyData->HairGroupInstance->HairGroupPublicData ? ProxyData->HairGroupInstance->HairGroupPublicData->GetMeshLODIndex() : -1;
 	const bool bIsHairValid = ProxyData != nullptr && ProxyData->HairStrandsBuffer && ProxyData->HairStrandsBuffer->IsInitialized();
-	const bool bIsHairGroupInstValid = ProxyData != nullptr && ProxyData->HairGroupInstSource != nullptr && ProxyData->HairGroupInstSource->ContainsGroupInstance(ProxyData->HairGroupInstance);
+	const bool bIsHairGroupInstValid = ProxyData != nullptr && ProxyData->HairGroupInstSource != nullptr && ProxyData->HairGroupInstance != nullptr && ProxyData->HairGroupInstSource->ContainsGroupInstance(ProxyData->HairGroupInstance);
 	const bool bHasSkinningBinding = bIsHairValid && bIsHairGroupInstValid && ProxyData->HairGroupInstance && ProxyData->HairGroupInstance->BindingType == EHairBindingType::Skinning;
 	const bool bIsRootValid = bIsHairValid && ProxyData->HairStrandsBuffer->SourceDeformedRootResources&& ProxyData->HairStrandsBuffer->SourceDeformedRootResources->IsValid(MeshLODIndex) && bHasSkinningBinding;
 	const bool bIsRestValid = bIsHairValid && ProxyData->HairStrandsBuffer->SourceRestResources && ProxyData->HairStrandsBuffer->SourceRestResources->IsInitialized()&&
