Fix crash in groom deformer when changing geometry type.

diff --git a/Engine/Plugins/Runtime/HairStrands/Source/HairStrandsDeformer/Private/DeformerDataInterfaceGroomExec.cpp b/Engine/Plugins/Runtime/HairStrands/Source/HairStrandsDeformer/Private/DeformerDataInterfaceGroomExec.cpp
--- a/Engine/Plugins/Runtime/HairStrands/Source/HairStrandsDeformer/Private/DeformerDataInterfaceGroomExec.cpp
+++ b/Engine/Plugins/Runtime/HairStrands/Source/HairStrandsDeformer/Private/DeformerDataInterfaceGroomExec.cpp
@@ -108,22 +108,37 @@ FOptimusGroomExecDataProviderProxy::FOptimusGroomExecDataProviderProxy(UGroomCom
 {
 	GroomComponent = InGroomComponent;
 	Domain = InDomain;
+	Instances.SetNum(0);
+
+	if (GroomComponent)
+	{
+		const uint32 InstanceCount = GroomComponent->GetGroupCount();
+		Instances.Reserve(InstanceCount);
+		for (uint32 InstanceIt = 0; InstanceIt < InstanceCount; ++InstanceIt)
+		{		
+			if (FHairGroupInstance* Instance = GroomComponent->GetGroupInstance(InstanceIt))
+			{
+				Instances.Add(Instance);
+			}
+		}
+	}
 }
 
 int32 FOptimusGroomExecDataProviderProxy::GetDispatchThreadCount(TArray<FIntVector>& ThreadCounts) const
 {
-	const int32 NumInvocations = GroomComponent ? GroomComponent->GetGroupCount() : 0;
-
 	ThreadCounts.Reset();
-	ThreadCounts.Reserve(NumInvocations);
-	for (int32 InvocationIndex = 0; InvocationIndex < NumInvocations; ++InvocationIndex)
-	{		
-		if (FHairGroupInstance* Instance = GroomComponent->GetGroupInstance(InvocationIndex))
+	ThreadCounts.Reserve(Instances.Num());
+
+	uint32 NumInvocations = 0;
+	for (const FHairGroupInstance* Instance : Instances)
+	{	
+		if (Instance && Instance->Strands.IsValid())
 		{
 			const int32 NumControlPoints = Instance->Strands.GetData().GetNumPoints();
 			const int32 NumCurves = Instance->Strands.GetData().GetNumCurves();
 			const int32 NumThreads = Domain == EOptimusGroomExecDomain::ControlPoint ? NumControlPoints : NumCurves;
 			ThreadCounts.Add(FIntVector(NumThreads, 1, 1));
+			++NumInvocations;
 		}
 	}
 	return NumInvocations;
@@ -135,7 +150,7 @@ bool FOptimusGroomExecDataProviderProxy::IsValid(FValidationData const& InValida
 	{
 		return false;
 	}
-	if (GroomComponent == nullptr || GroomComponent->GetGroupCount() == 0)
+	if (GroomComponent == nullptr || Instances.Num() == 0)
 	{
 		return false;
 	}
@@ -148,7 +163,7 @@ void FOptimusGroomExecDataProviderProxy::GatherDispatchData(FDispatchData const&
 	const TStridedView<FParameters> ParameterArray = MakeStridedParameterView<FParameters>(InDispatchData);
 	for (int32 InvocationIndex = 0; InvocationIndex < ParameterArray.Num(); ++InvocationIndex)
 	{
-		FHairGroupInstance* Instance = GroomComponent->GetGroupInstance(InvocationIndex);
+		FHairGroupInstance* Instance = Instances[InvocationIndex];
 		const int32 NumControlPoints = Instance->Strands.GetData().GetNumPoints();
 		const int32 NumCurves = Instance->Strands.GetData().GetNumCurves();
 		const int32 NumThreads = Domain == EOptimusGroomExecDomain::ControlPoint ? NumControlPoints : NumCurves;
diff --git a/Engine/Plugins/Runtime/HairStrands/Source/HairStrandsDeformer/Private/DeformerDataInterfaceGroomExec.h b/Engine/Plugins/Runtime/HairStrands/Source/HairStrandsDeformer/Private/DeformerDataInterfaceGroomExec.h
--- a/Engine/Plugins/Runtime/HairStrands/Source/HairStrandsDeformer/Private/DeformerDataInterfaceGroomExec.h
+++ b/Engine/Plugins/Runtime/HairStrands/Source/HairStrandsDeformer/Private/DeformerDataInterfaceGroomExec.h
@@ -9,6 +9,7 @@
 
 class FGroomExecDataInterfaceParameters;
 class UGroomComponent;
+struct FHairGroupInstance;
 
 UENUM()
 enum class EOptimusGroomExecDomain : uint8
@@ -92,5 +93,6 @@ private:
 	using FParameters = FGroomExecDataInterfaceParameters;
 
 	UGroomComponent* GroomComponent = nullptr;
+	TArray<FHairGroupInstance*> Instances;
 	EOptimusGroomExecDomain Domain = EOptimusGroomExecDomain::ControlPoint;
 };
