Fix crash in niagara hair strands code.

When groom instances are reloaded (attachment event, change of properties, ...), it can cause crashes due to using stall resources. The instances data are recreated by the groom component, and not updated on the Niagara side. This CL postpone the fetching of these resources to the last moment (SetParameter) to ensure the data are correct by the time we bind them.

#test initial test with a few groom repro project, but probably needs further stability test.

diff --git a/Engine/Plugins/Runtime/HairStrands/Source/HairStrandsCore/Private/Niagara/NiagaraDataInterfaceHairStrands.cpp b/Engine/Plugins/Runtime/HairStrands/Source/HairStrandsCore/Private/Niagara/NiagaraDataInterfaceHairStrands.cpp
--- a/Engine/Plugins/Runtime/HairStrands/Source/HairStrandsCore/Private/Niagara/NiagaraDataInterfaceHairStrands.cpp
+++ b/Engine/Plugins/Runtime/HairStrands/Source/HairStrandsCore/Private/Niagara/NiagaraDataInterfaceHairStrands.cpp
@@ -248,40 +248,138 @@ static FAutoConsoleVariableRef CVarHairSimulationMaxDelay(TEXT("r.HairStrands.Si
 static int32 GHairSimulationRestUpdate = false;
 static FAutoConsoleVariableRef CVarHairSimulationRestUpdate(TEXT("r.HairStrands.SimulationRestUpdate"), GHairSimulationRestUpdate, TEXT("Update the simulation rest pose"));
 
+//------------------------------------------------------------------------------------------------------------
+// Helper struct/function to extracting used resource prior to dispatching
+
+struct FHairGroupInstanceRDG
+{
+	bool bIsValid = false;
+	bool bIsRootValid = false;
+	bool bIsDeformedValid = false;
+	int32 MeshLODIndex = -1;
+	EHairGeometryType GeometryType = EHairGeometryType::NoneGeometry;
+	EHairBindingType BindingType = EHairBindingType::NoneBinding;
+	uint32 NumPoints = 0;
+	uint32 NumCurves = 0;
+	uint32 SampleCount = 0;
+
+	// Position / Curve
+	FVector3f RestPositionOffsetValue = FVector3f::ZeroVector;
+	FRDGBufferSRVRef RestPositionBuffer = nullptr;
+	FRDGBufferUAVRef DeformedPositionBufferUAV = nullptr;
+	FRDGBufferSRVRef DeformedPositionOffset = nullptr;
+	FRDGBufferSRVRef CurvesOffsetsBuffer = nullptr;
+	FRDGBufferSRVRef PointToCurveIndexBuffer = nullptr;
+
+	// Skinning
+	FRDGBufferSRVRef RestTrianglePositionBuffer = nullptr;
+	FRDGBufferSRVRef DeformedTrianglePositionBuffer = nullptr;
+	FRDGBufferSRVRef RootBarycentricCoordinatesBuffer = nullptr;
+	FRDGBufferSRVRef RootToUniqueTriangleIndexBuffer = nullptr;
+
+	// RBF
+	FRDGBufferSRVRef RestSamplePositionsBuffer = nullptr;
+	FRDGBufferSRVRef MeshSampleWeightsBuffer = nullptr;
+
+	bool IsValid() const 		{ return bIsValid && GeometryType != EHairGeometryType::NoneGeometry; }
+	bool IsDeformedValid() const{ return bIsDeformedValid && GeometryType != EHairGeometryType::NoneGeometry; }
+	bool IsRootValid() const	{ return bIsRootValid && IsValid(); }
+};
+
+static FHairGroupInstanceRDG Convert(FRDGBuilder& GraphBuilder, const FHairGroupInstance* In)
+{
+	FHairGroupInstanceRDG Out;
+	if (In && In->Guides.IsValid())
+	{
+		const int32 MeshLODIndex = In->HairGroupPublicData ? In->HairGroupPublicData->GetMeshLODIndex() : -1;
+
+		Out.MeshLODIndex 					= MeshLODIndex;
+		Out.GeometryType 					= In->GeometryType;
+		Out.BindingType						= In->BindingType;
+
+		if (FHairStrandsRestResource* RestResource = In->Guides.RestResource)
+		{
+			Out.NumPoints						= RestResource->GetPointCount();
+			Out.NumCurves						= RestResource->GetCurveCount();
+			Out.RestPositionOffsetValue			= (FVector3f)RestResource->GetPositionOffset();
+			Out.RestPositionBuffer				= RegisterAsSRV(GraphBuilder, RestResource->PositionBuffer);
+			Out.CurvesOffsetsBuffer				= RegisterAsSRV(GraphBuilder, RestResource->CurveBuffer);
+			Out.PointToCurveIndexBuffer			= RegisterAsSRV(GraphBuilder, RestResource->PointToCurveBuffer);
+		}
+
+		if (FHairStrandsDeformedResource* DeformedResource = In->Guides.DeformedResource)
+		{
+			Out.DeformedPositionBufferUAV	= RegisterAsUAV(GraphBuilder, DeformedResource->GetBuffer(FHairStrandsDeformedResource::EFrameType::Current));
+			Out.DeformedPositionOffset		= RegisterAsSRV(GraphBuilder, DeformedResource->GetPositionOffsetBuffer(FHairStrandsDeformedResource::EFrameType::Current));
+		}
+
+		const bool bHasSkinnedInterpolation = In->Guides.DeformedRootResource && In->Guides.DeformedRootResource->IsValid(MeshLODIndex);
+		if (bHasSkinnedInterpolation)
+		{
+			if (const FHairStrandsLODRestRootResource* RootResource = In->Guides.RestRootResource->GetLOD(MeshLODIndex))
+			{
+				Out.RestTrianglePositionBuffer		= RegisterAsSRV(GraphBuilder, RootResource->RestUniqueTrianglePositionBuffer);
+				Out.RootBarycentricCoordinatesBuffer= RegisterAsSRV(GraphBuilder, RootResource->RootBarycentricBuffer);
+				Out.RootToUniqueTriangleIndexBuffer	= RegisterAsSRV(GraphBuilder, RootResource->RootToUniqueTriangleIndexBuffer);
+				Out.RestSamplePositionsBuffer		= RegisterAsSRV(GraphBuilder, RootResource->RestSamplePositionsBuffer);
+				Out.SampleCount						= RootResource->SampleCount;
+			}
+
+			if (const FHairStrandsLODDeformedRootResource* RootResource = In->Guides.DeformedRootResource->GetLOD(MeshLODIndex))
+			{
+				Out.DeformedTrianglePositionBuffer	= RegisterAsSRV(GraphBuilder, RootResource->GetDeformedUniqueTrianglePositionBuffer(FHairStrandsLODDeformedRootResource::Current));
+				Out.MeshSampleWeightsBuffer			= RegisterAsSRV(GraphBuilder, RootResource->GetMeshSampleWeightsBuffer((FHairStrandsLODDeformedRootResource::Current)));
+				Out.bIsRootValid					= Out.BindingType == EHairBindingType::Skinning;
+			}
+		}
+	}
+
+	// Set Shader SRV
+	FRDGBufferSRVRef DummyStructuredBuffer   = GraphBuilder.CreateSRV(GSystemTextures.GetDefaultStructuredBuffer(GraphBuilder, 16u, FUintVector4(0,0,0,0)));
+	FRDGBufferSRVRef DummyByteAddressBuffer  = GraphBuilder.CreateSRV(GSystemTextures.GetDefaultByteAddressBuffer(GraphBuilder, 16u));
+	FRDGBufferSRVRef DummyVertexBuffer16byte = GraphBuilder.CreateSRV(GSystemTextures.GetDefaultBuffer(GraphBuilder, 16u), FHairStrandsMeshTrianglePositionFormat::Format);
+	FRDGBufferSRVRef DummyVertexBuffer4byte  = GraphBuilder.CreateSRV(GSystemTextures.GetDefaultBuffer(GraphBuilder, 4u), FHairStrandsRootBarycentricFormat::Format);
+
+	check(FHairStrandsMeshTrianglePositionFormat::SizeInByte == 16u);
+	check(FHairStrandsRootBarycentricFormat::SizeInByte == 4u);
+	check(FHairStrandsRootToUniqueTriangleIndexFormat::SizeInByte == 4u);
+
+	Out.bIsValid = Out.RestPositionBuffer && In->Guides.RestResource->IsInitialized();
+	Out.bIsDeformedValid = Out.DeformedPositionBufferUAV && In->Guides.DeformedResource->IsInitialized();
+
+	// Fallback
+	if (!Out.RestPositionBuffer) 				{ Out.RestPositionBuffer 				= DummyByteAddressBuffer; };
+	if (!Out.DeformedPositionBufferUAV)			{ Out.DeformedPositionBufferUAV 		= GraphBuilder.CreateUAV(GraphBuilder.CreateBuffer(FRDGBufferDesc::CreateByteAddressDesc(16u), TEXT("Niagara.Hair.DummyUAV")), ERDGUnorderedAccessViewFlags::SkipBarrier); };
+	if (!Out.DeformedPositionOffset)			{ Out.DeformedPositionOffset			= DummyStructuredBuffer; };
+	if (!Out.CurvesOffsetsBuffer)				{ Out.CurvesOffsetsBuffer 				= DummyByteAddressBuffer; };
+	if (!Out.RestTrianglePositionBuffer)		{ Out.RestTrianglePositionBuffer 		= DummyVertexBuffer16byte; };
+	if (!Out.DeformedTrianglePositionBuffer)	{ Out.DeformedTrianglePositionBuffer 	= DummyVertexBuffer16byte; };
+	if (!Out.RootBarycentricCoordinatesBuffer)	{ Out.RootBarycentricCoordinatesBuffer 	= DummyVertexBuffer4byte; };
+	if (!Out.RootToUniqueTriangleIndexBuffer)	{ Out.RootToUniqueTriangleIndexBuffer 	= DummyVertexBuffer4byte; };
+	if (!Out.RestSamplePositionsBuffer)			{ Out.RestSamplePositionsBuffer 		= DummyStructuredBuffer; };
+	if (!Out.MeshSampleWeightsBuffer)			{ Out.MeshSampleWeightsBuffer 			= DummyStructuredBuffer; };
+
+	return Out;
+}
+
 //------------------------------------------------------------------------------------------------------------
 
 void FNDIHairStrandsBuffer::Initialize(
-	const FHairStrandsRestResource* HairStrandsRestResource,
-	const FHairStrandsDeformedResource*  HairStrandsDeformedResource,
-	const FHairStrandsRestRootResource* HairStrandsRestRootResource,
-	const FHairStrandsDeformedRootResource* HairStrandsDeformedRootResource,
+	const FNDIHairStrandsInfo& In,
 	const TStaticArray<float, 32 * NumScales>& InParamsScale)
 {
-	SourceRestResources = HairStrandsRestResource;
-	SourceDeformedResources = HairStrandsDeformedResource;
-	SourceRestRootResources = HairStrandsRestRootResource;
-	SourceDeformedRootResources = HairStrandsDeformedRootResource;
+	InfoData = In;
 	ParamsScale = InParamsScale;
-	BoundingBoxOffsets = FIntVector4(0,1,2,3);
-
 	bValidGeometryType = false;
 }
 
-void FNDIHairStrandsBuffer::Update(
-	const FHairStrandsRestResource* HairStrandsRestResource,
-	const FHairStrandsDeformedResource* HairStrandsDeformedResource,
-	const FHairStrandsRestRootResource* HairStrandsRestRootResource,
-	const FHairStrandsDeformedRootResource* HairStrandsDeformedRootResource)
+void FNDIHairStrandsBuffer::Update()
 {
-	SourceRestResources = HairStrandsRestResource;
-	SourceDeformedResources = HairStrandsDeformedResource;
-	SourceRestRootResources = HairStrandsRestRootResource;
-	SourceDeformedRootResources = HairStrandsDeformedRootResource;
 }
 
 void FNDIHairStrandsBuffer::Transfer(FRDGBuilder& GraphBuilder, const TStaticArray<float, 32 * NumScales>& InParamsScale)
 {
-	if (SourceRestResources != nullptr && ParamsScaleBuffer.IsValid())
+	if (InfoData.IsValid() && ParamsScaleBuffer.IsValid())
 	{
 		const uint32 ScaleCount = 32 * NumScales;
 		const uint32 ScaleBytes = sizeof(float) * ScaleCount;
@@ -295,10 +393,9 @@ void FNDIHairStrandsBuffer::InitRHI(FRHICommandListBase&)
 {
 	//ReadbackBuffer = new FRHIGPUBufferReadback(TEXT("Hair.PositionOffsetBuffer"));
 	
-	if (SourceRestResources != nullptr)
+	if (InfoData.IsValid())
 	{
 		FRDGBuilder GraphBuilder(FRHICommandListImmediate::Get());
-		FHairStrandsBulkData* SourceDatas = &SourceRestResources->BulkData; // This could be released by that time depending on how the initialization order is
 		{
 			static const uint32 ZeroData[] = {
 				UINT_MAX,UINT_MAX,UINT_MAX,0,0,0,
@@ -320,14 +417,6 @@ void FNDIHairStrandsBuffer::InitRHI(FRHICommandListBase&)
 			GraphBuilder.QueueBufferUpload(ParamsScaleBuffer.GetOrCreateBuffer(GraphBuilder), ParamsScale.GetData(), ScaleBytes);
 			ParamsScaleBuffer.EndGraphUsage();
 		}
-		if (SourceDeformedResources == nullptr || (SourceDeformedResources && !SourceDeformedResources->IsInitialized()))
-		{
-			const uint32 PositionsCount = SourceDatas->GetNumPoints();
-
-			FRDGBufferDesc BufferDesc = FRDGBufferDesc::CreateByteAddressDesc(FHairStrandsPositionFormat::SizeInByte * PositionsCount);
-			FRDGBufferRef RDGBuffer = GraphBuilder.CreateBuffer(BufferDesc, TEXT("DeformedPositionBuffer"));
-			DeformedPositionBuffer = GraphBuilder.ConvertToExternalBuffer(RDGBuffer);
-		}
 		GraphBuilder.Execute();
 	}
 }
@@ -337,7 +426,6 @@ void FNDIHairStrandsBuffer::ReleaseRHI()
 	//delete ReadbackBuffer;
 	//ReadbackBuffer = nullptr;
 	
-	DeformedPositionBuffer.SafeRelease();
 	BoundingBoxBuffer.Release();
 	ParamsScaleBuffer.Release();
 }
@@ -373,31 +461,35 @@ void FNDIHairStrandsData::Release()
 	}
 }
 
-void FNDIHairStrandsData::Update(UNiagaraDataInterfaceHairStrands* Interface, FNiagaraSystemInstance* SystemInstance, const FHairStrandsBulkData* StrandsDatas,
-	UGroomAsset* GroomAsset, const int32 GroupIndex, const int32 LODIndex, const FTransform& LocalToWorld, const float DeltaSeconds)
+void FNDIHairStrandsData::Update(
+	UNiagaraDataInterfaceHairStrands* Interface, 
+	const FNDIHairStrandsInfo& InData,
+	const float DeltaSeconds)
 {
 	if (Interface != nullptr)
 	{
-		WorldTransform = LocalToWorld;
+		const UGroomAsset* GroomAsset = Interface->SourceComponent->GroomAsset;
+
+		WorldTransform = InData.LocalToWorld;
 
-		const bool bHasValidBindingAsset = (Interface->IsComponentValid() && Interface->SourceComponent->BindingAsset && Interface->SourceComponent->GroomAsset);
+		const bool bHasValidBindingAsset = (Interface->IsComponentValid() && Interface->SourceComponent->BindingAsset && GroomAsset);
 
-		GlobalInterpolation = bHasValidBindingAsset ? Interface->SourceComponent->GroomAsset->IsGlobalInterpolationEnable(GroupIndex, LODIndex) : false;
+		GlobalInterpolation = bHasValidBindingAsset ? GroomAsset->IsGlobalInterpolationEnable(InData.GroupIndex, InData.LODIndex) : false;
 		bSkinningTransfer = bHasValidBindingAsset ?
 			(Interface->SourceComponent->BindingAsset->GetSourceSkeletalMesh() && Interface->SourceComponent->BindingAsset->GetTargetSkeletalMesh() &&
 			 Interface->SourceComponent->BindingAsset->GetSourceSkeletalMesh() != Interface->SourceComponent->BindingAsset->GetTargetSkeletalMesh()) : false;
 		
 		TickingGroup = Interface->IsComponentValid() ? ComputeTickingGroup(Interface->SourceComponent) : NiagaraFirstTickGroup;
 
-		const bool bIsSimulationEnable = Interface->IsComponentValid() ? Interface->SourceComponent->IsSimulationEnable(GroupIndex, LODIndex) : 
-																		    GroomAsset ? GroomAsset->IsSimulationEnable(GroupIndex, LODIndex) : false;
+		const bool bIsSimulationEnable = Interface->IsComponentValid() ? Interface->SourceComponent->IsSimulationEnable(InData.GroupIndex, InData.LODIndex) : 
+			GroomAsset ? GroomAsset->IsSimulationEnable(InData.GroupIndex, InData.LODIndex) : false;
 
-		if (StrandsDatas != nullptr && GroomAsset != nullptr && GroupIndex >= 0 && GroupIndex < GroomAsset->GetHairGroupsPhysics().Num() && bIsSimulationEnable)
+		if (GroomAsset != nullptr && GroomAsset->GetHairGroupsPhysics().IsValidIndex(InData.GroupIndex) && bIsSimulationEnable)
 		{
-			FHairGroupsPhysics& HairPhysics = GroomAsset->GetHairGroupsPhysics()[GroupIndex];
+			const FHairGroupsPhysics& HairPhysics = GroomAsset->GetHairGroupsPhysics()[InData.GroupIndex];
 			StrandsSize = static_cast<uint8>(HairPhysics.StrandsParameters.StrandsSize);
 
-			HairGroupIndex = Interface->IsComponentValid() ? GroupIndex : -1;
+			HairGroupIndex = Interface->IsComponentValid() ? InData.GroupIndex : -1;
 			HairGroupInstSource = Interface->IsComponentValid() ? Interface->SourceComponent : nullptr;
 
 			SubSteps = HairPhysics.SolverSettings.SubSteps;
@@ -433,15 +525,13 @@ void FNDIHairStrandsData::Update(UNiagaraDataInterfaceHairStrands* Interface, FN
 			for (int32 i = 0; i < StrandsSize; ++i)
 			{
 				const float VertexCoord = static_cast<float>(i) / (StrandsSize-1.0);
-				ParamsScale[32 * BendOffset + i] = HairPhysics.MaterialConstraints.BendConstraint.BendScale.GetRichCurve()->Eval(VertexCoord);
-				ParamsScale[32 * StretchOffset + i] = HairPhysics.MaterialConstraints.StretchConstraint.StretchScale.GetRichCurve()->Eval(VertexCoord);
-				ParamsScale[32 * RadiusOffset + i] = HairPhysics.MaterialConstraints.CollisionConstraint.RadiusScale.GetRichCurve()->Eval(VertexCoord);
-				ParamsScale[32 * ThicknessOffset + i] = HairPhysics.StrandsParameters.ThicknessScale.GetRichCurve()->Eval(VertexCoord);
+				ParamsScale[32 * BendOffset + i] = HairPhysics.MaterialConstraints.BendConstraint.BendScale.GetRichCurveConst()->Eval(VertexCoord);
+				ParamsScale[32 * StretchOffset + i] = HairPhysics.MaterialConstraints.StretchConstraint.StretchScale.GetRichCurveConst()->Eval(VertexCoord);
+				ParamsScale[32 * RadiusOffset + i] = HairPhysics.MaterialConstraints.CollisionConstraint.RadiusScale.GetRichCurveConst()->Eval(VertexCoord);
+				ParamsScale[32 * ThicknessOffset + i] = HairPhysics.StrandsParameters.ThicknessScale.GetRichCurveConst()->Eval(VertexCoord);
 			}
 
-			const FBox& StrandsBox = StrandsDatas->Header.BoundingBox;
-
-			NumStrands = StrandsDatas->GetNumCurves();
+			NumStrands = InData.NumCurves;
 			LocalSimulation = false;
 			BoneTransform = FTransform::Identity;
 
@@ -522,21 +612,13 @@ bool FNDIHairStrandsData::Init(UNiagaraDataInterfaceHairStrands* Interface, FNia
 
 	if (Interface != nullptr)
 	{
-		FHairStrandsRestResource* StrandsRestResource = nullptr;
-		FHairStrandsDeformedResource* StrandsDeformedResource = nullptr;
-		FHairStrandsRestRootResource* StrandsRestRootResource = nullptr;
-		FHairStrandsDeformedRootResource* StrandsDeformedRootResource = nullptr;
-		UGroomAsset* GroomAsset = nullptr;
-		int32 GroupIndex = 0;
-		int32 LODIndex = 0;
-
 		{
-			FTransform LocalToWorld = FTransform::Identity;
-			Interface->ExtractDatasAndResources(SystemInstance, StrandsRestResource, StrandsDeformedResource, StrandsRestRootResource, StrandsDeformedRootResource, GroomAsset, GroupIndex, LODIndex, LocalToWorld);
-			Update(Interface, SystemInstance, StrandsRestResource ? &StrandsRestResource->BulkData : nullptr, GroomAsset, GroupIndex, LODIndex, LocalToWorld, 0.0f);
+			FNDIHairStrandsInfo InfoData;
+			Interface->ExtractDatasAndResources(SystemInstance, InfoData);
+			Update(Interface, InfoData, 0.0f);
 
 			HairStrandsBuffer = new FNDIHairStrandsBuffer();
-			HairStrandsBuffer->Initialize(StrandsRestResource, StrandsDeformedResource, StrandsRestRootResource, StrandsDeformedRootResource, ParamsScale);
+			HairStrandsBuffer->Initialize(InfoData, ParamsScale);
 
 			BeginInitResource(HairStrandsBuffer);
 
@@ -681,23 +763,16 @@ void UNiagaraDataInterfaceHairStrands::ExtractSourceComponent(FNiagaraSystemInst
 
 void UNiagaraDataInterfaceHairStrands::ExtractDatasAndResources(
 	FNiagaraSystemInstance* SystemInstance,
-	FHairStrandsRestResource*& OutStrandsRestResource,
-	FHairStrandsDeformedResource*& OutStrandsDeformedResource,
-	FHairStrandsRestRootResource*& OutStrandsRestRootResource,
-	FHairStrandsDeformedRootResource*& OutStrandsDeformedRootResource,
-	UGroomAsset*& OutGroomAsset,
-	int32& OutGroupIndex, 
-	int32& OutLODIndex,
-	FTransform& OutLocalToWorld)
+	FNDIHairStrandsInfo& Out)
 {
 	ExtractSourceComponent(SystemInstance);
 
-	OutStrandsRestResource = nullptr;
-	OutStrandsDeformedResource = nullptr;
-	OutStrandsRestRootResource = nullptr;
-	OutStrandsDeformedRootResource = nullptr;
-	OutGroupIndex = -1;
-	OutLODIndex = -1;
+	Out.NumControlPoints = 0;
+	Out.NumCurves = 0;
+	Out.GroupIndex = -1;
+	Out.LODIndex = -1;
+	Out.LocalToWorld = FTransform::Identity;
+	Out.bHasValidResources = false;
 
 	if (IsComponentValid() && SystemInstance)
 	{
@@ -709,33 +784,31 @@ void UNiagaraDataInterfaceHairStrands::ExtractDatasAndResources(
 				{
 					if (SystemInstanceController->GetSystemInstanceID() == SystemInstance->GetId())
 					{
-						OutGroupIndex = NiagaraIndex;
+						Out.GroupIndex = NiagaraIndex;
 						break;
 					}
 				}
 			}
 		}
-		if (OutGroupIndex >= 0 && OutGroupIndex < SourceComponent->NiagaraComponents.Num())
+		if (Out.GroupIndex >= 0 && Out.GroupIndex < SourceComponent->NiagaraComponents.Num())
 		{
-			OutStrandsRestResource = SourceComponent->GetGuideStrandsRestResource(OutGroupIndex);
-			OutStrandsDeformedResource = SourceComponent->GetGuideStrandsDeformedResource(OutGroupIndex);
-			OutStrandsRestRootResource = SourceComponent->GetGuideStrandsRestRootResource(OutGroupIndex);
-			OutStrandsDeformedRootResource = SourceComponent->GetGuideStrandsDeformedRootResource(OutGroupIndex);
-			OutGroomAsset = SourceComponent->GroomAsset;
-			OutLODIndex = SourceComponent->GetForcedLOD();
-			OutLocalToWorld = SourceComponent->GetComponentTransform();
+			Out.LODIndex = SourceComponent->GetForcedLOD();
+			Out.LocalToWorld = SourceComponent->GetComponentTransform();
+			if (SourceComponent->GroomAsset && SourceComponent->GroomAsset->GetHairGroupsPlatformData().IsValidIndex(Out.GroupIndex))
+			{
+				const FHairGroupPlatformData::FGuides& Guides = SourceComponent->GroomAsset->GetHairGroupsPlatformData()[Out.GroupIndex].Guides;
+				Out.NumControlPoints = Guides.BulkData.GetNumPoints();
+				Out.NumCurves = Guides.BulkData.GetNumCurves();
+				Out.bHasValidResources = SourceComponent->GetGuideStrandsRestResource(Out.GroupIndex) != nullptr;
+			}
 		}
 	}
 	else if (DefaultSource != nullptr)
 	{
-		OutGroupIndex = 0;
-		OutLODIndex = 0;
-		OutLocalToWorld = SystemInstance ? SystemInstance->GetWorldTransform() : FTransform::Identity;
-		if (OutGroupIndex < DefaultSource->GetNumHairGroups())
-		{
-			OutStrandsRestResource = DefaultSource->GetHairGroupsResources()[OutGroupIndex].Guides.RestResource;
-			OutGroomAsset = DefaultSource;
-		}
+		Out.NumControlPoints = 0;
+		Out.GroupIndex = 0;
+		Out.LODIndex = 0;
+		Out.LocalToWorld = SystemInstance ? SystemInstance->GetWorldTransform() : FTransform::Identity;
 	}
 }
 
@@ -777,20 +850,11 @@ void UNiagaraDataInterfaceHairStrands::DestroyPerInstanceData(void* PerInstanceD
 bool UNiagaraDataInterfaceHairStrands::PerInstanceTick(void* PerInstanceData, FNiagaraSystemInstance* SystemInstance, float InDeltaSeconds)
 {
 	FNDIHairStrandsData* InstanceData = static_cast<FNDIHairStrandsData*>(PerInstanceData);
-
-	FHairStrandsRestResource* StrandsRestResource = nullptr;
-	FHairStrandsDeformedResource* StrandsDeformedResource = nullptr;
-	FHairStrandsRestRootResource* StrandsRestRootResource = nullptr;
-	FHairStrandsDeformedRootResource* StrandsDeformedRootResource = nullptr;
-	UGroomAsset* GroomAsset = nullptr;
-	int32 GroupIndex = 0;
-	int32 LODIndex = 0;
-
 	InstanceData->TickCount = FMath::Min(GHairSimulationMaxDelay + 1, InstanceData->TickCount + 1);
 
-	FTransform LocalToWorld = FTransform::Identity;
-	ExtractDatasAndResources(SystemInstance, StrandsRestResource, StrandsDeformedResource, StrandsRestRootResource, StrandsDeformedRootResource, GroomAsset, GroupIndex, LODIndex, LocalToWorld);
-	InstanceData->HairStrandsBuffer->Update(StrandsRestResource, StrandsDeformedResource, StrandsRestRootResource, StrandsDeformedRootResource);
+	FNDIHairStrandsInfo InfoData;
+	ExtractDatasAndResources(SystemInstance, InfoData);
+	InstanceData->HairStrandsBuffer->Update();
 
 	if (SourceComponent != nullptr)
 	{
@@ -801,7 +865,7 @@ bool UNiagaraDataInterfaceHairStrands::PerInstanceTick(void* PerInstanceData, FN
 		}
 		InstanceData->ForceReset = SourceComponent->bResetSimulation;
 	}
-	InstanceData->Update(this, SystemInstance, StrandsRestResource ? &StrandsRestResource->BulkData : nullptr, GroomAsset, GroupIndex, LODIndex, LocalToWorld, InDeltaSeconds);
+	InstanceData->Update(this, InfoData, InDeltaSeconds);
 	return false;
 }
 
@@ -885,21 +949,28 @@ FMatrix44f ComputeWorldTransform(const FNDIHairStrandsData* InstanceData)
 	return WorldTransformFloat;
 }
 
-static void InterpolateGroomGuides(FRDGBuilder& GraphBuilder, FNiagaraDataBuffer* ParticlesBuffer,
-	const uint32 NodePositionComponent,  const uint32 RestPositionComponent, FNDIHairStrandsBuffer* HairStrandsBuffer,
-	const uint32 StrandsSize, const bool bHasSkinningBinding, const bool bHasValidGeometry, const FMatrix44f& WorldToLocal, int32 MeshLODIndex)
+static void InterpolateGroomGuides(
+	FRDGBuilder& GraphBuilder, 
+	FNiagaraDataBuffer* ParticlesBuffer,
+	const uint32 NodePositionComponent,  
+	const uint32 RestPositionComponent, 
+	FNDIHairStrandsBuffer* HairStrandsBuffer,
+	FHairGroupInstance* HairGroupInstance,
+	const uint32 StrandsSize, 
+	const bool bHasSkinningBinding, 
+	const bool bHasValidGeometry, 
+	const FMatrix44f& WorldToLocal, 
+	int32 MeshLODIndex)
 {
+	FHairGroupInstanceRDG InstanceRDG = Convert(GraphBuilder, HairGroupInstance);
+
 	const bool bIsHairValid = HairStrandsBuffer && HairStrandsBuffer->IsInitialized();
-	const bool bIsRestValid = bIsHairValid && HairStrandsBuffer->SourceRestResources && HairStrandsBuffer->SourceRestResources->IsInitialized();
-	const bool bIsDeformedValid = bIsHairValid && HairStrandsBuffer->SourceDeformedResources && HairStrandsBuffer->SourceDeformedResources->IsInitialized();
+	const bool bIsRestValid = bIsHairValid && InstanceRDG.IsValid();
+	const bool bIsDeformedValid = bIsHairValid && InstanceRDG.IsDeformedValid();
+	const bool bIsRootValid = bIsHairValid && InstanceRDG.IsRootValid();
 
 	if(bIsRestValid && bIsDeformedValid && bHasValidGeometry)
 	{
-		bool bIsRootValid = HairStrandsBuffer->SourceDeformedRootResources && HairStrandsBuffer->SourceDeformedRootResources->IsValid(MeshLODIndex) &&
-							HairStrandsBuffer->SourceRestRootResources && HairStrandsBuffer->SourceRestRootResources->IsValid(MeshLODIndex) && bHasSkinningBinding;
-
-		const uint32 NumPoints = HairStrandsBuffer->SourceRestResources->GetPointCount();
-
 		FInterpolateGroomGuidesCS::FPermutationDomain InterpolationDomain;
 		InterpolationDomain.Set<FInterpolateGroomGuidesCS::FInterpolationType>(!bIsRootValid);
 
@@ -907,7 +978,7 @@ static void InterpolateGroomGuides(FRDGBuilder& GraphBuilder, FNiagaraDataBuffer
 
 		FInterpolateGroomGuidesCS::FParameters* PassParameters = GraphBuilder.AllocParameters<FInterpolateGroomGuidesCS::FParameters>();
 		PassParameters->StrandsSize = StrandsSize;
-		PassParameters->NumPoints = NumPoints;
+		PassParameters->NumPoints = InstanceRDG.NumPoints;
 
 		PassParameters->NiagaraFloatBuffer = ParticlesBuffer->GetGPUBufferFloat().SRV;
 		PassParameters->NiagaraFloatStride = ParticlesBuffer->GetFloatStride() / sizeof(float);
@@ -915,28 +986,20 @@ static void InterpolateGroomGuides(FRDGBuilder& GraphBuilder, FNiagaraDataBuffer
 		PassParameters->RestPositionComponent = RestPositionComponent;
 		PassParameters->WorldToLocal = WorldToLocal;
 
-		const int32 BufferIndex = HairStrandsBuffer->SourceDeformedResources->GetIndex(FHairStrandsDeformedResource::Current);
-
-		PassParameters->DeformedPositionOffset = RegisterAsSRV(GraphBuilder, HairStrandsBuffer->SourceDeformedResources->DeformedOffsetBuffer[BufferIndex]);
-		PassParameters->RestPositionOffset = FVector3f(HairStrandsBuffer->SourceRestResources->GetPositionOffset());
-		PassParameters->DeformedPositionBuffer = RegisterAsUAV(GraphBuilder, HairStrandsBuffer->SourceDeformedResources->DeformedPositionBuffer[BufferIndex]);
-		PassParameters->RestPositionBuffer = RegisterAsSRV(GraphBuilder, HairStrandsBuffer->SourceRestResources->PositionBuffer);
+		PassParameters->DeformedPositionOffset = InstanceRDG.DeformedPositionOffset;
+		PassParameters->RestPositionOffset = InstanceRDG.RestPositionOffsetValue;
+		PassParameters->DeformedPositionBuffer = InstanceRDG.DeformedPositionBufferUAV;
+		PassParameters->RestPositionBuffer = InstanceRDG.RestPositionBuffer;
+		
+		PassParameters->PointToCurveIndexBuffer = InstanceRDG.PointToCurveIndexBuffer;
+		PassParameters->CurvesOffsetsBuffer = InstanceRDG.CurvesOffsetsBuffer;
+		PassParameters->RestTrianglePositionBuffer = InstanceRDG.RestTrianglePositionBuffer;
+		PassParameters->DeformedTrianglePositionBuffer = InstanceRDG.DeformedTrianglePositionBuffer;
+		PassParameters->RootBarycentricCoordinatesBuffer = InstanceRDG.RootBarycentricCoordinatesBuffer;
+		PassParameters->RootToUniqueTriangleIndexBuffer = InstanceRDG.RootToUniqueTriangleIndexBuffer;
 		
-		PassParameters->PointToCurveIndexBuffer = RegisterAsSRV(GraphBuilder, HairStrandsBuffer->SourceRestResources->PointToCurveBuffer);
-		PassParameters->CurvesOffsetsBuffer = RegisterAsSRV(GraphBuilder, HairStrandsBuffer->SourceRestResources->CurveBuffer);
-		if(bIsRootValid)
-		{
-			const FHairStrandsLODRestRootResource* RestMeshProjection = HairStrandsBuffer->SourceRestRootResources->GetLOD(MeshLODIndex) ;
-			const FHairStrandsLODDeformedRootResource* DeformedMeshProjection = HairStrandsBuffer->SourceDeformedRootResources->GetLOD(MeshLODIndex);
-
-			PassParameters->RestTrianglePositionBuffer = RegisterAsSRV(GraphBuilder, RestMeshProjection->RestUniqueTrianglePositionBuffer);
-			PassParameters->DeformedTrianglePositionBuffer = RegisterAsSRV(GraphBuilder, DeformedMeshProjection->GetDeformedUniqueTrianglePositionBuffer(FHairStrandsLODDeformedRootResource::Current));
-			PassParameters->RootBarycentricCoordinatesBuffer = RegisterAsSRV(GraphBuilder, RestMeshProjection->RootBarycentricBuffer);
-			PassParameters->RootToUniqueTriangleIndexBuffer = RegisterAsSRV(GraphBuilder, RestMeshProjection->RootToUniqueTriangleIndexBuffer);
-		}
-
 		const uint32 GroupSize = NIAGARA_HAIR_STRANDS_THREAD_COUNT_INTERPOLATE;
-		const uint32 DispatchCount = FMath::DivideAndRoundUp(NumPoints, GroupSize);
+		const uint32 DispatchCount = FMath::DivideAndRoundUp(InstanceRDG.NumPoints, GroupSize);
 
 		FComputeShaderUtils::AddPass(
 			GraphBuilder,
@@ -989,7 +1052,7 @@ void UNiagaraDataInterfaceHairStrands::SimCachePostReadFrame(void* OptionalPerIn
 						{
 							FMemMark MemMark(FMemStack::Get());
 							FRDGBuilder GraphBuilder(RHICmdList);
-							InterpolateGroomGuides(GraphBuilder, ParticlesBuffer, NodePositionComponent, RestPositionComponent, HairstrandsBuffer, StrandsSize, bHasSkinningBinding, bHasValidGeometry, WorldToLocal, LocalHairInstance->Debug.MeshLODIndex);
+							InterpolateGroomGuides(GraphBuilder, ParticlesBuffer, NodePositionComponent, RestPositionComponent, HairstrandsBuffer, LocalHairInstance, StrandsSize, bHasSkinningBinding, bHasValidGeometry, WorldToLocal, LocalHairInstance->Debug.MeshLODIndex);
 							GraphBuilder.Execute();
 						}
 					);
@@ -3352,74 +3415,26 @@ void UNiagaraDataInterfaceHairStrands::SetShaderParameters(const FNiagaraDataInt
 	FNDIHairStrandsProxy& DIProxy = Context.GetProxy<FNDIHairStrandsProxy>();
 	FNDIHairStrandsData* ProxyData = DIProxy.SystemInstancesToProxyData.Find(Context.GetSystemInstanceID());
 
-	// TODO: refactor all of that with if condition
-	const int32 MeshLODIndex = ProxyData && ProxyData->HairGroupInstance && ProxyData->HairGroupInstance->HairGroupPublicData ? ProxyData->HairGroupInstance->HairGroupPublicData->GetMeshLODIndex() : -1;
+	FHairGroupInstanceRDG InstanceRDG = Convert(GraphBuilder, ProxyData->HairGroupInstance);
+	
+	const int32 MeshLODIndex = InstanceRDG.MeshLODIndex;
 	const bool bIsHairValid = ProxyData != nullptr && ProxyData->HairStrandsBuffer && ProxyData->HairStrandsBuffer->IsInitialized();
-	const bool bIsHairGroupInstValid = ProxyData != nullptr && ProxyData->HairGroupInstSource != nullptr && ProxyData->HairGroupInstance != nullptr && ProxyData->HairGroupInstSource->ContainsGroupInstance(ProxyData->HairGroupInstance);
-	const bool bHasSkinningBinding = bIsHairValid && bIsHairGroupInstValid && ProxyData->HairGroupInstance && ProxyData->HairGroupInstance->BindingType == EHairBindingType::Skinning;
-	const bool bIsRootValid = bIsHairValid && ProxyData->HairStrandsBuffer->SourceDeformedRootResources&& ProxyData->HairStrandsBuffer->SourceDeformedRootResources->IsValid(MeshLODIndex) && bHasSkinningBinding;
-	const bool bIsRestValid = bIsHairValid && ProxyData->HairStrandsBuffer->SourceRestResources && ProxyData->HairStrandsBuffer->SourceRestResources->IsInitialized()&&
+	const bool bIsRestValid = bIsHairValid && InstanceRDG.IsValid() &&
 		// TEMP: These check are only temporary for avoiding crashes while we find the bottom of the issue.
 		ProxyData->HairStrandsBuffer->ParamsScaleBuffer.IsValid() && ProxyData->HairStrandsBuffer->BoundingBoxBuffer.IsValid();
 
-	const bool bIsGeometryValid = bIsHairValid && (!bIsHairGroupInstValid || (bIsHairGroupInstValid && (ProxyData->HairGroupInstance && ProxyData->HairGroupInstance->GeometryType != EHairGeometryType::NoneGeometry)));
-	const bool bIsDeformedValid = bIsHairValid && ProxyData->HairStrandsBuffer->SourceDeformedResources && ProxyData->HairStrandsBuffer->SourceDeformedResources->IsInitialized();
-
 	NDIHairStrandsLocal::FShaderParameters* ShaderParameters = Context.GetParameterNestedStruct<NDIHairStrandsLocal::FShaderParameters>();
-	if (bIsHairValid && bIsRestValid && bIsGeometryValid && bIsHairGroupInstValid)
+	if (bIsHairValid && bIsRestValid)
 	{
 		check(ProxyData);
 
-		FRDGBufferSRVRef DummyStructuredBufferSRV = GraphBuilder.CreateSRV(GSystemTextures.GetDefaultStructuredBuffer(GraphBuilder, 16u, FUintVector4(0,0,0,0)));
-
 		FNDIHairStrandsBuffer* HairStrandsBuffer = ProxyData->HairStrandsBuffer;
 
-		FRDGBufferUAVRef DeformedPositionBufferUAV = nullptr;
-		FRDGBufferSRVRef DeformedPositionOffsetSRV = nullptr;
-		if (bIsDeformedValid)
-		{
-			FRDGPooledBuffer* DeformedPositionBuffer = HairStrandsBuffer->SourceDeformedResources->DeformedPositionBuffer[HairStrandsBuffer->SourceDeformedResources->CurrentIndex].Buffer;
-			DeformedPositionBufferUAV = GraphBuilder.CreateUAV(GraphBuilder.RegisterExternalBuffer(DeformedPositionBuffer));	//-OPT: If constant across all dispatches can create in PreStage once
-			DeformedPositionOffsetSRV = RegisterAsSRV(GraphBuilder, HairStrandsBuffer->SourceDeformedResources->DeformedOffsetBuffer[HairStrandsBuffer->SourceDeformedResources->CurrentIndex]);
-
-			// Keep the readback commented out for future debug
-			// FRDGBufferRef Buffer = GraphBuilder.RegisterExternalBuffer(HairStrandsBuffer->SourceDeformedResources->DeformedOffsetBuffer[HairStrandsBuffer->SourceDeformedResources->CurrentIndex].Buffer);
-			// AddEnqueueCopyPass(GraphBuilder, HairStrandsBuffer->ReadbackBuffer, Buffer, 16u);
-			//
-			// if (HairStrandsBuffer->ReadbackBuffer->IsReady())
-			// {
-			// 	const FVector4f ReadData = *(FVector4f*)(HairStrandsBuffer->ReadbackBuffer->Lock(sizeof(uint32)));
-			// 	HairStrandsBuffer->ReadbackBuffer->Unlock();
-			// 	UE_LOG(LogHairStrands, Log, TEXT("Deformed Offset = %f %f %f"), ReadData.X, ReadData.Y, ReadData.Z);
-			// }
-		}
-		else
-		{
-			DeformedPositionBufferUAV = GraphBuilder.CreateUAV(GraphBuilder.RegisterExternalBuffer(HairStrandsBuffer->DeformedPositionBuffer));
-			DeformedPositionOffsetSRV = DummyStructuredBufferSRV;
-		}
-
 		// Projection Buffers
-		const bool bHasSkinnedInterpolation = (bIsRootValid && HairStrandsBuffer->SourceDeformedRootResources->IsValid(MeshLODIndex));
+		const bool bHasSkinnedInterpolation = InstanceRDG.IsRootValid();
 		const EHairSimulationInterpolationMode InterpolationModeValue = bHasSkinnedInterpolation ? (ProxyData->GlobalInterpolation ?
 			EHairSimulationInterpolationMode::RBF : EHairSimulationInterpolationMode::Skinned) : EHairSimulationInterpolationMode::Rigid;
 
-		const FHairStrandsLODRestRootResource* RestMeshProjection = bHasSkinnedInterpolation ? HairStrandsBuffer->SourceRestRootResources->GetLOD(MeshLODIndex) : nullptr;
-		const FHairStrandsLODDeformedRootResource* DeformedMeshProjection = bHasSkinnedInterpolation ? HairStrandsBuffer->SourceDeformedRootResources->GetLOD(MeshLODIndex) : nullptr;
-
-		FRDGBufferSRVRef RestTrianglePositionSRV = (bHasSkinnedInterpolation && RestMeshProjection) ? RegisterAsSRV(GraphBuilder, RestMeshProjection->RestUniqueTrianglePositionBuffer) : Context.GetComputeDispatchInterface().GetEmptyBufferSRV(GraphBuilder, FHairStrandsMeshTrianglePositionFormat::Format);
-		FRDGBufferSRVRef DeformedTrianglePositionSRV = (bHasSkinnedInterpolation && DeformedMeshProjection) ?  RegisterAsSRV(GraphBuilder, DeformedMeshProjection->GetDeformedUniqueTrianglePositionBuffer(FHairStrandsLODDeformedRootResource::Current)) : Context.GetComputeDispatchInterface().GetEmptyBufferSRV(GraphBuilder, FHairStrandsMeshTrianglePositionFormat::Format);
-		
- 		FRDGBufferSRVRef RootBarycentricCoordinatesSRV = (bHasSkinnedInterpolation && RestMeshProjection) ? RegisterAsSRV(GraphBuilder, RestMeshProjection->RootBarycentricBuffer) : Context.GetComputeDispatchInterface().GetEmptyBufferSRV(GraphBuilder, FHairStrandsRootBarycentricFormat::Format);
-		FRDGBufferSRVRef RootToUniqueTriangleIndexSRV = (bHasSkinnedInterpolation && RestMeshProjection) ? RegisterAsSRV(GraphBuilder, RestMeshProjection->RootToUniqueTriangleIndexBuffer) : Context.GetComputeDispatchInterface().GetEmptyBufferSRV(GraphBuilder, FHairStrandsRootToUniqueTriangleIndexFormat::Format);
-
-		// RBF buffers
-		const bool bHasSamples = (RestMeshProjection && RestMeshProjection->SampleCount > 0);
-		const int32 SampleCountValue = bHasSamples ? RestMeshProjection->SampleCount : 0;
-
-		FRDGBufferSRVRef RestSamplePositionsBufferSRV = (bHasSamples && RestMeshProjection) ? RegisterAsSRV(GraphBuilder,RestMeshProjection->RestSamplePositionsBuffer) : DummyStructuredBufferSRV;
-		FRDGBufferSRVRef MeshSampleWeightsBufferSRV = (bHasSamples && DeformedMeshProjection) ? RegisterAsSRV(GraphBuilder,DeformedMeshProjection->GetMeshSampleWeightsBuffer((FHairStrandsLODDeformedRootResource::Current))) : DummyStructuredBufferSRV;
-
 		// Simulation setup (we update the rest configuration based on the deformed positions 
 		// if in restupdate mode or if we are resetting the sim and using RBF transfer since the rest positions are not matching the physics asset)
 		const int32 NeedResetValue = (ProxyData->TickCount <= GHairSimulationMaxDelay) || !HairStrandsBuffer->bValidGeometryType || (HairStrandsBuffer->CurrentMeshLOD != MeshLODIndex);
@@ -3430,7 +3445,6 @@ void UNiagaraDataInterfaceHairStrands::SetShaderParameters(const FNiagaraDataInt
 		HairStrandsBuffer->CurrentMeshLOD = MeshLODIndex;
 
 		// Offsets / Transforms
-		FVector3f RestPositionOffsetValue = (FVector3f)ProxyData->HairStrandsBuffer->SourceRestResources->GetPositionOffset();
 
 		const FMatrix44f WorldTransformFloat = ComputeWorldTransform(ProxyData);
 		const FMatrix44f BoneTransformFloat = ComputeBoneTransform(ProxyData);
@@ -3440,17 +3454,12 @@ void UNiagaraDataInterfaceHairStrands::SetShaderParameters(const FNiagaraDataInt
 			UE_LOG(LogHairStrands, Log, TEXT("Bad bones state"));
 		}
 
-		if (!bIsRootValid && bHasSkinningBinding)
+		if (InstanceRDG.IsDeformedValid() &&  !InstanceRDG.IsRootValid() && InstanceRDG.BindingType == EHairBindingType::Skinning)
 		{
 			UE_LOG(LogHairStrands, Log, TEXT("FNDIHairStrandsParametersCS() Groom Asset %s from component %s is set to use skinning interpolation but the skin resources are not valid"),
 				*ProxyData->HairGroupInstance->Debug.GroomAssetName, *ProxyData->HairGroupInstance->Debug.MeshComponentName);
 		}
 
-		if (bHasSkinnedInterpolation && (!DeformedMeshProjection || !RestMeshProjection))
-		{
-			UE_LOG(LogHairStrands, Log, TEXT("Bad valid projection"));
-		}
-
 		// Set shader constants
 		ShaderParameters->BoundingBoxOffsets = HairStrandsBuffer->BoundingBoxOffsets;
 		ShaderParameters->WorldTransform = WorldTransformFloat;
@@ -3471,24 +3480,22 @@ void UNiagaraDataInterfaceHairStrands::SetShaderParameters(const FNiagaraDataInt
 		ShaderParameters->LocalSimulation = LocalSimulationValue;
 		ShaderParameters->RestRootOffset = FVector3f::ZeroVector;
 		ShaderParameters->DeformedRootOffset = FVector3f::ZeroVector;
-		ShaderParameters->RestPositionOffset = RestPositionOffsetValue;
-		ShaderParameters->SampleCount = SampleCountValue;
+		ShaderParameters->RestPositionOffset = InstanceRDG.RestPositionOffsetValue;
+		ShaderParameters->SampleCount = InstanceRDG.SampleCount;
+
+		ShaderParameters->DeformedPositionBuffer = InstanceRDG.DeformedPositionBufferUAV;
+		ShaderParameters->CurvesOffsetsBuffer = InstanceRDG.CurvesOffsetsBuffer;
+		ShaderParameters->RestPositionBuffer = InstanceRDG.RestPositionBuffer;
+		ShaderParameters->DeformedPositionOffset = InstanceRDG.DeformedPositionOffset;
+		ShaderParameters->RestTrianglePositionBuffer = InstanceRDG.RestTrianglePositionBuffer;
+		ShaderParameters->DeformedTrianglePositionBuffer = InstanceRDG.DeformedTrianglePositionBuffer;
+		ShaderParameters->RestSamplePositionsBuffer = InstanceRDG.RestSamplePositionsBuffer;
+		ShaderParameters->MeshSampleWeightsBuffer = InstanceRDG.MeshSampleWeightsBuffer;
+		ShaderParameters->RootBarycentricCoordinatesBuffer = InstanceRDG.RootBarycentricCoordinatesBuffer;
+		ShaderParameters->RootToUniqueTriangleIndexBuffer = InstanceRDG.RootToUniqueTriangleIndexBuffer;
 
-		// Set Shader UAV
-		ShaderParameters->DeformedPositionBuffer = DeformedPositionBufferUAV;
 		ShaderParameters->BoundingBoxBuffer = HairStrandsBuffer->BoundingBoxBuffer.GetOrCreateUAV(GraphBuilder);
-
-		// Set Shader SRV
-		ShaderParameters->CurvesOffsetsBuffer = RegisterAsSRV(GraphBuilder, HairStrandsBuffer->SourceRestResources->CurveBuffer);
 		ShaderParameters->ParamsScaleBuffer = HairStrandsBuffer->ParamsScaleBuffer.GetOrCreateSRV(GraphBuilder);
-		ShaderParameters->RestPositionBuffer = RegisterAsSRV(GraphBuilder,HairStrandsBuffer->SourceRestResources->PositionBuffer); //-OPT: If constant across all dispatches can create in PreStage once
-		ShaderParameters->DeformedPositionOffset = DeformedPositionOffsetSRV;
-		ShaderParameters->RestTrianglePositionBuffer = RestTrianglePositionSRV;
-		ShaderParameters->DeformedTrianglePositionBuffer = DeformedTrianglePositionSRV;
-		ShaderParameters->RestSamplePositionsBuffer = RestSamplePositionsBufferSRV;
-		ShaderParameters->MeshSampleWeightsBuffer = MeshSampleWeightsBufferSRV;
-		ShaderParameters->RootBarycentricCoordinatesBuffer = RootBarycentricCoordinatesSRV;
-		ShaderParameters->RootToUniqueTriangleIndexBuffer = RootToUniqueTriangleIndexSRV;
 	}
 	else
 	{
@@ -3519,23 +3526,18 @@ void UNiagaraDataInterfaceHairStrands::SetShaderParameters(const FNiagaraDataInt
 		ShaderParameters->RestPositionOffset = FVector3f::ZeroVector;
 		ShaderParameters->SampleCount = 0;
 
-		// Set Shader UAV
-		ShaderParameters->DeformedPositionBuffer = GraphBuilder.CreateUAV(GraphBuilder.CreateBuffer(FRDGBufferDesc::CreateByteAddressDesc(16u), TEXT("Niagara.Hair.DummyUAV")), ERDGUnorderedAccessViewFlags::SkipBarrier);
-		ShaderParameters->BoundingBoxBuffer = Context.GetComputeDispatchInterface().GetEmptyBufferUAV(GraphBuilder, PF_R32_UINT);
+		ShaderParameters->DeformedPositionBuffer = InstanceRDG.DeformedPositionBufferUAV;
+		ShaderParameters->CurvesOffsetsBuffer = InstanceRDG.CurvesOffsetsBuffer;
+		ShaderParameters->RestPositionBuffer = InstanceRDG.RestPositionBuffer;
+		ShaderParameters->DeformedPositionOffset = InstanceRDG.DeformedPositionOffset;
+		ShaderParameters->RestTrianglePositionBuffer = InstanceRDG.RestTrianglePositionBuffer;
+		ShaderParameters->DeformedTrianglePositionBuffer = InstanceRDG.DeformedTrianglePositionBuffer;
+		ShaderParameters->RestSamplePositionsBuffer = InstanceRDG.RestSamplePositionsBuffer;
+		ShaderParameters->MeshSampleWeightsBuffer = InstanceRDG.MeshSampleWeightsBuffer;
+		ShaderParameters->RootBarycentricCoordinatesBuffer = InstanceRDG.RootBarycentricCoordinatesBuffer;
+		ShaderParameters->RootToUniqueTriangleIndexBuffer = InstanceRDG.RootToUniqueTriangleIndexBuffer;
 
-		// Set Shader SRV
-		FRDGBufferSRVRef DummyStructuredBuffer = GraphBuilder.CreateSRV(GSystemTextures.GetDefaultStructuredBuffer(GraphBuilder, 16u, FUintVector4(0,0,0,0)));
-		FRDGBufferSRVRef DummyByteAddressBuffer = GraphBuilder.CreateSRV(GSystemTextures.GetDefaultByteAddressBuffer(GraphBuilder, 16u));
-
-		ShaderParameters->CurvesOffsetsBuffer = DummyByteAddressBuffer;
-		ShaderParameters->RestPositionBuffer = DummyByteAddressBuffer;
-		ShaderParameters->DeformedPositionOffset = DummyStructuredBuffer;
-		ShaderParameters->RestTrianglePositionBuffer = Context.GetComputeDispatchInterface().GetEmptyBufferSRV(GraphBuilder, FHairStrandsMeshTrianglePositionFormat::Format);
-		ShaderParameters->DeformedTrianglePositionBuffer = Context.GetComputeDispatchInterface().GetEmptyBufferSRV(GraphBuilder, FHairStrandsMeshTrianglePositionFormat::Format);
-		ShaderParameters->RestSamplePositionsBuffer = DummyStructuredBuffer;
-		ShaderParameters->MeshSampleWeightsBuffer = DummyStructuredBuffer;
-		ShaderParameters->RootBarycentricCoordinatesBuffer = Context.GetComputeDispatchInterface().GetEmptyBufferSRV(GraphBuilder, FHairStrandsRootBarycentricFormat::Format);
-		ShaderParameters->RootToUniqueTriangleIndexBuffer = Context.GetComputeDispatchInterface().GetEmptyBufferSRV(GraphBuilder, FHairStrandsRootToUniqueTriangleIndexFormat::Format);
+		ShaderParameters->BoundingBoxBuffer = Context.GetComputeDispatchInterface().GetEmptyBufferUAV(GraphBuilder, PF_R32_UINT);
 		ShaderParameters->ParamsScaleBuffer = Context.GetComputeDispatchInterface().GetEmptyBufferSRV(GraphBuilder, PF_R32_FLOAT);
 	}
 }
@@ -3585,9 +3587,10 @@ void FNDIHairStrandsProxy::PostSimulate(const FNDIGpuComputePostSimulateContext&
 	}
 
 	// MGPU DeformedPositionBuffer copy after simulation
-	if (HairStrandsBuffer->SourceDeformedResources && HairStrandsBuffer->SourceDeformedResources->IsInitialized())
+	FHairStrandsDeformedResource* GuideDeformedBuffer = ProxyData->HairGroupInstance ? ProxyData->HairGroupInstance->Guides.DeformedResource : nullptr;
+	if (GuideDeformedBuffer && GuideDeformedBuffer->IsInitialized())
 	{
-		const FRDGExternalBuffer& DeformedBuffer = HairStrandsBuffer->SourceDeformedResources->DeformedPositionBuffer[HairStrandsBuffer->SourceDeformedResources->CurrentIndex];
+		const FRDGExternalBuffer& DeformedBuffer = GuideDeformedBuffer->GetBuffer(FHairStrandsDeformedResource::EFrameType::Current);
 		if (DeformedBuffer.Buffer.IsValid())
 		{
 			if (FRHIBuffer* DeformedPositionBuffer = DeformedBuffer.Buffer->GetRHI())
diff --git a/Engine/Plugins/Runtime/HairStrands/Source/HairStrandsCore/Public/Niagara/NiagaraDataInterfaceHairStrands.h b/Engine/Plugins/Runtime/HairStrands/Source/HairStrandsCore/Public/Niagara/NiagaraDataInterfaceHairStrands.h
--- a/Engine/Plugins/Runtime/HairStrands/Source/HairStrandsCore/Public/Niagara/NiagaraDataInterfaceHairStrands.h
+++ b/Engine/Plugins/Runtime/HairStrands/Source/HairStrandsCore/Public/Niagara/NiagaraDataInterfaceHairStrands.h
@@ -21,6 +21,18 @@ static const int32 ThicknessOffset = 3;
 
 struct FNDIHairStrandsData;
 
+struct FNDIHairStrandsInfo
+{
+	int32  GroupIndex = 0;
+	int32  LODIndex = 0;
+	uint32 NumControlPoints = 0;
+	uint32 NumCurves = 0;
+	FTransform LocalToWorld = FTransform::Identity;
+	bool bHasValidResources = false;
+
+	bool IsValid() const { return bHasValidResources; }
+};
+
 FHairGroupInstance* GetHairGroupInstance(UGroomComponent* In, int32 InGroupIndex);
 
 /** Render buffers that will be used in hlsl functions */
@@ -28,18 +40,11 @@ struct FNDIHairStrandsBuffer : public FRenderResource
 {
 	/** Set the asset that will be used to affect the buffer */
 	void Initialize(
-		const FHairStrandsRestResource*  HairStrandsRestResource, 
-		const FHairStrandsDeformedResource*  HairStrandsDeformedResource, 
-		const FHairStrandsRestRootResource* HairStrandsRestRootResource, 
-		const FHairStrandsDeformedRootResource* HairStrandsDeformedRootResource,
+		const FNDIHairStrandsInfo& In,
 		const TStaticArray<float, 32 * NumScales>& InParamsScale);
 
 	/** Set the asset that will be used to affect the buffer */
-	void Update(
-		const FHairStrandsRestResource* HairStrandsRestResource,
-		const FHairStrandsDeformedResource* HairStrandsDeformedResource,
-		const FHairStrandsRestRootResource* HairStrandsRestRootResource,
-		const FHairStrandsDeformedRootResource* HairStrandsDeformedRootResource);
+	void Update();
 
 	/** Transfer CPU datas to GPU */
 	void Transfer(FRDGBuilder& GraphBuilder, const TStaticArray<float, 32 * NumScales>& InParamsScale);
@@ -53,27 +58,12 @@ struct FNDIHairStrandsBuffer : public FRenderResource
 	/** Get the resource name */
 	virtual FString GetFriendlyName() const override { return TEXT("FNDIHairStrandsBuffer"); }
 
-	/** Deformed position buffer in case no resource are there */
-	TRefCountPtr<FRDGPooledBuffer> DeformedPositionBuffer;
-
 	/** Bounding Box Buffer*/
 	FNiagaraPooledRWBuffer BoundingBoxBuffer;
 
 	/** Params scale buffer */
 	FNiagaraPooledRWBuffer ParamsScaleBuffer;
 
-	/** The strand asset resource from which to sample */
-	const FHairStrandsRestResource* SourceRestResources;
-
-	/** The strand deformed resource to write into */
-	const FHairStrandsDeformedResource* SourceDeformedResources;
-
-	/** The strand root resource to write into */
-	const FHairStrandsRestRootResource* SourceRestRootResources;
-	
-	/** The strand root resource to write into */
-	const FHairStrandsDeformedRootResource* SourceDeformedRootResources;
-
 	/** Scales along the strand */
 	TStaticArray<float, 32 * NumScales> ParamsScale;
 
@@ -86,6 +76,9 @@ struct FNDIHairStrandsBuffer : public FRenderResource
 	/** Mesh LOD that is being used for the root resources */
 	int32 CurrentMeshLOD = INDEX_NONE;
 
+	/** Desc of the group */
+	FNDIHairStrandsInfo InfoData;
+
 	// For debug only
 	//FRHIGPUBufferReadback* ReadbackBuffer = nullptr;
 };
@@ -104,7 +97,10 @@ struct FNDIHairStrandsData
 	void Release();
 
 	/** Update the buffers */
-	void Update(UNiagaraDataInterfaceHairStrands* Interface, FNiagaraSystemInstance* SystemInstance, const FHairStrandsBulkData* HairStrandsDatas, UGroomAsset* GroomAsset, const int32 GroupIndex, const int32 LODIndex, const FTransform& LocalToWorld, const float DeltaSeconds);
+	void Update(
+		UNiagaraDataInterfaceHairStrands* Interface, 
+		const FNDIHairStrandsInfo& InData,
+		const float DeltaSeconds);
 
 	inline void ResetDatas()
 	{
@@ -452,14 +448,7 @@ public:
 	/** Extract datas and resources */
 	void ExtractDatasAndResources(
 		FNiagaraSystemInstance* SystemInstance, 
-		FHairStrandsRestResource*& OutStrandsRestResource, 
-		FHairStrandsDeformedResource*& OutStrandsDeformedResource, 
-		FHairStrandsRestRootResource*& OutStrandsRestRootResource, 
-		FHairStrandsDeformedRootResource*& OutStrandsDeformedRootResource,
-		UGroomAsset*& OutGroomAsset,
-		int32& OutGroupIndex,
-		int32& OutLODIndex, 
-		FTransform& OutLocalToWorld);
+		FNDIHairStrandsInfo& Out);
 
 	/** Get the number of strands */
 	void GetNumStrands(FVectorVMExternalFunctionContext& Context);
