Fix static analysis error reported by CIS.

diff --git a/Engine/Plugins/Runtime/HairStrands/Source/HairStrandsCore/Private/Niagara/NiagaraDataInterfaceHairStrands.cpp b/Engine/Plugins/Runtime/HairStrands/Source/HairStrandsCore/Private/Niagara/NiagaraDataInterfaceHairStrands.cpp
--- a/Engine/Plugins/Runtime/HairStrands/Source/HairStrandsCore/Private/Niagara/NiagaraDataInterfaceHairStrands.cpp
+++ b/Engine/Plugins/Runtime/HairStrands/Source/HairStrandsCore/Private/Niagara/NiagaraDataInterfaceHairStrands.cpp
@@ -3415,7 +3415,7 @@ void UNiagaraDataInterfaceHairStrands::SetShaderParameters(const FNiagaraDataInt
 	FNDIHairStrandsProxy& DIProxy = Context.GetProxy<FNDIHairStrandsProxy>();
 	FNDIHairStrandsData* ProxyData = DIProxy.SystemInstancesToProxyData.Find(Context.GetSystemInstanceID());
 
-	FHairGroupInstanceRDG InstanceRDG = Convert(GraphBuilder, ProxyData->HairGroupInstance);
+	FHairGroupInstanceRDG InstanceRDG = Convert(GraphBuilder, ProxyData ? ProxyData->HairGroupInstance : nullptr);
 	
 	const int32 MeshLODIndex = InstanceRDG.MeshLODIndex;
 	const bool bIsHairValid = ProxyData != nullptr && ProxyData->HairStrandsBuffer && ProxyData->HairStrandsBuffer->IsInitialized();
