Fix crash when spawning MH due to recent hair strands niagara resources setup refactoring.

diff --git a/Engine/Plugins/Runtime/HairStrands/Source/HairStrandsCore/Private/Niagara/NiagaraDataInterfaceHairStrands.cpp b/Engine/Plugins/Runtime/HairStrands/Source/HairStrandsCore/Private/Niagara/NiagaraDataInterfaceHairStrands.cpp
--- a/Engine/Plugins/Runtime/HairStrands/Source/HairStrandsCore/Private/Niagara/NiagaraDataInterfaceHairStrands.cpp
+++ b/Engine/Plugins/Runtime/HairStrands/Source/HairStrandsCore/Private/Niagara/NiagaraDataInterfaceHairStrands.cpp
@@ -466,7 +466,7 @@ void FNDIHairStrandsData::Update(
 	const FNDIHairStrandsInfo& InData,
 	const float DeltaSeconds)
 {
-	if (Interface != nullptr)
+	if (Interface != nullptr && InData.IsValid())
 	{
 		const UGroomAsset* GroomAsset = Interface->SourceComponent->GroomAsset;
 
