Remove loading of Niagara solver directly from the FHairSolverSettings constructor.

The niagara solver is loaded during the PostLoad() of the UGroomAsset.

diff --git a/Engine/Plugins/Runtime/HairStrands/Source/HairStrandsCore/Private/GroomAssetPhysics.cpp b/Engine/Plugins/Runtime/HairStrands/Source/HairStrandsCore/Private/GroomAssetPhysics.cpp
--- a/Engine/Plugins/Runtime/HairStrands/Source/HairStrandsCore/Private/GroomAssetPhysics.cpp
+++ b/Engine/Plugins/Runtime/HairStrands/Source/HairStrandsCore/Private/GroomAssetPhysics.cpp
@@ -13,12 +13,6 @@ FHairSolverSettings::FHairSolverSettings()
 	IterationCount = 5;
 	GravityPreloading = 0.0f;
 	bForceVisible = false;
-
-	// Object loading is only allowed on the GameThread
-	if (IsInGameThread())
-	{
-		CustomSystem =  LoadObject<UNiagaraSystem>(nullptr, TEXT("/HairStrands/Emitters/StableSpringsSystem.StableSpringsSystem"));
-	}
 }
 
 FHairExternalForces::FHairExternalForces()
