Make HairStrandsCore depenency on Niagara public

GroomComponent.h includes NiagaraDataInterfacePhysicsAsset.h but Niagara was a private dependency of the HairStrandsCore module. This requires users to always include the Niagara module to use GroomComponent.h

diff --git a/Engine/Plugins/Runtime/HairStrands/Source/HairStrandsCore/HairStrandsCore.Build.cs b/Engine/Plugins/Runtime/HairStrands/Source/HairStrandsCore/HairStrandsCore.Build.cs
--- a/Engine/Plugins/Runtime/HairStrands/Source/HairStrandsCore/HairStrandsCore.Build.cs
+++ b/Engine/Plugins/Runtime/HairStrands/Source/HairStrandsCore/HairStrandsCore.Build.cs
@@ -17,7 +17,6 @@ namespace UnrealBuildTool.Rules
 					"MeshDescription",
 					"MovieScene",
 					"NiagaraCore",
-					"Niagara",
 					"NiagaraShader",
 					"RenderCore",
 					"Renderer",
@@ -26,6 +25,12 @@ namespace UnrealBuildTool.Rules
 					"StaticMeshDescription"
 				});
 
+			PublicDependencyModuleNames.AddRange(
+				new string[]
+				{
+					"Niagara"
+				});
+
 			PrivateIncludePathModuleNames.AddRange(
 				new string[]
 				{
