Fix crash when updating and reloading groom asset.

diff --git a/Engine/Plugins/Runtime/HairStrands/Source/HairStrandsCore/Private/GroomComponent.cpp b/Engine/Plugins/Runtime/HairStrands/Source/HairStrandsCore/Private/GroomComponent.cpp
--- a/Engine/Plugins/Runtime/HairStrands/Source/HairStrandsCore/Private/GroomComponent.cpp
+++ b/Engine/Plugins/Runtime/HairStrands/Source/HairStrandsCore/Private/GroomComponent.cpp
@@ -4333,7 +4333,11 @@ FGroomComponentRecreateRenderStateContext::~FGroomComponentRecreateRenderStateCo
 			{
 				GroomComponent->InitResources();
 			}
-			GroomComponent->CreateRenderState_Concurrent(nullptr);
+
+			if (!GroomComponent->IsRenderStateCreated())
+			{
+				GroomComponent->CreateRenderState_Concurrent(nullptr);
+			}
 		}
 	}
 }
