Fix crash when calling GroomBinding's Build() function as the resources might be released before the instanced are deleted, causing used of dangling pointers.

diff --git a/Engine/Plugins/Runtime/HairStrands/Source/HairStrandsCore/Private/GroomBindingAsset.cpp b/Engine/Plugins/Runtime/HairStrands/Source/HairStrandsCore/Private/GroomBindingAsset.cpp
--- a/Engine/Plugins/Runtime/HairStrands/Source/HairStrandsCore/Private/GroomBindingAsset.cpp
+++ b/Engine/Plugins/Runtime/HairStrands/Source/HairStrandsCore/Private/GroomBindingAsset.cpp
@@ -17,6 +17,7 @@
 #include "Misc/CoreMisc.h"
 #include "Interfaces/ITargetPlatform.h"
 #include "Interfaces/ITargetPlatformManagerModule.h"
+#include "GroomComponent.h"
 
 #include UE_INLINE_GENERATED_CPP_BY_NAME(GroomBindingAsset)
 
@@ -26,7 +27,6 @@
 #include "DerivedDataCacheInterface.h"
 #include "Serialization/LargeMemoryReader.h"
 #include "Serialization/LargeMemoryWriter.h"
-#include "GroomComponent.h"
 #endif
 
 LLM_DECLARE_TAG(Groom);
@@ -429,8 +429,34 @@ void UGroomBindingAsset::ReleaseResource(bool bResetLoadedSize)
 	}
 }
 
+void UGroomBindingAsset::FlushRenderingCommandIfUsed(const UGroomBindingAsset* In)
+{
+	bool bNeedFlushRenderingCommand = false;
+	for (TObjectIterator<UGroomComponent> HairStrandsComponentIt; HairStrandsComponentIt; ++HairStrandsComponentIt)
+	{
+		if (HairStrandsComponentIt->BindingAsset == In
+			#if WITH_EDITORONLY_DATA
+			|| HairStrandsComponentIt->BindingAssetBeingLoaded == In
+			#endif
+			)
+		{
+			if (HairStrandsComponentIt->IsRenderStateCreated())
+			{
+				HairStrandsComponentIt->DestroyRenderState_Concurrent();
+				bNeedFlushRenderingCommand = true;
+			}
+		}
+	}
+	if (bNeedFlushRenderingCommand)
+	{
+		FlushRenderingCommands();
+	}
+}
+
 void UGroomBindingAsset::Reset()
 {
+	FlushRenderingCommandIfUsed(this);
+
 	// No need for resetting LoadedSize as the bulk datas are removed (e.g., when reloading/reimporting a groom asset)
 	ReleaseResource(false/*bResetLoadedSize*/);
 	for (UGroomBindingAsset::FHairGroupPlatformData& Data : GetHairGroupsPlatformData())
@@ -988,25 +1014,8 @@ void UGroomBindingAsset::CacheDerivedDatas()
 		FGroomBindingCompilingManager::Get().FinishCompilation({ this });
 	}
 
-	bool bNeedFlushRenderingCommand = false;
-	for (TObjectIterator<UGroomComponent> HairStrandsComponentIt; HairStrandsComponentIt; ++HairStrandsComponentIt)
-	{
-		if (HairStrandsComponentIt->BindingAsset == this ||
-			HairStrandsComponentIt->BindingAssetBeingLoaded == this) // A GroomAsset was set on the component while it was still loading
-		{
-			if (HairStrandsComponentIt->IsRenderStateCreated())
-			{
-				HairStrandsComponentIt->DestroyRenderState_Concurrent();
-				bNeedFlushRenderingCommand = true;
-			}
-		}
-	}
-
-	if (bNeedFlushRenderingCommand)
-	{
-		// Flush the rendering commands generated by the detachments.
-		FlushRenderingCommands();
-	}
+	// Flush the rendering commands generated by the detachments.
+	FlushRenderingCommandIfUsed(this);
 
 	FGroomBindingBuildContext Context;
 #if WITH_EDITOR
diff --git a/Engine/Plugins/Runtime/HairStrands/Source/HairStrandsCore/Public/GroomBindingAsset.h b/Engine/Plugins/Runtime/HairStrands/Source/HairStrandsCore/Public/GroomBindingAsset.h
--- a/Engine/Plugins/Runtime/HairStrands/Source/HairStrandsCore/Public/GroomBindingAsset.h
+++ b/Engine/Plugins/Runtime/HairStrands/Source/HairStrandsCore/Public/GroomBindingAsset.h
@@ -428,6 +428,7 @@ private:
 #endif
 	}
 
+	static void FlushRenderingCommandIfUsed(const UGroomBindingAsset* In);
 public:
 #if WITH_EDITORONLY_DATA
 	/** Build/rebuild a binding asset */
