Allow AutoLOD bias to have negative value.

diff --git a/Engine/Plugins/Runtime/HairStrands/Source/HairStrandsCore/Private/GroomManager.cpp b/Engine/Plugins/Runtime/HairStrands/Source/HairStrandsCore/Private/GroomManager.cpp
--- a/Engine/Plugins/Runtime/HairStrands/Source/HairStrandsCore/Private/GroomManager.cpp
+++ b/Engine/Plugins/Runtime/HairStrands/Source/HairStrandsCore/Private/GroomManager.cpp
@@ -72,7 +72,7 @@ bool UseHairStrandsForceAutoLOD()
 
 float GetHairStrandsAutoLODBias()
 {
-	return FMath::Clamp(GHairStrands_AutoLOD_Bias, 0.f, 1.f);
+	return FMath::Clamp(GHairStrands_AutoLOD_Bias, -1.f, 1.f);
 }
 
 ///////////////////////////////////////////////////////////////////////////////////////////////////
@@ -2454,7 +2454,7 @@ static float ComputeActiveCurveRadiusScale(const FHairStrandsClusterResource* In
 static uint32 ComputeActiveCurveCount(float InScreenSize, float InAutoLODBias, uint32 InCurveCount, uint32 InClusterCount)
 {
 	const float Power = 1.f;  // This could be exposed per asset
-	const float ScreenSizeBias = FMath::Clamp(FMath::Max(InAutoLODBias, GetHairStrandsAutoLODBias()), 0.f, 1.f);
+	const float ScreenSizeBias = FMath::Clamp(InAutoLODBias + GetHairStrandsAutoLODBias(), -1.f, 1.f);
 	uint32 OutCurveCount = InCurveCount * FMath::Pow(FMath::Clamp(InScreenSize + ScreenSizeBias, 0.f, 1.0f), Power);
 	// Ensure there is at least 1 curve per cluster
 	OutCurveCount = FMath::Max(InClusterCount, OutCurveCount);
diff --git a/Engine/Plugins/Runtime/HairStrands/Source/HairStrandsCore/Public/GroomAsset.h b/Engine/Plugins/Runtime/HairStrands/Source/HairStrandsCore/Public/GroomAsset.h
--- a/Engine/Plugins/Runtime/HairStrands/Source/HairStrandsCore/Public/GroomAsset.h
+++ b/Engine/Plugins/Runtime/HairStrands/Source/HairStrandsCore/Public/GroomAsset.h
@@ -454,7 +454,7 @@ private:
 	UPROPERTY(EditAnywhere, Category = "HairLOD", meta = (DisplayName = "LOD Mode", ToolTip = "Define how LOD adapts curves & points for strands geometry. Auto: adapts the curve count based on screen coverage. Manual: use the discrete LOD created for each groups"))
 	EGroomLODMode LODMode = EGroomLODMode::Default;
 
-	UPROPERTY(EditAnywhere, Category = "HairLOD", meta = (DisplayName = "Auto LOD Bias", ToolTip = "When LOD mode is set to Auto, decrease the screen size at which curves reduction will occur.", ClampMin = "0", ClampMax = "1", UIMin = "0", UIMax = "1.0"))
+	UPROPERTY(EditAnywhere, Category = "HairLOD", meta = (DisplayName = "Auto LOD Bias", ToolTip = "When LOD mode is set to Auto, decrease the screen size at which curves reduction will occur.", ClampMin = "-1", ClampMax = "1", UIMin = "-1.0", UIMax = "1.0"))
 	float AutoLODBias = 0;
 
 public:
