Fix groom importing when groom contains curve with invalid position (NaN/INF).

diff --git a/Engine/Plugins/Runtime/HairStrands/Source/HairStrandsCore/Private/GroomBuilder.cpp b/Engine/Plugins/Runtime/HairStrands/Source/HairStrandsCore/Private/GroomBuilder.cpp
--- a/Engine/Plugins/Runtime/HairStrands/Source/HairStrandsCore/Private/GroomBuilder.cpp
+++ b/Engine/Plugins/Runtime/HairStrands/Source/HairStrandsCore/Private/GroomBuilder.cpp
@@ -2062,10 +2062,31 @@ bool FGroomBuilder::BuildHairDescriptionGroups(const FHairDescription& HairDescr
 			continue;
 		}
 
+		// Check if the current curve has invalid data
+		bool bHasInvalidPosition = false;
+		{
+			for (int32 VertexIndex = 0; VertexIndex < CurveNumVertices; ++VertexIndex)
+			{
+				FVertexID VertexID(GlobalVertexIndex + VertexIndex);
+				const FVector3f P = VertexPositions[VertexID];
+				if (!(FMath::IsFinite(P.X) && FMath::IsFinite(P.Y) & FMath::IsFinite(P.Z)))
+				{
+					bHasInvalidPosition = true;
+					break;
+				}
+			}
+		}
+
 		FHairStrandsDatas* CurrentHairStrandsDatas = nullptr;
 		FHairDescriptionGroup& Group = FindOrAdd(StrandID);
 		bool bNumCurveValid = false;
-		if (!bIsGuide)
+		if (bHasInvalidPosition)
+		{
+			// Remove curve with NaN positions
+			GlobalVertexIndex += CurveNumVertices;
+			continue;
+		}
+		else if (!bIsGuide)
 		{
 			CurrentHairStrandsDatas = &Group.Strands;
 			if (Group.Info.NumCurves < HAIR_MAX_NUM_CURVE_PER_GROUP)
