* Rework Groom importer warning/error reporting.
* Add report of invalid position.

diff --git a/Engine/Plugins/Runtime/HairStrands/Source/HairStrandsCore/Private/GroomBuilder.cpp b/Engine/Plugins/Runtime/HairStrands/Source/HairStrandsCore/Private/GroomBuilder.cpp
--- a/Engine/Plugins/Runtime/HairStrands/Source/HairStrandsCore/Private/GroomBuilder.cpp
+++ b/Engine/Plugins/Runtime/HairStrands/Source/HairStrandsCore/Private/GroomBuilder.cpp
@@ -2084,6 +2084,7 @@ bool FGroomBuilder::BuildHairDescriptionGroups(const FHairDescription& HairDescr
 		{
 			// Remove curve with NaN positions
 			GlobalVertexIndex += CurveNumVertices;
+			Group.Info.Flags |= uint32(EHairGroupInfoFlags::HasInvalidPoint);
 			continue;
 		}
 		else if (!bIsGuide)
diff --git a/Engine/Plugins/Runtime/HairStrands/Source/HairStrandsCore/Public/GroomAsset.h b/Engine/Plugins/Runtime/HairStrands/Source/HairStrandsCore/Public/GroomAsset.h
--- a/Engine/Plugins/Runtime/HairStrands/Source/HairStrandsCore/Public/GroomAsset.h
+++ b/Engine/Plugins/Runtime/HairStrands/Source/HairStrandsCore/Public/GroomAsset.h
@@ -35,8 +35,9 @@ struct FHairStrandsRaytracingResource;
 
 enum class EHairGroupInfoFlags : uint8
 {
-	HasTrimmedPoint = 1,
-	HasTrimmedCurve = 2
+	HasTrimmedPoint = 1<<0,
+	HasTrimmedCurve = 1<<1,
+	HasInvalidPoint = 1<<2
 };
 
 USTRUCT(BlueprintType)
diff --git a/Engine/Plugins/Runtime/HairStrands/Source/HairStrandsEditor/Private/GroomImportOptionsWindow.cpp b/Engine/Plugins/Runtime/HairStrands/Source/HairStrandsEditor/Private/GroomImportOptionsWindow.cpp
--- a/Engine/Plugins/Runtime/HairStrands/Source/HairStrandsEditor/Private/GroomImportOptionsWindow.cpp
+++ b/Engine/Plugins/Runtime/HairStrands/Source/HairStrandsEditor/Private/GroomImportOptionsWindow.cpp
@@ -19,29 +19,12 @@
 
 #define LOCTEXT_NAMESPACE "GroomImportOptionsWindow"
 
-enum class EHairDescriptionStatus
-{
-	Unset,
-	Valid,
-	NoGroup,
-	NoCurve,
-	GroomCache, // groom cache with unspecified groom asset
-	GroomCacheCompatible,
-	GroomCacheIncompatible,
-	GuidesOnly, // guides-only with unspecified groom asset
-	GuidesOnlyCompatible,
-	GuidesOnlyIncompatible,
-	ValidPointLimit,
-	ValidCurveLimit,
-	ValidCurveAndPointLimit,
-	Unknown
-};
-
 void SGroomImportOptionsWindow::UpdateStatus(UGroomHairGroupsPreview* Description) const
 {
+	CurrentStatus = EHairDescriptionStatus::None;
 	if (!Description)
 	{
-		CurrentStatus = EHairDescriptionStatus::Unknown;
+		EnumAddFlags(CurrentStatus, EHairDescriptionStatus::Unknown);
 		return;
 	}
 
@@ -49,25 +32,24 @@ void SGroomImportOptionsWindow::UpdateStatus(UGroomHairGroupsPreview* Descriptio
 	const bool bImportGroomCache = GroomCacheImportOptions && GroomCacheImportOptions->ImportSettings.bImportGroomCache;
 	if (!bImportGroomAsset && !bImportGroomCache)
 	{
-		CurrentStatus = EHairDescriptionStatus::Unset;
+		EnumAddFlags(CurrentStatus, EHairDescriptionStatus::Unknown);
 		return;
 	}
 
 	if (Description->Groups.Num() == 0)
 	{
-		CurrentStatus = EHairDescriptionStatus::NoGroup;
+		EnumAddFlags(CurrentStatus, EHairDescriptionStatus::NoGroup);
 		return;
 	}
 
 	// Check the validity of the groom to import
-	CurrentStatus = EHairDescriptionStatus::Valid;
 
 	bool bGuidesOnly = false;
 	for (const FGroomHairGroupPreview& Group : Description->Groups)
 	{
 		if (Group.CurveCount == 0)
 		{
-			CurrentStatus = EHairDescriptionStatus::NoCurve;
+			EnumAddFlags(CurrentStatus, EHairDescriptionStatus::NoCurve);
 			if (Group.GuideCount > 0)
 			{
 				bGuidesOnly = true;
@@ -79,22 +61,23 @@ void SGroomImportOptionsWindow::UpdateStatus(UGroomHairGroupsPreview* Descriptio
 	// Check if any curve or point have been trimmed
 	for (const FGroomHairGroupPreview& Group : Description->Groups)
 	{
-		if ((Group.Flags & uint32(EHairGroupInfoFlags::HasTrimmedCurve)) && (Group.Flags & uint32(EHairGroupInfoFlags::HasTrimmedPoint)))
+		if (Group.Flags & uint32(EHairGroupInfoFlags::HasTrimmedCurve))
 		{
-			CurrentStatus = EHairDescriptionStatus::ValidCurveAndPointLimit;
+			EnumAddFlags(CurrentStatus, EHairDescriptionStatus::CurveLimit);
 		}
-		else if (Group.Flags & uint32(EHairGroupInfoFlags::HasTrimmedCurve))
+		if (Group.Flags & uint32(EHairGroupInfoFlags::HasTrimmedPoint))
 		{
-			CurrentStatus = EHairDescriptionStatus::ValidCurveLimit;
+			EnumAddFlags(CurrentStatus, EHairDescriptionStatus::PointLimit);
 		}
-		else if (Group.Flags & uint32(EHairGroupInfoFlags::HasTrimmedPoint))
+		if (Group.Flags & uint32(EHairGroupInfoFlags::HasInvalidPoint))
 		{
-			CurrentStatus = EHairDescriptionStatus::ValidPointLimit;
-		}
+			EnumAddFlags(CurrentStatus, EHairDescriptionStatus::InvalidPoint);
+		}		
 	}
 
 	if (!bImportGroomCache)
 	{
+		EnumAddFlags(CurrentStatus, EHairDescriptionStatus::GroomValid);
 		return;
 	}
 
@@ -110,24 +93,23 @@ void SGroomImportOptionsWindow::UpdateStatus(UGroomHairGroupsPreview* Descriptio
 		if (!GroomAssetForCache)
 		{
 			// No groom asset provided or loaded but one is needed with this setting
-			CurrentStatus = bGuidesOnly ? EHairDescriptionStatus::GuidesOnly : EHairDescriptionStatus::GroomCache;
+			EnumAddFlags(CurrentStatus, bGuidesOnly ? EHairDescriptionStatus::GuidesOnly : EHairDescriptionStatus::GroomCache);
 			return;
 		}
 
 		const TArray<FHairGroupPlatformData>& GroomHairGroupsData = GroomAssetForCache->GetHairGroupsPlatformData();
 		if (GroomHairGroupsData.Num() != Description->Groups.Num())
 		{
-			CurrentStatus = bGuidesOnly ? EHairDescriptionStatus::GuidesOnlyIncompatible : EHairDescriptionStatus::GroomCacheIncompatible;
+			EnumAddFlags(CurrentStatus, bGuidesOnly ? EHairDescriptionStatus::GuidesOnlyIncompatible : EHairDescriptionStatus::GroomCacheIncompatible);
 			return;
 		}
 
-		CurrentStatus = bGuidesOnly ? EHairDescriptionStatus::GuidesOnlyCompatible : EHairDescriptionStatus::GroomCacheCompatible;
 		for (int32 Index = 0; Index < GroomHairGroupsData.Num(); ++Index)
 		{
 			// Check the strands compatibility
 			if (!bGuidesOnly && Description->Groups[Index].CurveCount != GroomHairGroupsData[Index].Strands.BulkData.GetNumCurves())
 			{
-				CurrentStatus = EHairDescriptionStatus::GroomCacheIncompatible;
+				EnumAddFlags(CurrentStatus, EHairDescriptionStatus::GroomCacheIncompatible);
 				break;
 			}
 
@@ -137,80 +119,52 @@ void SGroomImportOptionsWindow::UpdateStatus(UGroomHairGroupsPreview* Descriptio
 			if (Description->Groups[Index].GuideCount > 0 &&
 				Description->Groups[Index].GuideCount != GroomHairGroupsData[Index].Guides.BulkData.GetNumCurves())
 			{
-				CurrentStatus = bGuidesOnly ? EHairDescriptionStatus::GuidesOnlyIncompatible : EHairDescriptionStatus::GroomCacheIncompatible;
+				EnumAddFlags(CurrentStatus, bGuidesOnly ? EHairDescriptionStatus::GuidesOnlyIncompatible : EHairDescriptionStatus::GroomCacheIncompatible);
 				break;
 			}
 		}
+
+		EnumAddFlags(CurrentStatus, bGuidesOnly ? EHairDescriptionStatus::GuidesOnlyCompatible : EHairDescriptionStatus::GroomCacheCompatible);
 	}
 	else
 	{
 		// A guides-only groom cannot be imported as asset, but otherwise the imported groom asset
 		// is always compatible with the groom cache since they are from the same file
-		CurrentStatus = bGuidesOnly ? EHairDescriptionStatus::GuidesOnly : EHairDescriptionStatus::Valid;
+		EnumAddFlags(CurrentStatus, bGuidesOnly ? EHairDescriptionStatus::GuidesOnly : EHairDescriptionStatus::GroomValid);
 	}
 }
 
+
 FText SGroomImportOptionsWindow::GetStatusText() const
 {
-	switch (CurrentStatus)
-	{
-		case EHairDescriptionStatus::Valid:
-			return LOCTEXT("GroomOptionsWindow_ValidationText0", "Valid");
-		case EHairDescriptionStatus::NoCurve:
-			return LOCTEXT("GroomOptionsWindow_ValidationText1", "Invalid. Some groups have 0 curves.");
-		case EHairDescriptionStatus::NoGroup:
-			return LOCTEXT("GroomOptionsWindow_ValidationText2", "Invalid. The groom does not contain any group.");
-		case EHairDescriptionStatus::GroomCache:
-			return LOCTEXT("GroomOptionsWindow_ValidationText4", "A compatible groom asset must be provided to import the groom cache.");
-		case EHairDescriptionStatus::GroomCacheCompatible:
-			return LOCTEXT("GroomOptionsWindow_ValidationText5", "The groom cache is compatible with the groom asset provided .");
-		case EHairDescriptionStatus::GroomCacheIncompatible:
-			return LOCTEXT("GroomOptionsWindow_ValidationText6", "The groom cache is incompatible with the groom asset provided .");
-		case EHairDescriptionStatus::GuidesOnly:
-			return LOCTEXT("GroomOptionsWindow_ValidationText7", "Only guides were detected. A compatible groom asset must be provided.");
-		case EHairDescriptionStatus::GuidesOnlyCompatible:
-			return LOCTEXT("GroomOptionsWindow_ValidationText8", "Only guides were detected. The groom asset provided is compatible.");
-		case EHairDescriptionStatus::GuidesOnlyIncompatible:
-			return LOCTEXT("GroomOptionsWindow_ValidationText9", "Only guides were detected. The groom asset provided is incompatible.");
-		case EHairDescriptionStatus::ValidCurveLimit:
-			return LOCTEXT("GroomOptionsWindow_ValidationText10", "Valid. At least one group contains more curves than allowed limit (Max:4M). Curves beyond that limit will be trimmed."); static_assert(HAIR_MAX_NUM_CURVE_PER_GROUP == 4194303);
-		case EHairDescriptionStatus::ValidPointLimit:
-			return LOCTEXT("GroomOptionsWindow_ValidationText11", "Valid. At least one group contains more control points per curve than the allowed limit (Max:255). Control points beyond that limit will be trimmed."); static_assert(HAIR_MAX_NUM_POINT_PER_CURVE == 255);
-		case EHairDescriptionStatus::ValidCurveAndPointLimit:
-			return LOCTEXT("GroomOptionsWindow_ValidationText12", "Valid. At least one group contains more control points per curve and more curves than the allowed limit (curve limit:4M, point limit:255). Curves and control points beyond that limit will be trimmed.");
-		case EHairDescriptionStatus::Unset:
-		case EHairDescriptionStatus::Unknown:
-		default:
-			return LOCTEXT("GroomOptionsWindow_ValidationText3", "Unknown");
-	}
+	FString Out;
+	if (EnumHasAnyFlags(CurrentStatus, EHairDescriptionStatus::Error)) 						{ Out += LOCTEXT("GroomOptionsWindow_ValidationText0",  "Error\n").ToString(); }
+	else if (EnumHasAnyFlags(CurrentStatus, EHairDescriptionStatus::Warning)) 				{ Out += LOCTEXT("GroomOptionsWindow_ValidationText1",  "Warning\n").ToString(); }
+	else if (EnumHasAnyFlags(CurrentStatus, EHairDescriptionStatus::Valid))					{ Out += LOCTEXT("GroomOptionsWindow_ValidationText2",  "Valid\n").ToString(); }
+
+	if (EnumHasAnyFlags(CurrentStatus, EHairDescriptionStatus::NoCurve)) 					{ Out += LOCTEXT("GroomOptionsWindow_ValidationText3",  "Some groups have 0 curves.\n").ToString(); }
+	if (EnumHasAnyFlags(CurrentStatus, EHairDescriptionStatus::NoGroup)) 					{ Out += LOCTEXT("GroomOptionsWindow_ValidationText4",  "The groom does not contain any group.\n").ToString(); }
+	if (EnumHasAnyFlags(CurrentStatus, EHairDescriptionStatus::GroomCache))					{ Out += LOCTEXT("GroomOptionsWindow_ValidationText5",  "A compatible groom asset must be provided to import the groom cache.\n").ToString(); }
+	if (EnumHasAnyFlags(CurrentStatus, EHairDescriptionStatus::GroomCacheCompatible)) 		{ Out += LOCTEXT("GroomOptionsWindow_ValidationText6",  "The groom cache is compatible with the groom asset provided.\n").ToString(); }
+	if (EnumHasAnyFlags(CurrentStatus, EHairDescriptionStatus::GroomCacheIncompatible))		{ Out += LOCTEXT("GroomOptionsWindow_ValidationText7",  "The groom cache is incompatible with the groom asset provided.\n").ToString(); }
+	if (EnumHasAnyFlags(CurrentStatus, EHairDescriptionStatus::GuidesOnly)) 				{ Out += LOCTEXT("GroomOptionsWindow_ValidationText8",  "Only guides were detected. A compatible groom asset must be provided.\n").ToString(); }
+	if (EnumHasAnyFlags(CurrentStatus, EHairDescriptionStatus::GuidesOnlyCompatible)) 		{ Out += LOCTEXT("GroomOptionsWindow_ValidationText9",  "Only guides were detected. The groom asset provided is compatible.\n").ToString(); }
+	if (EnumHasAnyFlags(CurrentStatus, EHairDescriptionStatus::GuidesOnlyIncompatible))		{ Out += LOCTEXT("GroomOptionsWindow_ValidationText10", "Only guides were detected. The groom asset provided is incompatible.\n").ToString(); }
+	if (EnumHasAnyFlags(CurrentStatus, EHairDescriptionStatus::CurveLimit))					{ Out += LOCTEXT("GroomOptionsWindow_ValidationText11", "At least one group contains more curves than allowed limit (Max:4M). Curves beyond that limit will be trimmed.\n").ToString(); static_assert(HAIR_MAX_NUM_CURVE_PER_GROUP == 4194303); }
+	if (EnumHasAnyFlags(CurrentStatus, EHairDescriptionStatus::PointLimit))					{ Out += LOCTEXT("GroomOptionsWindow_ValidationText12", "At least one group contains more control points per curve than the allowed limit (Max:255). Control points beyond that limit will be trimmed.\n").ToString(); static_assert(HAIR_MAX_NUM_POINT_PER_CURVE == 255); }
+	if (EnumHasAnyFlags(CurrentStatus, EHairDescriptionStatus::InvalidPoint)) 				{ Out += LOCTEXT("GroomOptionsWindow_ValidationText13", "At least one group contains a curve with invalid points. These curves will be trimmed from the asset.\n").ToString(); }
+	if (EnumHasAnyFlags(CurrentStatus, EHairDescriptionStatus::Unknown)) 					{ Out += LOCTEXT("GroomOptionsWindow_ValidationText14", "Unknown\n").ToString(); }
+
+	return FText::FromString(Out);
 }
 
 FSlateColor SGroomImportOptionsWindow::GetStatusColor() const
 {
-	switch (CurrentStatus)
-	{
-		case EHairDescriptionStatus::Valid:
-		case EHairDescriptionStatus::GroomCacheCompatible:
-		case EHairDescriptionStatus::GuidesOnlyCompatible:
-			return FLinearColor(0, 0.80f, 0, 1);
-		case EHairDescriptionStatus::NoCurve:
-		case EHairDescriptionStatus::GroomCacheIncompatible:
-		case EHairDescriptionStatus::GuidesOnlyIncompatible:
-			return FLinearColor(0.80f, 0, 0, 1);
-		case EHairDescriptionStatus::NoGroup:
-			return FLinearColor(1, 0, 0, 1);
-		case EHairDescriptionStatus::GroomCache:
-		case EHairDescriptionStatus::GuidesOnly:
-			return FLinearColor(0.80f, 0.80f, 0, 1);
-		case EHairDescriptionStatus::ValidCurveLimit:
-		case EHairDescriptionStatus::ValidPointLimit:
-		case EHairDescriptionStatus::ValidCurveAndPointLimit:
-			return FLinearColor(0.80f, 0.80f, 0, 1);
-		case EHairDescriptionStatus::Unset:
-		case EHairDescriptionStatus::Unknown:
-		default:
-			return FLinearColor(1, 1, 1);
-	}
+	if (EnumHasAnyFlags(CurrentStatus, EHairDescriptionStatus::Error))  	{ return FLinearColor(0.80f, 0, 0, 1); }
+	if (EnumHasAnyFlags(CurrentStatus, EHairDescriptionStatus::Warning))  	{ return FLinearColor(0.80f, 0.80f, 0, 1); }
+	if (EnumHasAnyFlags(CurrentStatus, EHairDescriptionStatus::Valid))		{ return FLinearColor(0, 0.80f, 0, 1); }
+
+	return FLinearColor(1, 1, 1);
 }
 
 static void AddAttribute(SVerticalBox::FScopedWidgetSlotArguments& Slot, FText AttributeLegend)
@@ -280,7 +234,7 @@ void SGroomImportOptionsWindow::Construct(const FArguments& InArgs)
 	GroomCacheDetailsView = PropertyEditorModule.CreateDetailView(DetailsViewArgs);
 	GroomCacheDetailsView->SetObject(GroomCacheImportOptions);
 
-	CurrentStatus = EHairDescriptionStatus::Unset;
+	CurrentStatus = EHairDescriptionStatus::None;
 	UpdateStatus(GroupsPreview);
 
 	// Aggregate attributes from all groups (ideally we should display each group attribute separately, to check if one groom is not missing data)
@@ -457,7 +411,7 @@ void SGroomImportOptionsWindow::Construct(const FArguments& InArgs)
 
 bool SGroomImportOptionsWindow::CanImport() const
 {
-	bool bNeedUpdate = CurrentStatus == EHairDescriptionStatus::Unset;
+	bool bNeedUpdate = CurrentStatus == EHairDescriptionStatus::None;
 	if (GroomCacheImportOptions)
 	{
 		bNeedUpdate |= bImportGroomAssetState != GroomCacheImportOptions->ImportSettings.bImportGroomAsset;
@@ -470,26 +424,7 @@ bool SGroomImportOptionsWindow::CanImport() const
 		UpdateStatus(GroupsPreview);
 	}
 
-	switch (CurrentStatus)
-	{
-		case EHairDescriptionStatus::Valid:
-		case EHairDescriptionStatus::GroomCacheCompatible:
-		case EHairDescriptionStatus::GuidesOnlyCompatible:
-		case EHairDescriptionStatus::ValidPointLimit:
-		case EHairDescriptionStatus::ValidCurveLimit:
-		case EHairDescriptionStatus::ValidCurveAndPointLimit:
-			return true;
-		case EHairDescriptionStatus::Unset:
-		case EHairDescriptionStatus::NoGroup:
-		case EHairDescriptionStatus::NoCurve:
-		case EHairDescriptionStatus::GroomCache:
-		case EHairDescriptionStatus::GroomCacheIncompatible:
-		case EHairDescriptionStatus::GuidesOnly:
-		case EHairDescriptionStatus::GuidesOnlyIncompatible:
-		case EHairDescriptionStatus::Unknown:
-		default:
-			return false;
-	}
+	return EnumHasAnyFlags(CurrentStatus, EHairDescriptionStatus::Valid | EHairDescriptionStatus::Warning);
 }
 
 enum class EGroomOptionsVisibility : uint8
diff --git a/Engine/Plugins/Runtime/HairStrands/Source/HairStrandsEditor/Private/GroomImportOptionsWindow.h b/Engine/Plugins/Runtime/HairStrands/Source/HairStrandsEditor/Private/GroomImportOptionsWindow.h
--- a/Engine/Plugins/Runtime/HairStrands/Source/HairStrandsEditor/Private/GroomImportOptionsWindow.h
+++ b/Engine/Plugins/Runtime/HairStrands/Source/HairStrandsEditor/Private/GroomImportOptionsWindow.h
@@ -9,13 +9,37 @@
 #include "Widgets/SCompoundWidget.h"
 #include "Widgets/SWindow.h"
 
+#include "GroomImportOptionsWindow.generated.h"
+
 class SButton;
 class UGroomImportOptions;
 class UGroomCacheImportOptions;
 class UGroomHairGroupsPreview;
 struct FHairGroupInfo;
 
-enum class EHairDescriptionStatus;
+UENUM()
+enum class EHairDescriptionStatus : uint32
+{
+	None					= 0,
+	GroomValid				= 1u<<1,
+	NoGroup					= 1u<<2,
+	NoCurve					= 1u<<3,
+	GroomCache				= 1u<<4, // groom cache with unspecified groom asset
+	GroomCacheCompatible	= 1u<<5,
+	GroomCacheIncompatible	= 1u<<6,
+	GuidesOnly				= 1u<<7, // guides-only with unspecified groom asset
+	GuidesOnlyCompatible	= 1u<<8,
+	GuidesOnlyIncompatible	= 1u<<9,
+	PointLimit				= 1u<<10,
+	CurveLimit				= 1u<<11,
+	InvalidPoint			= 1u<<12,
+	Unknown					= 1u<<13,
+
+	Warning = CurveLimit | PointLimit | InvalidPoint,
+	Error   = NoCurve | NoGroup | GroomCacheIncompatible | GuidesOnlyIncompatible | GroomCache | GuidesOnly | Unknown,
+	Valid   = GroomValid | GroomCacheCompatible | GuidesOnlyCompatible
+};
+ENUM_CLASS_FLAGS(EHairDescriptionStatus);
 
 class SGroomImportOptionsWindow : public SCompoundWidget
 {
