Fix CIS error due to typo.

diff --git a/Engine/Plugins/Runtime/HairStrands/Source/HairStrandsCore/Private/GroomBuilder.cpp b/Engine/Plugins/Runtime/HairStrands/Source/HairStrandsCore/Private/GroomBuilder.cpp
--- a/Engine/Plugins/Runtime/HairStrands/Source/HairStrandsCore/Private/GroomBuilder.cpp
+++ b/Engine/Plugins/Runtime/HairStrands/Source/HairStrandsCore/Private/GroomBuilder.cpp
@@ -2069,7 +2069,7 @@ bool FGroomBuilder::BuildHairDescriptionGroups(const FHairDescription& HairDescr
 			{
 				FVertexID VertexID(GlobalVertexIndex + VertexIndex);
 				const FVector3f P = VertexPositions[VertexID];
-				if (!(FMath::IsFinite(P.X) && FMath::IsFinite(P.Y) & FMath::IsFinite(P.Z)))
+				if (!(FMath::IsFinite(P.X) && FMath::IsFinite(P.Y) && FMath::IsFinite(P.Z)))
 				{
 					bHasInvalidPosition = true;
 					break;
