Fix crash in hair strands niagara component when accessing proxy data.

diff --git a/Engine/Plugins/Runtime/HairStrands/Source/HairStrandsCore/Private/Niagara/NiagaraDataInterfaceHairStrands.cpp b/Engine/Plugins/Runtime/HairStrands/Source/HairStrandsCore/Private/Niagara/NiagaraDataInterfaceHairStrands.cpp
--- a/Engine/Plugins/Runtime/HairStrands/Source/HairStrandsCore/Private/Niagara/NiagaraDataInterfaceHairStrands.cpp
+++ b/Engine/Plugins/Runtime/HairStrands/Source/HairStrandsCore/Private/Niagara/NiagaraDataInterfaceHairStrands.cpp
@@ -3414,8 +3414,14 @@ void UNiagaraDataInterfaceHairStrands::SetShaderParameters(const FNiagaraDataInt
 	FRDGBuilder& GraphBuilder = Context.GetGraphBuilder();
 	FNDIHairStrandsProxy& DIProxy = Context.GetProxy<FNDIHairStrandsProxy>();
 	FNDIHairStrandsData* ProxyData = DIProxy.SystemInstancesToProxyData.Find(Context.GetSystemInstanceID());
+	
+	TRefCountPtr<FHairGroupInstance> HairGroupInstance = nullptr;
+	if (ProxyData->HairGroupInstSource.Get())
+	{
+		HairGroupInstance = GetHairGroupInstance(ProxyData->HairGroupInstSource.Get(), ProxyData->HairGroupIndex);
+	}
 
-	FHairGroupInstanceRDG InstanceRDG = Convert(GraphBuilder, ProxyData ? ProxyData->HairGroupInstance : nullptr);
+	FHairGroupInstanceRDG InstanceRDG = Convert(GraphBuilder, ProxyData ? HairGroupInstance : nullptr);
 	
 	const int32 MeshLODIndex = InstanceRDG.MeshLODIndex;
 	const bool bIsHairValid = ProxyData != nullptr && ProxyData->HairStrandsBuffer && ProxyData->HairStrandsBuffer->IsInitialized();
@@ -3457,7 +3463,7 @@ void UNiagaraDataInterfaceHairStrands::SetShaderParameters(const FNiagaraDataInt
 		if (InstanceRDG.IsDeformedValid() &&  !InstanceRDG.IsRootValid() && InstanceRDG.BindingType == EHairBindingType::Skinning)
 		{
 			UE_LOG(LogHairStrands, Log, TEXT("FNDIHairStrandsParametersCS() Groom Asset %s from component %s is set to use skinning interpolation but the skin resources are not valid"),
-				*ProxyData->HairGroupInstance->Debug.GroomAssetName, *ProxyData->HairGroupInstance->Debug.MeshComponentName);
+				*HairGroupInstance->Debug.GroomAssetName, *HairGroupInstance->Debug.MeshComponentName);
 		}
 
 		// Set shader constants
@@ -3586,8 +3592,14 @@ void FNDIHairStrandsProxy::PostSimulate(const FNDIGpuComputePostSimulateContext&
 		return;
 	}
 
+	TRefCountPtr<FHairGroupInstance> HairGroupInstance = nullptr;
+	if (ProxyData->HairGroupInstSource.Get())
+	{
+		HairGroupInstance = GetHairGroupInstance(ProxyData->HairGroupInstSource.Get(), ProxyData->HairGroupIndex);
+	}
+
 	// MGPU DeformedPositionBuffer copy after simulation
-	FHairStrandsDeformedResource* GuideDeformedBuffer = ProxyData->HairGroupInstance ? ProxyData->HairGroupInstance->Guides.DeformedResource : nullptr;
+	FHairStrandsDeformedResource* GuideDeformedBuffer = HairGroupInstance ? HairGroupInstance->Guides.DeformedResource : nullptr;
 	if (GuideDeformedBuffer && GuideDeformedBuffer->IsInitialized())
 	{
 		const FRDGExternalBuffer& DeformedBuffer = GuideDeformedBuffer->GetBuffer(FHairStrandsDeformedResource::EFrameType::Current);
