Fix CIS warning.

diff --git a/Engine/Plugins/Runtime/HairStrands/Source/HairStrandsCore/Private/Niagara/NiagaraDataInterfaceHairStrands.cpp b/Engine/Plugins/Runtime/HairStrands/Source/HairStrandsCore/Private/Niagara/NiagaraDataInterfaceHairStrands.cpp
--- a/Engine/Plugins/Runtime/HairStrands/Source/HairStrandsCore/Private/Niagara/NiagaraDataInterfaceHairStrands.cpp
+++ b/Engine/Plugins/Runtime/HairStrands/Source/HairStrandsCore/Private/Niagara/NiagaraDataInterfaceHairStrands.cpp
@@ -3416,7 +3416,7 @@ void UNiagaraDataInterfaceHairStrands::SetShaderParameters(const FNiagaraDataInt
 	FNDIHairStrandsData* ProxyData = DIProxy.SystemInstancesToProxyData.Find(Context.GetSystemInstanceID());
 	
 	TRefCountPtr<FHairGroupInstance> HairGroupInstance = nullptr;
-	if (ProxyData->HairGroupInstSource.Get())
+	if (ProxyData && ProxyData->HairGroupInstSource.Get())
 	{
 		HairGroupInstance = GetHairGroupInstance(ProxyData->HairGroupInstSource.Get(), ProxyData->HairGroupIndex);
 	}
