Move FNDIHairStrandsInfo into .cpp in NiagaraDataInterfaceHairStrands.h/.cpp and remove unecessary parts.

diff --git a/Engine/Plugins/Runtime/HairStrands/Source/HairStrandsCore/Private/Niagara/NiagaraDataInterfaceHairStrands.cpp b/Engine/Plugins/Runtime/HairStrands/Source/HairStrandsCore/Private/Niagara/NiagaraDataInterfaceHairStrands.cpp
--- a/Engine/Plugins/Runtime/HairStrands/Source/HairStrandsCore/Private/Niagara/NiagaraDataInterfaceHairStrands.cpp
+++ b/Engine/Plugins/Runtime/HairStrands/Source/HairStrandsCore/Private/Niagara/NiagaraDataInterfaceHairStrands.cpp
@@ -364,22 +364,32 @@ static FHairGroupInstanceRDG Convert(FRDGBuilder& GraphBuilder, const FHairGroup
 
 //------------------------------------------------------------------------------------------------------------
 
+struct FNDIHairStrandsInfo
+{
+	int32  GroupIndex = 0;
+	int32  LODIndex = 0;
+	uint32 NumControlPoints = 0;
+	uint32 NumCurves = 0;
+	FTransform LocalToWorld = FTransform::Identity;
+	bool bHasValidResources = false;
+
+	bool IsValid() const { return bHasValidResources; }
+};
+
+//------------------------------------------------------------------------------------------------------------
+
 void FNDIHairStrandsBuffer::Initialize(
 	const FNDIHairStrandsInfo& In,
 	const TStaticArray<float, 32 * NumScales>& InParamsScale)
 {
-	InfoData = In;
+	bNeedResouces = In.IsValid();
 	ParamsScale = InParamsScale;
 	bValidGeometryType = false;
 }
 
-void FNDIHairStrandsBuffer::Update()
-{
-}
-
 void FNDIHairStrandsBuffer::Transfer(FRDGBuilder& GraphBuilder, const TStaticArray<float, 32 * NumScales>& InParamsScale)
 {
-	if (InfoData.IsValid() && ParamsScaleBuffer.IsValid())
+	if (bNeedResouces && ParamsScaleBuffer.IsValid())
 	{
 		const uint32 ScaleCount = 32 * NumScales;
 		const uint32 ScaleBytes = sizeof(float) * ScaleCount;
@@ -393,7 +403,7 @@ void FNDIHairStrandsBuffer::InitRHI(FRHICommandListBase&)
 {
 	//ReadbackBuffer = new FRHIGPUBufferReadback(TEXT("Hair.PositionOffsetBuffer"));
 	
-	if (InfoData.IsValid())
+	if (bNeedResouces)
 	{
 		FRDGBuilder GraphBuilder(FRHICommandListImmediate::Get());
 		{
@@ -426,8 +436,11 @@ void FNDIHairStrandsBuffer::ReleaseRHI()
 	//delete ReadbackBuffer;
 	//ReadbackBuffer = nullptr;
 	
-	BoundingBoxBuffer.Release();
-	ParamsScaleBuffer.Release();
+	if (bNeedResouces)
+	{
+		BoundingBoxBuffer.Release();
+		ParamsScaleBuffer.Release();
+	}
 }
 
 //------------------------------------------------------------------------------------------------------------
@@ -854,7 +867,6 @@ bool UNiagaraDataInterfaceHairStrands::PerInstanceTick(void* PerInstanceData, FN
 
 	FNDIHairStrandsInfo InfoData;
 	ExtractDatasAndResources(SystemInstance, InfoData);
-	InstanceData->HairStrandsBuffer->Update();
 
 	if (SourceComponent != nullptr)
 	{
diff --git a/Engine/Plugins/Runtime/HairStrands/Source/HairStrandsCore/Public/Niagara/NiagaraDataInterfaceHairStrands.h b/Engine/Plugins/Runtime/HairStrands/Source/HairStrandsCore/Public/Niagara/NiagaraDataInterfaceHairStrands.h
--- a/Engine/Plugins/Runtime/HairStrands/Source/HairStrandsCore/Public/Niagara/NiagaraDataInterfaceHairStrands.h
+++ b/Engine/Plugins/Runtime/HairStrands/Source/HairStrandsCore/Public/Niagara/NiagaraDataInterfaceHairStrands.h
@@ -20,18 +20,7 @@ static const int32 RadiusOffset = 2;
 static const int32 ThicknessOffset = 3;
 
 struct FNDIHairStrandsData;
-
-struct FNDIHairStrandsInfo
-{
-	int32  GroupIndex = 0;
-	int32  LODIndex = 0;
-	uint32 NumControlPoints = 0;
-	uint32 NumCurves = 0;
-	FTransform LocalToWorld = FTransform::Identity;
-	bool bHasValidResources = false;
-
-	bool IsValid() const { return bHasValidResources; }
-};
+struct FNDIHairStrandsInfo;
 
 FHairGroupInstance* GetHairGroupInstance(UGroomComponent* In, int32 InGroupIndex);
 
@@ -43,9 +32,6 @@ struct FNDIHairStrandsBuffer : public FRenderResource
 		const FNDIHairStrandsInfo& In,
 		const TStaticArray<float, 32 * NumScales>& InParamsScale);
 
-	/** Set the asset that will be used to affect the buffer */
-	void Update();
-
 	/** Transfer CPU datas to GPU */
 	void Transfer(FRDGBuilder& GraphBuilder, const TStaticArray<float, 32 * NumScales>& InParamsScale);
 
@@ -76,8 +62,8 @@ struct FNDIHairStrandsBuffer : public FRenderResource
 	/** Mesh LOD that is being used for the root resources */
 	int32 CurrentMeshLOD = INDEX_NONE;
 
-	/** Desc of the group */
-	FNDIHairStrandsInfo InfoData;
+	/** True if the internal resources (BoundingBoxBuffer/ParamsScaleBuffer) needs to be built */
+	bool bNeedResouces = false;
 
 	// For debug only
 	//FRHIGPUBufferReadback* ReadbackBuffer = nullptr;
@@ -289,13 +275,13 @@ struct FNDIHairStrandsData
 	bool ForceReset;
 
 	/** Strands Gpu buffer */
-	FNDIHairStrandsBuffer* HairStrandsBuffer;
+	FNDIHairStrandsBuffer* HairStrandsBuffer = nullptr;
 
 	/** Hair group index */
 	int32 HairGroupIndex = -1;
 
 	/** Hair group instance */
-	FHairGroupInstance* HairGroupInstance;
+	FHairGroupInstance* HairGroupInstance = nullptr;
 
 	/** Source component of the hair group instance */
 	TWeakObjectPtr<class UGroomComponent> HairGroupInstSource;
