Ensure that ReleaseResource is not pending when the resource is destroyed.

#rnx
#tests PIE

diff --git a/Engine/Plugins/Experimental/ChaosNiagara/Source/ChaosNiagara/Private/NiagaraDataInterfaceGeometryCollection.cpp b/Engine/Plugins/Experimental/ChaosNiagara/Source/ChaosNiagara/Private/NiagaraDataInterfaceGeometryCollection.cpp
--- a/Engine/Plugins/Experimental/ChaosNiagara/Source/ChaosNiagara/Private/NiagaraDataInterfaceGeometryCollection.cpp
+++ b/Engine/Plugins/Experimental/ChaosNiagara/Source/ChaosNiagara/Private/NiagaraDataInterfaceGeometryCollection.cpp
@@ -102,10 +102,10 @@ void FNDIGeometryCollectionData::Release()
 {
 	if (AssetBuffer)
 	{
-		BeginReleaseResource(AssetBuffer);
 		ENQUEUE_RENDER_COMMAND(DeleteResource)(
 			[ParamPointerToRelease = AssetBuffer](FRHICommandListImmediate& RHICmdList)
 			{
+				ParamPointerToRelease->ReleaseResource();
 				delete ParamPointerToRelease;
 			});
 		AssetBuffer = nullptr;
diff --git a/Engine/Plugins/FX/Niagara/Source/Niagara/Private/DataInterface/NiagaraDataInterfacePhysicsAsset.cpp b/Engine/Plugins/FX/Niagara/Source/Niagara/Private/DataInterface/NiagaraDataInterfacePhysicsAsset.cpp
--- a/Engine/Plugins/FX/Niagara/Source/Niagara/Private/DataInterface/NiagaraDataInterfacePhysicsAsset.cpp
+++ b/Engine/Plugins/FX/Niagara/Source/Niagara/Private/DataInterface/NiagaraDataInterfacePhysicsAsset.cpp
@@ -438,10 +438,10 @@ void FNDIPhysicsAssetData::Release()
 {
 	if (AssetBuffer)
 	{
-		BeginReleaseResource(AssetBuffer);
 		ENQUEUE_RENDER_COMMAND(DeleteResource)(
 			[ParamPointerToRelease = AssetBuffer](FRHICommandListImmediate& RHICmdList)
 			{
+				ParamPointerToRelease->ReleaseResource();
 				delete ParamPointerToRelease;
 			});
 		AssetBuffer = nullptr;
diff --git a/Engine/Plugins/FX/Niagara/Source/Niagara/Private/NiagaraDataInterfaceRigidMeshCollisionQuery.cpp b/Engine/Plugins/FX/Niagara/Source/Niagara/Private/NiagaraDataInterfaceRigidMeshCollisionQuery.cpp
--- a/Engine/Plugins/FX/Niagara/Source/Niagara/Private/NiagaraDataInterfaceRigidMeshCollisionQuery.cpp
+++ b/Engine/Plugins/FX/Niagara/Source/Niagara/Private/NiagaraDataInterfaceRigidMeshCollisionQuery.cpp
@@ -805,10 +805,10 @@ void FNDIRigidMeshCollisionData::ReleaseBuffers()
 {
 	if (AssetBuffer)
 	{
-		BeginReleaseResource(AssetBuffer);
 		ENQUEUE_RENDER_COMMAND(DeleteResource)(
 			[ParamPointerToRelease = AssetBuffer](FRHICommandListImmediate& RHICmdList)
 			{
+				ParamPointerToRelease->ReleaseResource();
 				delete ParamPointerToRelease;
 			});
 		AssetBuffer = nullptr;
diff --git a/Engine/Plugins/Runtime/HairStrands/Source/HairStrandsCore/Private/Niagara/NiagaraDataInterfaceHairStrands.cpp b/Engine/Plugins/Runtime/HairStrands/Source/HairStrandsCore/Private/Niagara/NiagaraDataInterfaceHairStrands.cpp
--- a/Engine/Plugins/Runtime/HairStrands/Source/HairStrandsCore/Private/Niagara/NiagaraDataInterfaceHairStrands.cpp
+++ b/Engine/Plugins/Runtime/HairStrands/Source/HairStrandsCore/Private/Niagara/NiagaraDataInterfaceHairStrands.cpp
@@ -464,10 +464,10 @@ void FNDIHairStrandsData::Release()
 {
 	if (HairStrandsBuffer)
 	{
-		BeginReleaseResource(HairStrandsBuffer);
 		ENQUEUE_RENDER_COMMAND(DeleteResource)(
 			[ParamPointerToRelease = HairStrandsBuffer](FRHICommandListImmediate& RHICmdList)
 		{
+			ParamPointerToRelease->ReleaseResource();
 			delete ParamPointerToRelease;
 		});
 		HairStrandsBuffer = nullptr;
diff --git a/Engine/Plugins/Runtime/HairStrands/Source/HairStrandsCore/Private/Niagara/NiagaraDataInterfaceVelocityGrid.cpp b/Engine/Plugins/Runtime/HairStrands/Source/HairStrandsCore/Private/Niagara/NiagaraDataInterfaceVelocityGrid.cpp
--- a/Engine/Plugins/Runtime/HairStrands/Source/HairStrandsCore/Private/Niagara/NiagaraDataInterfaceVelocityGrid.cpp
+++ b/Engine/Plugins/Runtime/HairStrands/Source/HairStrandsCore/Private/Niagara/NiagaraDataInterfaceVelocityGrid.cpp
@@ -85,20 +85,21 @@ void FNDIVelocityGridData::Release()
 {
 	if (CurrentGridBuffer)
 	{
-		BeginReleaseResource(CurrentGridBuffer);
 		ENQUEUE_RENDER_COMMAND(DeleteResourceA)(
 			[ParamPointerToRelease = CurrentGridBuffer](FRHICommandListImmediate& RHICmdList)
 		{
+			ParamPointerToRelease->ReleaseResource();
 			delete ParamPointerToRelease;
 		});
 		CurrentGridBuffer = nullptr;
 	}
 	if (DestinationGridBuffer)
 	{
-		BeginReleaseResource(DestinationGridBuffer);
+
 		ENQUEUE_RENDER_COMMAND(DeleteResourceB)(
 			[ParamPointerToRelease = DestinationGridBuffer](FRHICommandListImmediate& RHICmdList)
 		{
+			ParamPointerToRelease->ReleaseResource();
 			delete ParamPointerToRelease;
 		});
 		DestinationGridBuffer = nullptr;
