[GroomBinding]
 - Add support for deferred groom bindings building when dependencies are not ready
 - Calling FinishCompilation on a groom binding now handles dependencies correctly
 - Make sure we unregister and reregister for delegates when dependencies are modified

#rnx

diff --git a/Engine/Plugins/Runtime/HairStrands/Source/HairStrandsCore/Private/GroomBindingAsset.cpp b/Engine/Plugins/Runtime/HairStrands/Source/HairStrandsCore/Private/GroomBindingAsset.cpp
--- a/Engine/Plugins/Runtime/HairStrands/Source/HairStrandsCore/Private/GroomBindingAsset.cpp
+++ b/Engine/Plugins/Runtime/HairStrands/Source/HairStrandsCore/Private/GroomBindingAsset.cpp
@@ -474,6 +474,13 @@ void UGroomBindingAsset::Reset()
 }
 
 #if WITH_EDITORONLY_DATA
+bool UGroomBindingAsset::HasAnyDependenciesCompiling() const
+{
+	return 
+		(GetTargetSkeletalMesh() && GetTargetSkeletalMesh()->IsCompiling()) ||
+		(GetSourceSkeletalMesh() && GetSourceSkeletalMesh()->IsCompiling());
+}
+
 void UGroomBindingAsset::InvalidateBinding(class USkeletalMesh*)
 {
 	InvalidateBinding();
@@ -481,12 +488,8 @@ void UGroomBindingAsset::InvalidateBinding(class USkeletalMesh*)
 
 void UGroomBindingAsset::InvalidateBinding()
 {
-	const bool bDependenciesStillCaching =
-		(GetTargetSkeletalMesh() && GetTargetSkeletalMesh()->IsCompiling()) ||
-		(GetSourceSkeletalMesh() && GetSourceSkeletalMesh()->IsCompiling());
-
 	// Nothing else to do as we're going to be rebuilt when our dependencies finish caching anyway
-	if (bDependenciesStillCaching)
+	if (HasAnyDependenciesCompiling())
 	{
 		return;
 	}
@@ -511,22 +514,14 @@ void UGroomBindingAsset::PostLoad()
 		LocalGroom->ConditionalPostLoad();
 
 	#if WITH_EDITOR
-		USkeletalMesh* LocalTargetSkeletalMesh = GetTargetSkeletalMesh();
-		USkeletalMesh* LocalSourceSkeletalMesh = GetSourceSkeletalMesh();
-
-		if (LocalTargetSkeletalMesh)
+		if (GetTargetSkeletalMesh())
 		{
-			LocalTargetSkeletalMesh->ConditionalPostLoad();
-			LocalTargetSkeletalMesh->OnPostMeshCached().AddUObject(this, &UGroomBindingAsset::InvalidateBinding);
-			bRegisterTargetMeshCallback = true;
+			GetTargetSkeletalMesh()->ConditionalPostLoad();
 		}
 
-		// Do not register the same skeletal mesh twice as it will make us invalidate bindings twice in a row and cause a stall.
-		if (LocalSourceSkeletalMesh && LocalSourceSkeletalMesh != LocalTargetSkeletalMesh)
+		if (GetSourceSkeletalMesh())
 		{
-			LocalSourceSkeletalMesh->ConditionalPostLoad();
-			LocalSourceSkeletalMesh->OnPostMeshCached().AddUObject(this, &UGroomBindingAsset::InvalidateBinding);
-			bRegisterSourceMeshCallback = true;
+			GetSourceSkeletalMesh()->ConditionalPostLoad();
 		}
 
 		if (GetTargetGeometryCache())
@@ -539,9 +534,8 @@ void UGroomBindingAsset::PostLoad()
 			GetSourceGeometryCache()->ConditionalPostLoad();
 		}
 
-		LocalGroom->GetOnGroomAssetChanged().AddUObject(this, &UGroomBindingAsset::InvalidateBinding);
-		LocalGroom->GetOnGroomAssetResourcesChanged().AddUObject(this, &UGroomBindingAsset::InvalidateBinding);
-		bRegisterGroomAssetCallback = true;
+		RegisterSkeletalMeshDelegates();
+		RegisterGroomDelegates();
 
 		InvalidateBinding();
 	#endif
@@ -593,24 +587,9 @@ void UGroomBindingAsset::BeginDestroy()
 	Super::BeginDestroy();
 
 #if WITH_EDITOR
-	if (GetTargetSkeletalMesh() && bRegisterTargetMeshCallback)
-	{
-		GetTargetSkeletalMesh()->OnPostMeshCached().RemoveAll(this);
-		bRegisterTargetMeshCallback = false;
-	}
+	UnregisterSkeletalMeshDelegates();
 
-	if (GetSourceSkeletalMesh() && bRegisterSourceMeshCallback)
-	{
-		GetSourceSkeletalMesh()->OnPostMeshCached().RemoveAll(this);
-		bRegisterSourceMeshCallback = false;
-	}
-
-	if (GetGroom() && bRegisterGroomAssetCallback)
-	{
-		GetGroom()->GetOnGroomAssetResourcesChanged().RemoveAll(this);
-		GetGroom()->GetOnGroomAssetChanged().RemoveAll(this);
-		bRegisterGroomAssetCallback = false;
-	}
+	UnregisterGroomDelegates();
 #endif
 }
 
@@ -1014,6 +993,14 @@ void UGroomBindingAsset::CacheDerivedDatas()
 		FGroomBindingCompilingManager::Get().FinishCompilation({ this });
 	}
 
+	if (HasAnyDependenciesCompiling())
+	{
+		FGroomBindingCompilingManager::Get().AddGroomBindingsWithPendingDependencies({ this });
+		return;
+	}
+
+	TRACE_CPUPROFILER_EVENT_SCOPE(UGroomBindingAsset::CacheDerivedDatas);
+
 	// Flush the rendering commands generated by the detachments.
 	FlushRenderingCommandIfUsed(this);
 
@@ -1028,6 +1015,10 @@ void UGroomBindingAsset::CacheDerivedDatas()
 	AcquireAsyncProperty((uint64)EGroomBindingAsyncProperties::HairGroupPlatformData, EGroomBindingAsyncPropertyLockType::WriteOnly);
 	AcquireAsyncProperty((uint64)EGroomBindingAsyncProperties::GroupInfos, EGroomBindingAsyncPropertyLockType::WriteOnly);
 
+	// Dependency management should have been taken care of already. Just making sure.
+	check(GetTargetSkeletalMesh() == nullptr || GetTargetSkeletalMesh()->IsCompiling() == false);
+	check(GetSourceSkeletalMesh() == nullptr || GetSourceSkeletalMesh()->IsCompiling() == false);
+
 	// Do synchronous build until GeometryCache is thread-safe
 	const bool bIsAsyncBuildSupported = GetTargetGeometryCache() == nullptr && GetSourceGeometryCache() == nullptr;
 
@@ -1231,6 +1222,8 @@ static void CacheDerivedDatas(UGroomBindingAsset* In, const uint32 InGroupIndex,
 
 void UGroomBindingAsset::FinishCacheDerivedDatas(FGroomBindingBuildContext& Context)
 {
+	TRACE_CPUPROFILER_EVENT_SCOPE(UGroomBindingAsset::FinishCacheDerivedDatas);
+
 	ReleaseAsyncProperty();
 
 	if (!IsTemplate() && Context.bReloadResource && IsValid())
@@ -1394,13 +1387,15 @@ FName UGroomBindingAsset::GetAssetPathName(int32 LODIndex)
 #endif
 }
 
-#define DEFINE_GROOMBINDING_MEMBER_ACCESSOR(Type, Access, Name, GetConst, SetConst)\
+#define DEFINE_GROOMBINDING_MEMBER_NAME(Name)\
 	FName UGroomBindingAsset::Get##Name##MemberName()\
 	{\
 		PRAGMA_DISABLE_DEPRECATION_WARNINGS\
 		return FName(TEXT(#Name));\
 		PRAGMA_ENABLE_DEPRECATION_WARNINGS\
-	}\
+	}
+
+#define DEFINE_GROOMBINDING_MEMBER_GETTER(Type, Access, Name, GetConst, SetConst)\
 	UFUNCTION(BlueprintGetter)\
 	GetConst Type Access UGroomBindingAsset::Get##Name() const\
 	{\
@@ -1408,7 +1403,9 @@ FName UGroomBindingAsset::GetAssetPathName(int32 LODIndex)
 		PRAGMA_DISABLE_DEPRECATION_WARNINGS\
 		return Name;\
 		PRAGMA_ENABLE_DEPRECATION_WARNINGS\
-	}\
+	}
+
+#define DEFINE_GROOMBINDING_MEMBER_SETTER(Type, Access, Name, GetConst, SetConst)\
 	UFUNCTION(BlueprintSetter)\
 	void UGroomBindingAsset::Set##Name(SetConst Type Access In##Name)\
 	{\
@@ -1418,17 +1415,137 @@ FName UGroomBindingAsset::GetAssetPathName(int32 LODIndex)
 		PRAGMA_ENABLE_DEPRECATION_WARNINGS\
 	}
 
+#define DEFINE_GROOMBINDING_MEMBER_ACCESSOR(Type, Access, Name, GetConst, SetConst)\
+	DEFINE_GROOMBINDING_MEMBER_NAME(Name) \
+	DEFINE_GROOMBINDING_MEMBER_GETTER(Type, Access, Name, GetConst, SetConst) \
+	DEFINE_GROOMBINDING_MEMBER_SETTER(Type, Access, Name, GetConst, SetConst)
+
 // Define most of the binding member accessor
 DEFINE_GROOMBINDING_MEMBER_ACCESSOR(EGroomBindingMeshType, , GroomBindingType, , );
-DEFINE_GROOMBINDING_MEMBER_ACCESSOR(UGroomAsset, *, Groom, , );
-DEFINE_GROOMBINDING_MEMBER_ACCESSOR(USkeletalMesh, *, SourceSkeletalMesh, , );
-DEFINE_GROOMBINDING_MEMBER_ACCESSOR(USkeletalMesh, *, TargetSkeletalMesh, , );
+DEFINE_GROOMBINDING_MEMBER_NAME(Groom);
+DEFINE_GROOMBINDING_MEMBER_GETTER(UGroomAsset, *, Groom, , );
+DEFINE_GROOMBINDING_MEMBER_NAME(SourceSkeletalMesh);
+DEFINE_GROOMBINDING_MEMBER_NAME(TargetSkeletalMesh);
+DEFINE_GROOMBINDING_MEMBER_GETTER(USkeletalMesh, *, SourceSkeletalMesh, , );
+DEFINE_GROOMBINDING_MEMBER_GETTER(USkeletalMesh, *, TargetSkeletalMesh, , );
 DEFINE_GROOMBINDING_MEMBER_ACCESSOR(UGeometryCache, *, SourceGeometryCache, , );
 DEFINE_GROOMBINDING_MEMBER_ACCESSOR(UGeometryCache, *, TargetGeometryCache, , );
 DEFINE_GROOMBINDING_MEMBER_ACCESSOR(int32, , NumInterpolationPoints, , );
 DEFINE_GROOMBINDING_MEMBER_ACCESSOR(int32, , MatchingSection, , );
 DEFINE_GROOMBINDING_MEMBER_ACCESSOR(TArray<FGoomBindingGroupInfo>, &, GroupInfos, const, const);
 
+#if WITH_EDITORONLY_DATA
+
+void UGroomBindingAsset::UnregisterGroomDelegates()
+{
+	if (GetGroom())
+	{
+		GetGroom()->GetOnGroomAssetResourcesChanged().RemoveAll(this);
+		GetGroom()->GetOnGroomAssetChanged().RemoveAll(this);
+	}
+}
+
+void UGroomBindingAsset::RegisterGroomDelegates()
+{
+	if (GetGroom())
+	{
+		GetGroom()->GetOnGroomAssetChanged().AddUObject(this, &UGroomBindingAsset::InvalidateBinding);
+		GetGroom()->GetOnGroomAssetResourcesChanged().AddUObject(this, &UGroomBindingAsset::InvalidateBinding);
+	}
+}
+
+void UGroomBindingAsset::UnregisterSkeletalMeshDelegates()
+{
+	if (GetTargetSkeletalMesh())
+	{
+		GetTargetSkeletalMesh()->OnPostMeshCached().RemoveAll(this);
+	}
+
+	if (GetSourceSkeletalMesh())
+	{
+		GetSourceSkeletalMesh()->OnPostMeshCached().RemoveAll(this);
+	}
+}
+
+void UGroomBindingAsset::RegisterSkeletalMeshDelegates()
+{
+	if (GetTargetSkeletalMesh())
+	{
+		GetTargetSkeletalMesh()->OnPostMeshCached().AddUObject(this, &UGroomBindingAsset::InvalidateBinding);
+	}
+
+	// Do not register the same skeletal mesh twice as it will make us invalidate bindings twice in a row and cause a stall.
+	if (GetSourceSkeletalMesh() && GetSourceSkeletalMesh() != GetTargetSkeletalMesh())
+	{
+		GetSourceSkeletalMesh()->OnPostMeshCached().AddUObject(this, &UGroomBindingAsset::InvalidateBinding);
+	}
+}
+
+#endif // WITH_EDITORONLY_DATA
+
+void UGroomBindingAsset::SetGroom(UGroomAsset* InGroom)
+{
+	if (GetGroom() == InGroom)
+	{
+		return;
+	}
+
+#if WITH_EDITOR
+	UnregisterGroomDelegates();
+#endif
+
+	WaitUntilAsyncPropertyReleased(EGroomBindingAsyncProperties::Groom, EGroomBindingAsyncPropertyLockType::WriteOnly);
+	PRAGMA_DISABLE_DEPRECATION_WARNINGS
+	Groom = InGroom;
+	PRAGMA_ENABLE_DEPRECATION_WARNINGS
+
+#if WITH_EDITOR
+	RegisterGroomDelegates();
+#endif
+}
+
+void UGroomBindingAsset::SetSourceSkeletalMesh(USkeletalMesh* InSourceSkeletalMesh)
+{
+	if (GetSourceSkeletalMesh() == InSourceSkeletalMesh)
+	{
+		return;
+	}
+
+#if WITH_EDITOR
+	UnregisterSkeletalMeshDelegates();
+#endif
+
+	WaitUntilAsyncPropertyReleased(EGroomBindingAsyncProperties::SourceSkeletalMesh, EGroomBindingAsyncPropertyLockType::WriteOnly);
+	PRAGMA_DISABLE_DEPRECATION_WARNINGS
+	SourceSkeletalMesh = InSourceSkeletalMesh;
+	PRAGMA_ENABLE_DEPRECATION_WARNINGS
+
+#if WITH_EDITOR
+	RegisterSkeletalMeshDelegates();
+#endif
+}
+
+void UGroomBindingAsset::SetTargetSkeletalMesh(USkeletalMesh* InTargetSkeletalMesh)
+{
+	if (GetTargetSkeletalMesh() == InTargetSkeletalMesh)
+	{
+		return;
+	}
+
+#if WITH_EDITOR
+	UnregisterSkeletalMeshDelegates();
+#endif
+
+	WaitUntilAsyncPropertyReleased(EGroomBindingAsyncProperties::TargetSkeletalMesh, EGroomBindingAsyncPropertyLockType::WriteOnly);
+	PRAGMA_DISABLE_DEPRECATION_WARNINGS
+	TargetSkeletalMesh = InTargetSkeletalMesh;
+	PRAGMA_ENABLE_DEPRECATION_WARNINGS
+
+#if WITH_EDITOR
+	RegisterSkeletalMeshDelegates();
+#endif
+}
+
 TArray<FGoomBindingGroupInfo>& UGroomBindingAsset::GetGroupInfos()
 {
 	WaitUntilAsyncPropertyReleased(EGroomBindingAsyncProperties::GroupInfos, EGroomBindingAsyncPropertyLockType::ReadWrite);
diff --git a/Engine/Plugins/Runtime/HairStrands/Source/HairStrandsCore/Private/GroomBindingCompiler.cpp b/Engine/Plugins/Runtime/HairStrands/Source/HairStrandsCore/Private/GroomBindingCompiler.cpp
--- a/Engine/Plugins/Runtime/HairStrands/Source/HairStrandsCore/Private/GroomBindingCompiler.cpp
+++ b/Engine/Plugins/Runtime/HairStrands/Source/HairStrandsCore/Private/GroomBindingCompiler.cpp
@@ -107,6 +107,16 @@ void FGroomBindingCompilingManager::OnPostReachabilityAnalysis()
 
 		TArray<UGroomBindingAsset*> PendingAssets;
 		PendingAssets.Reserve(GetNumRemainingJobs());
+		
+		// Compilation has not started yet... just remove it from our pending list
+		for (auto Iterator = GroomBindingWithPendingDependencies.CreateIterator(); Iterator; ++Iterator)
+		{
+			UGroomBindingAsset* GroomBinding = *Iterator;
+			if (GroomBinding->IsUnreachable())
+			{
+				Iterator.RemoveCurrent();
+			}
+		}
 
 		for (auto Iterator = RegisteredGroomBindingAssets.CreateIterator(); Iterator; ++Iterator)
 		{
@@ -307,7 +317,20 @@ FGroomBindingCompilingManager& FGroomBindingCompilingManager::Get()
 
 int32 FGroomBindingCompilingManager::GetNumRemainingJobs() const
 {
-	return RegisteredGroomBindingAssets.Num();
+	return RegisteredGroomBindingAssets.Num() + GroomBindingWithPendingDependencies.Num();
+}
+
+void FGroomBindingCompilingManager::AddGroomBindingsWithPendingDependencies(TArrayView<UGroomBindingAsset* const> InGroomBindingAssets)
+{
+	TRACE_CPUPROFILER_EVENT_SCOPE(FGroomBindingCompilingManager::AddGroomBindingsWithPendingDependencies);
+	check(IsInGameThread());
+
+	for (UGroomBindingAsset* GroomBindingAsset : InGroomBindingAssets)
+	{
+		GroomBindingWithPendingDependencies.Emplace(GroomBindingAsset);
+	}
+
+	UpdateCompilationNotification();
 }
 
 void FGroomBindingCompilingManager::AddGroomBindings(TArrayView<UGroomBindingAsset* const> InGroomBindingAssets)
@@ -317,6 +340,9 @@ void FGroomBindingCompilingManager::AddGroomBindings(TArrayView<UGroomBindingAss
 
 	for (UGroomBindingAsset* GroomBindingAsset : InGroomBindingAssets)
 	{
+		// If the compilation is launched while we still have it in our list, we don't want to schedule it again so remove it now.
+		GroomBindingWithPendingDependencies.Remove(GroomBindingAsset);
+
 		check(GroomBindingAsset->AsyncTask != nullptr);
 		RegisteredGroomBindingAssets.Emplace(GroomBindingAsset);
 
@@ -332,6 +358,25 @@ void FGroomBindingCompilingManager::FinishCompilation(TArrayView<UGroomBindingAs
 
 	check(IsInGameThread());
 
+	TSet<USkinnedAsset*> Dependencies;
+	for (UGroomBindingAsset* GroomBindingAsset : InGroomBindingAssets)
+	{
+		if (GroomBindingAsset->GetTargetSkeletalMesh() && GroomBindingAsset->GetTargetSkeletalMesh()->IsCompiling())
+		{
+			Dependencies.Emplace(GroomBindingAsset->GetTargetSkeletalMesh());
+		}
+
+		if (GroomBindingAsset->GetSourceSkeletalMesh() && GroomBindingAsset->GetSourceSkeletalMesh()->IsCompiling())
+		{
+			Dependencies.Emplace(GroomBindingAsset->GetSourceSkeletalMesh());
+		}
+	}
+
+	FSkinnedAssetCompilingManager::Get().FinishCompilation(Dependencies.Array());
+	
+	// Now that dependencies have finished, we can launch the compilations for the groom bindings.
+	SchedulePendingCompilations();
+
 	TArray<UGroomBindingAsset*> PendingGroomBindingAssets;
 	PendingGroomBindingAssets.Reserve(InGroomBindingAssets.Num());
 
@@ -385,6 +430,26 @@ void FGroomBindingCompilingManager::FinishCompilation(TArrayView<UGroomBindingAs
 	}
 }
 
+void FGroomBindingCompilingManager::SchedulePendingCompilations()
+{
+	TArray<UGroomBindingAsset*> ReadyToSchedule;
+	for (auto It = GroomBindingWithPendingDependencies.CreateIterator(); It; ++It)
+	{
+		UGroomBindingAsset* GroomBindingAsset = *It;
+		if (!GroomBindingAsset->HasAnyDependenciesCompiling())
+		{
+			ReadyToSchedule.Emplace(GroomBindingAsset);
+			It.RemoveCurrent();
+		}
+	}
+
+	// Call CacheDerivedDatas again so it's scheduled for real this time.
+	for (UGroomBindingAsset* GroomBindingAsset : ReadyToSchedule)
+	{
+		GroomBindingAsset->CacheDerivedDatas();
+	}
+}
+
 void FGroomBindingCompilingManager::Reschedule()
 {
 	// Implement if we ever want to prioritize nearest visible grooms first
@@ -489,6 +554,8 @@ void FGroomBindingCompilingManager::ProcessAsyncTasks(bool bLimitExecutionTime)
 	FinishCompilationsForGame();
 
 	Reschedule();
+	
+	SchedulePendingCompilations();
 
 	ProcessGroomBindingAssets(bLimitExecutionTime);
 
@@ -497,4 +564,4 @@ void FGroomBindingCompilingManager::ProcessAsyncTasks(bool bLimitExecutionTime)
 
 #endif // #if WITH_EDITOR
 
-#undef LOCTEXT_NAMESPACE
\ No newline at end of file
+#undef LOCTEXT_NAMESPACE
diff --git a/Engine/Plugins/Runtime/HairStrands/Source/HairStrandsCore/Public/GroomBindingAsset.h b/Engine/Plugins/Runtime/HairStrands/Source/HairStrandsCore/Public/GroomBindingAsset.h
--- a/Engine/Plugins/Runtime/HairStrands/Source/HairStrandsCore/Public/GroomBindingAsset.h
+++ b/Engine/Plugins/Runtime/HairStrands/Source/HairStrandsCore/Public/GroomBindingAsset.h
@@ -436,6 +436,8 @@ public:
 
 	void CacheDerivedDatas();
 
+	bool HasAnyDependenciesCompiling() const;
+
 	virtual void BeginCacheForCookedPlatformData(const ITargetPlatform* TargetPlatform);
 	virtual void ClearAllCachedCookedPlatformData();
 	TArray<FHairGroupPlatformData>* GetCachedCookedPlatformData(const ITargetPlatform* TargetPlatform);
@@ -451,9 +453,11 @@ public:
 private:
 	TArray<FCachedCookedPlatformData*> CachedCookedPlatformDatas;
 
-	bool bRegisterSourceMeshCallback = false;
-	bool bRegisterTargetMeshCallback = false;
-	bool bRegisterGroomAssetCallback = false;
+	void RegisterGroomDelegates();
+	void UnregisterGroomDelegates();
+	void RegisterSkeletalMeshDelegates();
+	void UnregisterSkeletalMeshDelegates();
+
 	TArray<FString> CachedDerivedDataKey;
 #endif
 #if WITH_EDITOR
diff --git a/Engine/Plugins/Runtime/HairStrands/Source/HairStrandsCore/Public/GroomBindingCompiler.h b/Engine/Plugins/Runtime/HairStrands/Source/HairStrandsCore/Public/GroomBindingCompiler.h
--- a/Engine/Plugins/Runtime/HairStrands/Source/HairStrandsCore/Public/GroomBindingCompiler.h
+++ b/Engine/Plugins/Runtime/HairStrands/Source/HairStrandsCore/Public/GroomBindingCompiler.h
@@ -37,6 +37,11 @@ public:
 	 */
 	HAIRSTRANDSCORE_API void AddGroomBindings(TArrayView<UGroomBindingAsset* const> InGroomBindingAssets);
 
+	/**
+	 * Register groom binding assets to compile once their dependencies are finished compiling.
+	 */
+	HAIRSTRANDSCORE_API void AddGroomBindingsWithPendingDependencies(TArrayView<UGroomBindingAsset* const> InGroomBindingAssets);
+
 	/** 
 	 * Blocks until completion of the requested groom binding assets.
 	 */
@@ -81,6 +86,7 @@ private:
 	friend class FAssetCompilingManager;
 	
 	bool bHasShutdown = false;
+	TSet<UGroomBindingAsset*> GroomBindingWithPendingDependencies;
 	TSet<UGroomBindingAsset*> RegisteredGroomBindingAssets;
 	TMultiMap<class USkeletalMesh*, UGroomBindingAsset*> RegisteredSkeletalMeshes;
 	TMultiMap<class UGroomAsset*, UGroomBindingAsset*> RegisteredGroomAssets;
@@ -92,7 +98,7 @@ private:
 	
 	void AttachDependencies(UGroomBindingAsset* GroomBindingAsset);
 	void DetachDependencies(UGroomBindingAsset* GroomBindingAsset);
-
+	void SchedulePendingCompilations();
 	void PostCompilation(UGroomBindingAsset* GroomBindingAsset);
 	void PostCompilation(TArrayView<UGroomBindingAsset* const> InGroomBindingAssets);
 	void OnPostReachabilityAnalysis();
