Save Hair Strands Textures generation options into groom strands to avoid re-entering them each time textures are generated.

diff --git a/Engine/Plugins/Runtime/HairStrands/Source/HairStrandsCore/Private/GroomTextureBuilder.cpp b/Engine/Plugins/Runtime/HairStrands/Source/HairStrandsCore/Private/GroomTextureBuilder.cpp
--- a/Engine/Plugins/Runtime/HairStrands/Source/HairStrandsCore/Private/GroomTextureBuilder.cpp
+++ b/Engine/Plugins/Runtime/HairStrands/Source/HairStrandsCore/Private/GroomTextureBuilder.cpp
@@ -31,9 +31,6 @@
 #define DEBUG_TEXTURE_GENERATION 0
 #define DEBUG_TEXTURE_RENDERDOC 0
 
-static int32 GHairStrandsTextureDilationPassCount = 8;
-static FAutoConsoleVariableRef CVarHairStrandsTextureDilationPassCount(TEXT("r.HairStrands.Textures.DilationCount"), GHairStrandsTextureDilationPassCount, TEXT("Number of dilation pass run onto the generated hair strands textures (Default:8)."));
-
 ///////////////////////////////////////////////////////////////////////////////////////////////////
 
 DEFINE_LOG_CATEGORY_STATIC(LogGroomTextureBuilder, Log, All);
@@ -1242,7 +1239,7 @@ static bool TraceTextures(
 	return true;
 }
 
-static FHairStrandsRDGTextures DilateTextures(FRDGBuilder& GraphBuilder, FGlobalShaderMap* ShaderMap, FHairStrandsRDGTextures& In)
+static FHairStrandsRDGTextures DilateTextures(FRDGBuilder& GraphBuilder, FGlobalShaderMap* ShaderMap, FHairStrandsRDGTextures& In, uint32 Dilation)
 {
 	const FIntPoint OutputResolution = In.Texture->Desc.Extent;
 
@@ -1256,8 +1253,7 @@ static FHairStrandsRDGTextures DilateTextures(FRDGBuilder& GraphBuilder, FGlobal
 
 	uint32 SourceIndex = 0;
 	uint32 TargetIndex = 0;
-	const uint32 DilationPassCount = FMath::Max(0, GHairStrandsTextureDilationPassCount);
-	for (uint32 DilationIt=0; DilationIt< DilationPassCount; ++DilationIt)
+	for (uint32 DilationIt=0; DilationIt < Dilation; ++DilationIt)
 	{
 		TargetIndex = (SourceIndex + 1) % 2;
 
@@ -1400,7 +1396,7 @@ void RunHairStrandsTexturesQueries(FRDGBuilder& GraphBuilder, FGlobalShaderMap*
 			#if !DEBUG_TEXTURE_GENERATION
 				if (Textures.IsCompleted())
 				{
-					Textures = DilateTextures(GraphBuilder, ShaderMap, Textures);
+					Textures = DilateTextures(GraphBuilder, ShaderMap, Textures, R->Info.Dilation);
 					CopyTextures(GraphBuilder, ShaderMap, R->Info.Layout, Textures, R->Readback);
 					GStrandsTexturesReadbacks.Enqueue(R);
 				}
@@ -1460,7 +1456,7 @@ void RunHairStrandsTexturesQueries(FRDGBuilder& GraphBuilder, FGlobalShaderMap*
 
 			if (Textures.IsCompleted())
 			{
-				Textures = DilateTextures(GraphBuilder, ShaderMap, Textures);
+				Textures = DilateTextures(GraphBuilder, ShaderMap, Textures, R->Info.Dilation);
 				CopyTextures(GraphBuilder, ShaderMap, R->Info.Layout, Textures, R->Readback);
 				GStrandsTexturesReadbacks.Enqueue(R);
 			}
diff --git a/Engine/Plugins/Runtime/HairStrands/Source/HairStrandsCore/Public/GroomAssetImportData.h b/Engine/Plugins/Runtime/HairStrands/Source/HairStrandsCore/Public/GroomAssetImportData.h
--- a/Engine/Plugins/Runtime/HairStrands/Source/HairStrandsCore/Public/GroomAssetImportData.h
+++ b/Engine/Plugins/Runtime/HairStrands/Source/HairStrandsCore/Public/GroomAssetImportData.h
@@ -16,4 +16,7 @@ public:
 
 	UPROPERTY()
 	TObjectPtr<class UGroomImportOptions> ImportOptions;
+
+	UPROPERTY()
+	TObjectPtr<class UGroomCreateStrandsTexturesOptions> HairStrandsTexturesOptions;
 };
diff --git a/Engine/Plugins/Runtime/HairStrands/Source/HairStrandsCore/Public/GroomCreateStrandsTexturesOptions.h b/Engine/Plugins/Runtime/HairStrands/Source/HairStrandsCore/Public/GroomCreateStrandsTexturesOptions.h
--- a/Engine/Plugins/Runtime/HairStrands/Source/HairStrandsCore/Public/GroomCreateStrandsTexturesOptions.h
+++ b/Engine/Plugins/Runtime/HairStrands/Source/HairStrandsCore/Public/GroomCreateStrandsTexturesOptions.h
@@ -76,4 +76,8 @@ public:
 	/** Groom index which should be baked into the textures. When the array is empty, all groups will be included (Default). */
 	UPROPERTY(EditAnywhere, BlueprintReadWrite, meta = (ShowOnlyInnerProperties), Category = Options)
 	TArray<int32> GroupIndex;
+
+	/** Number pixels expanded by the post-process dilation (0..64). */
+	UPROPERTY(EditAnywhere, BlueprintReadWrite, meta = (ShowOnlyInnerProperties, ClampMin = 0, ClampMax = 32, UIMin = 0, UIMax = 32), Category = Options)
+	int32 Dilation = 8;
 };
diff --git a/Engine/Plugins/Runtime/HairStrands/Source/HairStrandsCore/Public/GroomTextureBuilder.h b/Engine/Plugins/Runtime/HairStrands/Source/HairStrandsCore/Public/GroomTextureBuilder.h
--- a/Engine/Plugins/Runtime/HairStrands/Source/HairStrandsCore/Public/GroomTextureBuilder.h
+++ b/Engine/Plugins/Runtime/HairStrands/Source/HairStrandsCore/Public/GroomTextureBuilder.h
@@ -44,6 +44,7 @@ struct FStrandsTexturesInfo
 	uint32 UVChannelIndex = 0;
 	float MaxTracingDistance = 1;
 	int32 TracingDirection = 1;
+	uint32 Dilation = 0;
 	TArray<int32> GroupIndices;
 };
 
diff --git a/Engine/Plugins/Runtime/HairStrands/Source/HairStrandsEditor/Private/GroomActions.cpp b/Engine/Plugins/Runtime/HairStrands/Source/HairStrandsEditor/Private/GroomActions.cpp
--- a/Engine/Plugins/Runtime/HairStrands/Source/HairStrandsEditor/Private/GroomActions.cpp
+++ b/Engine/Plugins/Runtime/HairStrands/Source/HairStrandsEditor/Private/GroomActions.cpp
@@ -359,11 +359,26 @@ void ExecuteCreateStrandsTextures(const FToolMenuContext& InContext)
 	}
 
 	// Duplicate the options to prevent dirtying the asset when they are modified but the rebuild is cancelled
-	UGroomCreateStrandsTexturesOptions* CurrentOptions = NewObject<UGroomCreateStrandsTexturesOptions>();
 	for (UGroomAsset* GroomAsset : GroomAssets)
 	{
-		if (GroomAsset && GroomAsset->IsValid() && CurrentOptions)
+		if (GroomAsset && GroomAsset->IsValid())
 		{
+			UGroomCreateStrandsTexturesOptions* CurrentOptions = nullptr;
+			UGroomAssetImportData* GroomAssetImportData = Cast<UGroomAssetImportData>(GroomAsset->AssetImportData);
+			if (GroomAssetImportData)
+			{
+				// Duplicate the options to prevent dirtying the asset when they are modified but the rebuild is cancelled
+				if (GroomAssetImportData->HairStrandsTexturesOptions)
+				{
+					CurrentOptions = DuplicateObject<UGroomCreateStrandsTexturesOptions>(GroomAssetImportData->HairStrandsTexturesOptions, nullptr);
+				}
+			}
+
+			if (CurrentOptions == nullptr)
+			{
+				CurrentOptions = NewObject<UGroomCreateStrandsTexturesOptions>();
+			}
+
 			TSharedPtr<SGroomCreateStrandsTexturesOptionsWindow> GroomOptionWindow = SGroomCreateStrandsTexturesOptionsWindow::DisplayCreateStrandsTexturesOptions(CurrentOptions);
 			if (!GroomOptionWindow->ShouldCreate())
 			{
@@ -402,7 +417,7 @@ void ExecuteCreateStrandsTextures(const FToolMenuContext& InContext)
 				
 				FStrandsTexturesInfo Info;
 				Info.Layout = CurrentOptions->Layout;
-				Info.GroomAsset   = GroomAsset;
+				Info.GroomAsset  = GroomAsset;
 				Info.TracingDirection = SignDirection;
 				Info.MaxTracingDistance = MaxDistance;
 				Info.Resolution = FMath::RoundUpToPowerOfTwo(FMath::Max(256, CurrentOptions->Resolution));
@@ -411,6 +426,7 @@ void ExecuteCreateStrandsTextures(const FToolMenuContext& InContext)
 				Info.UVChannelIndex= FMath::Max(0, CurrentOptions->UVChannelIndex);
 				Info.SkeletalMesh = SkeletalMesh;
 				Info.StaticMesh = StaticMesh;
+				Info.Dilation = FMath::Clamp(CurrentOptions->Dilation, 0, 64);
 				if (CurrentOptions->GroupIndex.Num())
 				{
 					Info.GroupIndices = CurrentOptions->GroupIndex;
@@ -426,6 +442,18 @@ void ExecuteCreateStrandsTextures(const FToolMenuContext& InContext)
 				if (Output.IsValid())
 				{
 					FGroomTextureBuilder::BuildStrandsTextures(Info, Output);
+
+					// Save settings used for generating these textures
+					if (GroomAssetImportData)
+					{
+						if (GroomAssetImportData->HairStrandsTexturesOptions || GroomAssetImportData->HairStrandsTexturesOptions != CurrentOptions)
+						{
+							// Move the transient ImportOptions to the asset package and set it on the GroomAssetImportData for serialization
+							CurrentOptions->Rename(nullptr, GroomAssetImportData);
+							GroomAssetImportData->HairStrandsTexturesOptions = CurrentOptions;
+							GroomAsset->MarkPackageDirty();
+						}
+					}
 				}
 			}
 		}
