Fixed a memory leak when streaming HairStrands bulk data

diff --git a/Engine/Plugins/Runtime/HairStrands/Source/HairStrandsCore/Private/HairStrandsDatas.cpp b/Engine/Plugins/Runtime/HairStrands/Source/HairStrandsCore/Private/HairStrandsDatas.cpp
--- a/Engine/Plugins/Runtime/HairStrands/Source/HairStrandsCore/Private/HairStrandsDatas.cpp
+++ b/Engine/Plugins/Runtime/HairStrands/Source/HairStrandsCore/Private/HairStrandsDatas.cpp
@@ -536,7 +536,10 @@ const uint8* FHairStreamingRequest::FChunk::GetData() const
 void FHairStreamingRequest::FChunk::Release() 
 { 
 #if !WITH_EDITORONLY_DATA
+	// Take ownership of the buffer
 	TIoStatusOr<uint8*> Out = Data_IO.Release();
+	// Free the buffer
+	FMemory::Free(Out.ConsumeValueOrDie());
 #else
 	Data_DDC.Reset(); 
 #endif
