Fix hair cache interpolation

diff --git a/Engine/Plugins/FX/Niagara/Source/Niagara/Private/NiagaraSimCache.cpp b/Engine/Plugins/FX/Niagara/Source/Niagara/Private/NiagaraSimCache.cpp
--- a/Engine/Plugins/FX/Niagara/Source/Niagara/Private/NiagaraSimCache.cpp
+++ b/Engine/Plugins/FX/Niagara/Source/Niagara/Private/NiagaraSimCache.cpp
@@ -875,6 +875,11 @@ bool UNiagaraSimCache::EndWrite(bool bAllowAnalytics)
 bool UNiagaraSimCache::CanRead(UNiagaraSystem* NiagaraSystem)
 {
 	check(IsInGameThread());
+	
+	if ( !NiagaraSystem )
+    {
+    	return false;
+    }
 
 	if ( NiagaraSystem != SoftNiagaraSystem.Get() )
 	{
diff --git a/Engine/Plugins/Runtime/HairStrands/Source/HairStrandsCore/Private/Niagara/NiagaraDataInterfaceHairStrands.cpp b/Engine/Plugins/Runtime/HairStrands/Source/HairStrandsCore/Private/Niagara/NiagaraDataInterfaceHairStrands.cpp
--- a/Engine/Plugins/Runtime/HairStrands/Source/HairStrandsCore/Private/Niagara/NiagaraDataInterfaceHairStrands.cpp
+++ b/Engine/Plugins/Runtime/HairStrands/Source/HairStrandsCore/Private/Niagara/NiagaraDataInterfaceHairStrands.cpp
@@ -969,10 +969,7 @@ static void InterpolateGroomGuides(
 	FNDIHairStrandsBuffer* HairStrandsBuffer,
 	FHairGroupInstance* HairGroupInstance,
 	const uint32 StrandsSize, 
-	const bool bHasSkinningBinding, 
-	const bool bHasValidGeometry, 
-	const FMatrix44f& WorldToLocal, 
-	int32 MeshLODIndex)
+	const FMatrix44f& WorldToLocal)
 {
 	FHairGroupInstanceRDG InstanceRDG = Convert(GraphBuilder, HairGroupInstance);
 
@@ -981,7 +978,7 @@ static void InterpolateGroomGuides(
 	const bool bIsDeformedValid = bIsHairValid && InstanceRDG.IsDeformedValid();
 	const bool bIsRootValid = bIsHairValid && InstanceRDG.IsRootValid();
 
-	if(bIsRestValid && bIsDeformedValid && bHasValidGeometry)
+	if(bIsRestValid && bIsDeformedValid)
 	{
 		FInterpolateGroomGuidesCS::FPermutationDomain InterpolationDomain;
 		InterpolationDomain.Set<FInterpolateGroomGuidesCS::FInterpolationType>(!bIsRootValid);
@@ -1038,37 +1035,30 @@ void UNiagaraDataInterfaceHairStrands::SimCachePostReadFrame(void* OptionalPerIn
 				
 		if (RestPositionIndex != INDEX_NONE && NodePositionIndex != INDEX_NONE)
 		{
-			FNDIHairStrandsData* InstanceData = static_cast<FNDIHairStrandsData*>(OptionalPerInstanceData);
-			
 			if( EmitterDatas.GetCompiledData().VariableLayouts.IsValidIndex(NodePositionIndex) &&
-				EmitterDatas.GetCompiledData().VariableLayouts.IsValidIndex(RestPositionIndex) && InstanceData)
+				EmitterDatas.GetCompiledData().VariableLayouts.IsValidIndex(RestPositionIndex))
 			{
 				const uint32 NodePositionComponent = EmitterDatas.GetCompiledData().VariableLayouts[NodePositionIndex].GetFloatComponentStart();
 				const uint32 RestPositionComponent = EmitterDatas.GetCompiledData().VariableLayouts[RestPositionIndex].GetFloatComponentStart();
-				
 				FNiagaraDataBuffer* ParticlesBuffer = EmitterDatas.GetCurrentData();
-				FNDIHairStrandsBuffer* HairstrandsBuffer = InstanceData->HairStrandsBuffer;
 				
-				if(HairstrandsBuffer && InstanceData->HairGroupInstance)
-				{
-					const FMatrix44f WorldToLocal = ComputeWorldTransform(InstanceData).Inverse();
-					
-					const uint32 StrandsSize = InstanceData->StrandsSize;
-					const bool bHasSkinningBinding = InstanceData->HairGroupInstance->BindingType == EHairBindingType::Skinning;
-					const bool bHasValidGeometry = InstanceData->HairGroupInstance->GeometryType != EHairGeometryType::NoneGeometry;
+				FNDIHairStrandsData ProxyData;
+				ProxyData.CopyDatas(static_cast<FNDIHairStrandsData*>(OptionalPerInstanceData));
 				
-					FHairGroupInstance* LocalHairInstance = InstanceData->HairGroupInstance;
+				ENQUEUE_RENDER_COMMAND(NiagaraInterpolateGroomSimCache) (
+					[ProxyData, ParticlesBuffer,
+						NodePositionComponent, RestPositionComponent](FRHICommandListImmediate& RHICmdList)
+					{
+						FMemMark MemMark(FMemStack::Get());
+						FRDGBuilder GraphBuilder(RHICmdList);
 
-					ENQUEUE_RENDER_COMMAND(NiagaraInterpolateGroomSimCache)(
-						[ParticlesBuffer, HairstrandsBuffer, NodePositionComponent, RestPositionComponent, StrandsSize, bHasSkinningBinding, bHasValidGeometry, WorldToLocal, LocalHairInstance](FRHICommandListImmediate& RHICmdList)
-						{
-							FMemMark MemMark(FMemStack::Get());
-							FRDGBuilder GraphBuilder(RHICmdList);
-							InterpolateGroomGuides(GraphBuilder, ParticlesBuffer, NodePositionComponent, RestPositionComponent, HairstrandsBuffer, LocalHairInstance, StrandsSize, bHasSkinningBinding, bHasValidGeometry, WorldToLocal, LocalHairInstance->Debug.MeshLODIndex);
-							GraphBuilder.Execute();
-						}
-					);
-				}
+						const FMatrix44f WorldToLocal = ComputeWorldTransform(&ProxyData).Inverse();
+
+						InterpolateGroomGuides(GraphBuilder, ParticlesBuffer, NodePositionComponent, RestPositionComponent,
+							ProxyData.HairStrandsBuffer, ProxyData.HairGroupInstance, ProxyData.StrandsSize, WorldToLocal);
+						GraphBuilder.Execute();
+					}
+				);
 			}
 		}
 	}
