Added StreamIn function to Groom Binding Assets, to allow users to force resources to stream in from disk

diff --git a/Engine/Plugins/Runtime/HairStrands/Source/HairStrandsCore/Private/GroomBindingAsset.cpp b/Engine/Plugins/Runtime/HairStrands/Source/HairStrandsCore/Private/GroomBindingAsset.cpp
--- a/Engine/Plugins/Runtime/HairStrands/Source/HairStrandsCore/Private/GroomBindingAsset.cpp
+++ b/Engine/Plugins/Runtime/HairStrands/Source/HairStrandsCore/Private/GroomBindingAsset.cpp
@@ -453,6 +453,20 @@ void UGroomBindingAsset::FlushRenderingCommandIfUsed(const UGroomBindingAsset* I
 	}
 }
 
+void UGroomBindingAsset::StreamInForCPUAccess(bool bWait)
+{
+	check(IsValid());
+
+	for (UGroomBindingAsset::FHairGroupPlatformData& Data : GetHairGroupsPlatformData())
+	{
+		for (FHairStrandsRootBulkData& BulkData : Data.RenRootBulkDatas)
+		{
+			FHairStreamingRequest R;
+			R.Request(HAIR_MAX_NUM_CURVE_PER_GROUP, HAIR_MAX_NUM_POINT_PER_GROUP, BulkData, bWait, true /*bFillBulkdata*/, false /*bWarmCache*/, GetFName());
+		}
+	}
+}
+
 void UGroomBindingAsset::Reset()
 {
 	FlushRenderingCommandIfUsed(this);
diff --git a/Engine/Plugins/Runtime/HairStrands/Source/HairStrandsCore/Private/GroomBindingBuilder.cpp b/Engine/Plugins/Runtime/HairStrands/Source/HairStrandsCore/Private/GroomBindingBuilder.cpp
--- a/Engine/Plugins/Runtime/HairStrands/Source/HairStrandsCore/Private/GroomBindingBuilder.cpp
+++ b/Engine/Plugins/Runtime/HairStrands/Source/HairStrandsCore/Private/GroomBindingBuilder.cpp
@@ -2026,6 +2026,8 @@ void CopyToBulkData(FHairBulkContainer& Out, const TArray<typename TFormatType::
 template<typename TFormatType>
 void CopyFromBulkData(TArray<typename TFormatType::Type>& Out, const FByteBulkData& In)
 {
+	check(In.IsBulkDataLoaded());
+
 	const uint32 InDataSize = In.GetBulkDataSize();
 	const uint32 ElementCount = InDataSize / sizeof(typename TFormatType::BulkType);
 	Out.SetNum(ElementCount);
@@ -2119,12 +2121,12 @@ static void BuildRootBulkData(
 }
 
 // Convert "root data" <- "root bulk data"
+//
+// Bulk data must be loaded before calling this function
 static void BuildRootData(
 	FHairStrandsRootData& Out,
 	const FHairStrandsRootBulkData& In)
 {
-	// TODO: do we need to force load the data into the bulk data prior to running the convertion (bulk -> data)?
-
 	Out.RootCount = In.Header.RootCount;
 	Out.PointCount = In.Header.PointCount;
 	{
diff --git a/Engine/Plugins/Runtime/HairStrands/Source/HairStrandsCore/Private/GroomRBFDeformer.cpp b/Engine/Plugins/Runtime/HairStrands/Source/HairStrandsCore/Private/GroomRBFDeformer.cpp
--- a/Engine/Plugins/Runtime/HairStrands/Source/HairStrandsCore/Private/GroomRBFDeformer.cpp
+++ b/Engine/Plugins/Runtime/HairStrands/Source/HairStrandsCore/Private/GroomRBFDeformer.cpp
@@ -657,11 +657,17 @@ static float SampleMaskTexture(const FVector2f& InUV, const FUintPoint& InResolu
 	return 1.0f;
 }
 
-void FGroomRBFDeformer::GetRBFDeformedGroomAsset(const UGroomAsset* InGroomAsset, const UGroomBindingAsset* BindingAsset, FTextureSource* MaskTextureSource, const float MaskScale, UGroomAsset* OutGroomAsset)
+void FGroomRBFDeformer::GetRBFDeformedGroomAsset(const UGroomAsset* InGroomAsset, UGroomBindingAsset* BindingAsset, FTextureSource* MaskTextureSource, const float MaskScale, UGroomAsset* OutGroomAsset)
 {
 #if WITH_EDITORONLY_DATA
 	if (InGroomAsset && BindingAsset && BindingAsset->GetTargetSkeletalMesh() && BindingAsset->GetSourceSkeletalMesh())
 	{
+		// Ensure the bulk data is streamed in and accessible from the CPU
+		{
+			const bool bWait = true;
+			BindingAsset->StreamInForCPUAccess(bWait);
+		}
+
 		// Get the skel. mesh render data for the current platform
 		// This similar to how we fetch skel. mesh data when building groom binding. This is requires in order to get 
 		// identical skel. mesh render data and ensure that the hair deformation is done correctly.
diff --git a/Engine/Plugins/Runtime/HairStrands/Source/HairStrandsCore/Public/GroomBindingAsset.h b/Engine/Plugins/Runtime/HairStrands/Source/HairStrandsCore/Public/GroomBindingAsset.h
--- a/Engine/Plugins/Runtime/HairStrands/Source/HairStrandsCore/Public/GroomBindingAsset.h
+++ b/Engine/Plugins/Runtime/HairStrands/Source/HairStrandsCore/Public/GroomBindingAsset.h
@@ -347,6 +347,15 @@ public:
 	/** Release the hair strands resource. */
 	void ReleaseResource(bool bResetLoadedSize);
 
+	/**
+	 * Stream in all of this binding's streamable resources and make them accessible from the CPU.
+	 *
+	 * This is only needed for advanced use cases involving editing grooms or binding data.
+	 *
+	 * @param bWait If true, this call will block until the resources have been streamed in
+	 */
+	void StreamInForCPUAccess(bool bWait);
+
 	void Reset();
 
 	/** Return true if the binding asset is valid, i.e., correctly built and loaded. */
diff --git a/Engine/Plugins/Runtime/HairStrands/Source/HairStrandsCore/Public/GroomRBFDeformer.h b/Engine/Plugins/Runtime/HairStrands/Source/HairStrandsCore/Public/GroomRBFDeformer.h
--- a/Engine/Plugins/Runtime/HairStrands/Source/HairStrandsCore/Public/GroomRBFDeformer.h
+++ b/Engine/Plugins/Runtime/HairStrands/Source/HairStrandsCore/Public/GroomRBFDeformer.h
@@ -11,7 +11,10 @@ struct FTextureSource;
 struct HAIRSTRANDSCORE_API FGroomRBFDeformer
 {
 	// Return a new GroomAsset with the RBF deformation from the BindingAsset baked into it
-	void GetRBFDeformedGroomAsset(const UGroomAsset* InGroomAsset, const UGroomBindingAsset* BindingAsset, FTextureSource* MaskTextureSource, const float MaskScale, UGroomAsset* OutGroomAsset);
+	//
+	// BindingAsset is non-const because its resources will be streamed in from disk if necessary.
+	// Apart from that, it won't be modified.
+	void GetRBFDeformedGroomAsset(const UGroomAsset* InGroomAsset, UGroomBindingAsset* BindingAsset, FTextureSource* MaskTextureSource, const float MaskScale, UGroomAsset* OutGroomAsset);
 
 	static uint32 GetEntryCount(uint32 InSampleCount);
 	static uint32 GetWeightCount(uint32 InSampleCount);
