Fix 0-byte request done by groom system on cooked build, causing crashes on loose-file/non-packaged builds.

diff --git a/Engine/Plugins/Runtime/HairStrands/Source/HairStrandsCore/Private/HairStrandsDatas.cpp b/Engine/Plugins/Runtime/HairStrands/Source/HairStrandsCore/Private/HairStrandsDatas.cpp
--- a/Engine/Plugins/Runtime/HairStrands/Source/HairStrandsCore/Private/HairStrandsDatas.cpp
+++ b/Engine/Plugins/Runtime/HairStrands/Source/HairStrandsCore/Private/HairStrandsDatas.cpp
@@ -278,19 +278,36 @@ void FHairStrandsBulkCommon::FQuery::Add(FHairBulkContainer& In, const TCHAR* In
 		}
 		check (InSize >= InOffset);
 
-		// 1. Add chunk request to the streaming request. The chunk will hold the request result.
 		check(StreamingRequest);
-		FHairStreamingRequest::FChunk& Chunk = StreamingRequest->Chunks.AddDefaulted_GetRef();
-		Chunk.Status 	= FHairStreamingRequest::FChunk::EStatus::Pending;
-		Chunk.Container = &In;
-		Chunk.Size 		= InSize - InOffset;
-		Chunk.Offset 	= InOffset;
-		Chunk.TotalSize = InSize;
-		In.ChunkRequest = &Chunk;
 
-		// 2. Fill in actual DDC request
-		check(OutReadIO);
-		OutReadIO->Read(In.Data, InOffset, Chunk.Size, EAsyncIOPriorityAndFlags::AIOP_Normal, Chunk.Data_IO);
+		// IO system does not support 0-bytes requests, so bypass them.
+		const bool bHasValidSize = int32(InSize) - int32(InOffset) > 0;
+		if (bHasValidSize)
+		{
+			// 1. Add chunk request to the streaming request. The chunk will hold the request result.
+			FHairStreamingRequest::FChunk& Chunk = StreamingRequest->Chunks.AddDefaulted_GetRef();
+			Chunk.Status 	= FHairStreamingRequest::FChunk::EStatus::Pending;
+			Chunk.Container = &In;
+			Chunk.Size 		= InSize - InOffset;
+			Chunk.Offset 	= InOffset;
+			Chunk.TotalSize = InSize;
+			In.ChunkRequest = &Chunk;
+
+			// 2. Fill in actual DDC request
+			check(OutReadIO);
+			OutReadIO->Read(In.Data, InOffset, Chunk.Size, EAsyncIOPriorityAndFlags::AIOP_Normal, Chunk.Data_IO);
+		}
+		else
+		{
+			// Add a (dummy) chunk request, so that the buffer creation path will be identical to the streaming path
+			FHairStreamingRequest::FChunk& Chunk = StreamingRequest->Chunks.AddDefaulted_GetRef();
+			Chunk.Status = FHairStreamingRequest::FChunk::EStatus::Completed;
+			Chunk.Container = &In;
+			Chunk.Size = InSize - InOffset;
+			Chunk.Offset = InOffset;
+			Chunk.TotalSize = InSize;
+			In.ChunkRequest = &Chunk;
+		}
 	}
 	else if (Type == ReadWriteIO)
 	{
