Skip streaming in hair bulk data from disk when it's already in memory. This is needed to support building groom bindings at runtime.

diff --git a/Engine/Plugins/Runtime/HairStrands/Source/HairStrandsCore/Private/GroomResources.cpp b/Engine/Plugins/Runtime/HairStrands/Source/HairStrandsCore/Private/GroomResources.cpp
--- a/Engine/Plugins/Runtime/HairStrands/Source/HairStrandsCore/Private/GroomResources.cpp
+++ b/Engine/Plugins/Runtime/HairStrands/Source/HairStrandsCore/Private/GroomResources.cpp
@@ -1179,7 +1179,19 @@ bool FHairStrandsRestResource::InternalGetOrRequestData(uint32 InRequestedCurveC
 {
 	if (StreamingRequest.IsNone())
 	{
-		StreamingRequest.Request(InRequestedCurveCount, InRequestedPointCount, BulkData, false, false, false, OwnerName);
+#if !WITH_EDITORONLY_DATA
+		if (BulkData.IsCookedBulkDataLoadedForCPUAccess())
+		{
+			// Data is already in memory
+			StreamingRequest.CurveCount = InRequestedCurveCount; 
+			StreamingRequest.PointCount = InRequestedPointCount;
+			return true;
+		}
+		else
+#endif
+		{
+			StreamingRequest.Request(InRequestedCurveCount, InRequestedPointCount, BulkData, false, false, false, OwnerName);
+		}
 	}
 	return StreamingRequest.IsCompleted();
 }
@@ -1413,7 +1425,19 @@ bool FHairStrandsClusterResource::InternalGetOrRequestData(uint32 InRequestedCur
 {
 	if (StreamingRequest.IsNone())
 	{
-		StreamingRequest.Request(InRequestedCurveCount, InRequestedPointCount, BulkData, false, false, false, OwnerName);
+#if !WITH_EDITORONLY_DATA
+		if (BulkData.IsCookedBulkDataLoadedForCPUAccess())
+		{
+			// Data is already in memory
+			StreamingRequest.CurveCount = InRequestedCurveCount;
+			StreamingRequest.PointCount = InRequestedPointCount;
+			return true;
+		}
+		else
+#endif
+		{
+			StreamingRequest.Request(InRequestedCurveCount, InRequestedPointCount, BulkData, false, false, false, OwnerName);
+		}
 	}
 	return StreamingRequest.IsCompleted();
 }
@@ -1542,7 +1566,19 @@ bool FHairStrandsLODRestRootResource::InternalGetOrRequestData(uint32 InRequeste
 {
 	if (StreamingRequest.IsNone())
 	{
-		StreamingRequest.Request(InRequestedCurveCount, InRequestedPointCount, BulkData, false, false, false, OwnerName);
+#if !WITH_EDITORONLY_DATA
+		if (BulkData.IsCookedBulkDataLoadedForCPUAccess())
+		{
+			// Data is already in memory
+			StreamingRequest.CurveCount = InRequestedCurveCount;
+			StreamingRequest.PointCount = InRequestedPointCount;
+			return true;
+		}
+		else
+#endif
+		{
+			StreamingRequest.Request(InRequestedCurveCount, InRequestedPointCount, BulkData, false, false, false, OwnerName);
+		}
 	}
 	return StreamingRequest.IsCompleted();
 }
@@ -1807,7 +1843,19 @@ bool FHairStrandsInterpolationResource::InternalGetOrRequestData(uint32 InReques
 {
 	if (StreamingRequest.IsNone())
 	{
-		StreamingRequest.Request(InRequestedCurveCount, InRequestedPointCount, BulkData, false, false, false, OwnerName);
+#if !WITH_EDITORONLY_DATA
+		if (BulkData.IsCookedBulkDataLoadedForCPUAccess())
+		{
+			// Data is already in memory
+			StreamingRequest.CurveCount = InRequestedCurveCount;
+			StreamingRequest.PointCount = InRequestedPointCount;
+			return true;
+		}
+		else
+#endif
+		{
+			StreamingRequest.Request(InRequestedCurveCount, InRequestedPointCount, BulkData, false, false, false, OwnerName);
+		}
 	}
 	return StreamingRequest.IsCompleted();
 }
diff --git a/Engine/Plugins/Runtime/HairStrands/Source/HairStrandsCore/Private/HairStrandsDatas.cpp b/Engine/Plugins/Runtime/HairStrands/Source/HairStrandsCore/Private/HairStrandsDatas.cpp
--- a/Engine/Plugins/Runtime/HairStrands/Source/HairStrandsCore/Private/HairStrandsDatas.cpp
+++ b/Engine/Plugins/Runtime/HairStrands/Source/HairStrandsCore/Private/HairStrandsDatas.cpp
@@ -725,6 +725,23 @@ float FHairStrandsBulkData::GetCoverageScale(float InCurvePercentage /*[0..1]*/)
 	return Out;
 }
 
+// Relies on bulk metadata being populated, which is only in cooked builds
+#if !WITH_EDITORONLY_DATA
+bool FHairStrandsBulkData::IsCookedBulkDataLoadedForCPUAccess() const
+{
+	bool bIsAllBulkDataLoaded = true;
+
+	bIsAllBulkDataLoaded &= Data.Positions.IsBulkDataEmptyOrLoaded();
+	bIsAllBulkDataLoaded &= Data.TranscodedPositions.IsBulkDataEmptyOrLoaded();
+	bIsAllBulkDataLoaded &= Data.CurveAttributes.IsBulkDataEmptyOrLoaded();
+	bIsAllBulkDataLoaded &= Data.PointAttributes.IsBulkDataEmptyOrLoaded();
+	bIsAllBulkDataLoaded &= Data.PointToCurve.IsBulkDataEmptyOrLoaded();
+	bIsAllBulkDataLoaded &= Data.Curves.IsBulkDataEmptyOrLoaded();
+
+	return bIsAllBulkDataLoaded;
+}
+#endif
+
 /////////////////////////////////////////////////////////////////////////////////////////
 // Interpolation bulk data
 
@@ -788,6 +805,19 @@ uint32 FHairStrandsInterpolationBulkData::GetSize() const
 	return Out;
 }
 
+// Relies on bulk metadata being populated, which is only in cooked builds
+#if !WITH_EDITORONLY_DATA
+bool FHairStrandsInterpolationBulkData::IsCookedBulkDataLoadedForCPUAccess() const
+{
+	bool bIsAllBulkDataLoaded = true;
+
+	bIsAllBulkDataLoaded &= Data.CurveInterpolation.IsBulkDataEmptyOrLoaded();
+	bIsAllBulkDataLoaded &= Data.PointInterpolation.IsBulkDataEmptyOrLoaded();
+
+	return bIsAllBulkDataLoaded;
+}
+#endif
+
 void FHairStrandsInterpolationBulkData::ResetLoadedSize()
 {
 	Data.CurveInterpolation.LoadedSize = 0;
@@ -888,6 +918,20 @@ uint32 FHairStrandsClusterBulkData::GetSize() const
 	return Out;
 }
 
+// Relies on bulk metadata being populated, which is only in cooked builds
+#if !WITH_EDITORONLY_DATA
+bool FHairStrandsClusterBulkData::IsCookedBulkDataLoadedForCPUAccess() const
+{
+	bool bIsAllBulkDataLoaded = true;
+
+	bIsAllBulkDataLoaded &= Data.PackedClusterInfos.IsBulkDataEmptyOrLoaded();
+	bIsAllBulkDataLoaded &= Data.CurveToClusterIds.IsBulkDataEmptyOrLoaded();
+	bIsAllBulkDataLoaded &= Data.PointLODs.IsBulkDataEmptyOrLoaded();
+
+	return bIsAllBulkDataLoaded;
+}
+#endif
+
 uint32 FHairStrandsClusterBulkData::GetCurveCount(float InLODIndex) const
 {
 	const TArray<FHairLODInfo>& LODInfos = Header.LODInfos;
@@ -998,6 +1042,24 @@ uint32 FHairStrandsRootBulkData::GetDataSize() const
 	return Out;
 }
 
+// Relies on bulk metadata being populated, which is only in cooked builds
+#if !WITH_EDITORONLY_DATA
+bool FHairStrandsRootBulkData::IsCookedBulkDataLoadedForCPUAccess() const
+{
+	bool bIsAllBulkDataLoaded = true;
+
+	bIsAllBulkDataLoaded &= Data.RootBarycentricBuffer.IsBulkDataEmptyOrLoaded();
+	bIsAllBulkDataLoaded &= Data.RootToUniqueTriangleIndexBuffer.IsBulkDataEmptyOrLoaded();
+	bIsAllBulkDataLoaded &= Data.UniqueTriangleIndexBuffer.IsBulkDataEmptyOrLoaded();
+	bIsAllBulkDataLoaded &= Data.RestUniqueTrianglePositionBuffer.IsBulkDataEmptyOrLoaded();
+	bIsAllBulkDataLoaded &= Data.MeshInterpolationWeightsBuffer.IsBulkDataEmptyOrLoaded();
+	bIsAllBulkDataLoaded &= Data.MeshSampleIndicesAndSectionsBuffer.IsBulkDataEmptyOrLoaded();
+	bIsAllBulkDataLoaded &= Data.RestSamplePositionsBuffer.IsBulkDataEmptyOrLoaded();
+
+	return bIsAllBulkDataLoaded;
+}
+#endif
+
 void FHairStrandsRootBulkData::Reset()
 {
 	// Header
diff --git a/Engine/Plugins/Runtime/HairStrands/Source/HairStrandsCore/Public/HairStrandsDatas.h b/Engine/Plugins/Runtime/HairStrands/Source/HairStrandsCore/Public/HairStrandsDatas.h
--- a/Engine/Plugins/Runtime/HairStrands/Source/HairStrandsCore/Public/HairStrandsDatas.h
+++ b/Engine/Plugins/Runtime/HairStrands/Source/HairStrandsCore/Public/HairStrandsDatas.h
@@ -442,6 +442,9 @@ struct FHairBulkContainer
 	void SetBulkDataFlags(uint32 Flags) { Data.SetBulkDataFlags(Flags); }
 	void RemoveBulkData() 				{ Data.RemoveBulkData();  }
 	void Serialize(FArchive& Ar, UObject* Owner, int32 ChunkIndex, bool bAttemptFileMapping) { Data.Serialize(Ar, Owner, ChunkIndex, bAttemptFileMapping); }
+
+	// Returns true if all loadable data is loaded, i.e. either the data is loaded or there's no data to load
+	bool IsBulkDataEmptyOrLoaded()const { return Data.GetBulkDataSize() == 0 || Data.IsBulkDataLoaded(); }
 };
 
 struct HAIRSTRANDSCORE_API FHairStrandsBulkCommon
@@ -623,6 +626,9 @@ struct HAIRSTRANDSCORE_API FHairStrandsBulkData : FHairStrandsBulkCommon
 	void Reset();
 	virtual void ResetLoadedSize() override;
 	float GetCoverageScale(float InCurvePercentage /*[0..1]*/) const;
+#if !WITH_EDITORONLY_DATA
+	bool IsCookedBulkDataLoadedForCPUAccess() const;
+#endif
 
 	uint32 GetNumCurves() const { return Header.CurveCount;  };
 	uint32 GetNumPoints() const { return Header.PointCount; };
@@ -750,6 +756,9 @@ struct HAIRSTRANDSCORE_API FHairStrandsInterpolationBulkData : FHairStrandsBulkC
 	virtual void GetResources(FQuery& Out) override;
 	uint32 GetPointCount() const { return Header.PointCount; };
 	uint32 GetSize() const;
+#if !WITH_EDITORONLY_DATA
+	bool IsCookedBulkDataLoadedForCPUAccess() const;
+#endif
 
 	struct FHeader
 	{
@@ -801,6 +810,9 @@ struct HAIRSTRANDSCORE_API FHairStrandsClusterBulkData : FHairStrandsBulkCommon
 	virtual void GetResources(FQuery& Out) override;
 	uint32 GetCurveCount(float InLODIndex) const;
 	uint32 GetSize() const;
+#if !WITH_EDITORONLY_DATA
+	bool IsCookedBulkDataLoadedForCPUAccess() const;
+#endif
 
 	bool IsValid() const { return Header.ClusterCount > 0 && Header.PointCount > 0; }
 
@@ -897,6 +909,9 @@ struct FHairStrandsRootBulkData : FHairStrandsBulkCommon
 	const TArray<uint32>& GetValidSectionIndices() const;
 	uint32 GetRootCount()const { return Header.RootCount; }
 	uint32 GetDataSize() const;
+#if !WITH_EDITORONLY_DATA
+	bool IsCookedBulkDataLoadedForCPUAccess() const;
+#endif
 
 	struct FHeader
 	{		
