Fix CPU RBF deformation for freezing groom asset.

The transformed then exported control points were not ordered similarly to the original hair description, because of added control point at the end of each curves.

diff --git a/Engine/Plugins/Runtime/HairStrands/Source/HairStrandsCore/Private/GroomBuilder.cpp b/Engine/Plugins/Runtime/HairStrands/Source/HairStrandsCore/Private/GroomBuilder.cpp
--- a/Engine/Plugins/Runtime/HairStrands/Source/HairStrandsCore/Private/GroomBuilder.cpp
+++ b/Engine/Plugins/Runtime/HairStrands/Source/HairStrandsCore/Private/GroomBuilder.cpp
@@ -1901,7 +1901,12 @@ static FName ToGroupName(const FStrandID& InStrandID, const uint32 InGroupID, co
 	}
 }
 
-bool FGroomBuilder::BuildHairDescriptionGroups(const FHairDescription& HairDescription, FHairDescriptionGroups& Out)
+bool NeedsToAddEndingControlPoint(uint32 CurveNumVertices)
+{
+	return GetHairStrandsUsesTriangleStrips() && uint32(CurveNumVertices + 1) <= HAIR_MAX_NUM_POINT_PER_CURVE;
+}
+
+bool FGroomBuilder::BuildHairDescriptionGroups(const FHairDescription& HairDescription, FHairDescriptionGroups& Out, bool bAllowAddEndControlPoint)
 {
 	TRACE_CPUPROFILER_EVENT_SCOPE(FGroomBuilder::BuildHairDescriptionGroups);
 
@@ -2146,7 +2151,7 @@ bool FGroomBuilder::BuildHairDescriptionGroups(const FHairDescription& HairDescr
 			continue;
 		}
 
-		const bool bAddEndingControlPoint = GetHairStrandsUsesTriangleStrips() && uint32(CurveNumVertices + 1) <= HAIR_MAX_NUM_POINT_PER_CURVE;
+		const bool bAddEndingControlPoint = NeedsToAddEndingControlPoint(CurveNumVertices) && bAllowAddEndControlPoint;
 		CurrentHairStrandsDatas->StrandsCurves.CurvesCount.Add(FMath::Min(uint32(CurveNumVertices + (bAddEndingControlPoint ? 1u : 0u)), HAIR_MAX_NUM_POINT_PER_CURVE));
 
 		if (bCanUseClosestGuidesAndWeights)
diff --git a/Engine/Plugins/Runtime/HairStrands/Source/HairStrandsCore/Private/GroomRBFDeformer.cpp b/Engine/Plugins/Runtime/HairStrands/Source/HairStrandsCore/Private/GroomRBFDeformer.cpp
--- a/Engine/Plugins/Runtime/HairStrands/Source/HairStrandsCore/Private/GroomRBFDeformer.cpp
+++ b/Engine/Plugins/Runtime/HairStrands/Source/HairStrandsCore/Private/GroomRBFDeformer.cpp
@@ -697,8 +697,10 @@ void FGroomRBFDeformer::GetRBFDeformedGroomAsset(const UGroomAsset* InGroomAsset
 
 		// Use the vertices positions from the HairDescription instead of the GroomAsset since the latter
 		// may contain decimated or auto-generated guides depending on the import settings
+		// When building the hair description groups, do not add control points at the end of each curves, 
+		// as this would make cause a mismatch ordering of the flatten CP's position vs. the hair description CP's positions
 		FHairDescriptionGroups HairDescriptionGroups;
-		FGroomBuilder::BuildHairDescriptionGroups(InGroomAsset->GetHairDescription(), HairDescriptionGroups);
+		FGroomBuilder::BuildHairDescriptionGroups(InGroomAsset->GetHairDescription(), HairDescriptionGroups, false /*bAllowAddEndControlPoint*/);
 
 		TArray<FRBFDeformedPositions> DeformedPositions;
 		DeformedPositions.SetNum(NumGroups);
diff --git a/Engine/Plugins/Runtime/HairStrands/Source/HairStrandsCore/Public/GroomBuilder.h b/Engine/Plugins/Runtime/HairStrands/Source/HairStrandsCore/Public/GroomBuilder.h
--- a/Engine/Plugins/Runtime/HairStrands/Source/HairStrandsCore/Public/GroomBuilder.h
+++ b/Engine/Plugins/Runtime/HairStrands/Source/HairStrandsCore/Public/GroomBuilder.h
@@ -59,7 +59,8 @@ struct HAIRSTRANDSCORE_API FGroomBuilder
 	// 1. Build hair group based on the hair description
 	static bool BuildHairDescriptionGroups(
 		const FHairDescription& HairDescription, 
-		FHairDescriptionGroups& Out);
+		FHairDescriptionGroups& Out,
+		bool bAllowAddEndControlPoint=true);
 
 	// 2.a Build FHairStrandsDatas for Strands & Guides, based on HairDescriptionGroups and DecimationSettings
 	static void BuildData(
