Fix CPU RBF deformation for freezing groom asset.

Recompute binding data from scratch, and store more information in the intermediate CPU binding data. These extra information are transient (non-serialized). These info are used later during the RBF process to ensure skel. mesh triangle data are consistent.

diff --git a/Engine/Plugins/Runtime/HairStrands/Source/HairStrandsCore/Private/GroomBindingBuilder.cpp b/Engine/Plugins/Runtime/HairStrands/Source/HairStrandsCore/Private/GroomBindingBuilder.cpp
--- a/Engine/Plugins/Runtime/HairStrands/Source/HairStrandsCore/Private/GroomBindingBuilder.cpp
+++ b/Engine/Plugins/Runtime/HairStrands/Source/HairStrandsCore/Private/GroomBindingBuilder.cpp
@@ -1,6 +1,7 @@
 // Copyright Epic Games, Inc. All Rights Reserved.
 
 #include "GroomBindingBuilder.h"
+#include "GroomBindingCommon.h"
 #include "GeometryCache.h"
 #include "GeometryCacheMeshData.h"
 #include "GroomAsset.h"
@@ -122,14 +123,6 @@ static float PackNormalToFloat(const FVector3f& InN)
 //////////////////////////////////////////////////////////////////////////
 // Intermediate data struct
 
-/** Binding data */
-struct FHairRootGroupData
-{
-	TArray<FHairStrandsRootData>		SimRootDatas;
-	TArray<FHairStrandsRootData>		RenRootDatas;
-	TArray<TArray<FHairStrandsRootData>>CardsRootDatas;
-};
-
 namespace GroomBinding_Mesh
 {
 //////////////////////////////////////////////////////////////////////////
@@ -802,8 +795,7 @@ namespace GroomBinding_RBFWeighting
 		const bool bNeedStrandsRoot,
 		const uint32 NumInterpolationPoints, 
 		const int32 MatchingSection, 
-		const GroomBinding_Mesh::IMeshData* MeshData, 
-		const TArray<TArray<FVector3f>>& TransferedPositions)
+		const GroomBinding_Mesh::IMeshData* MeshData)
 	{
 		const uint32 MeshLODCount= MeshData->GetNumLODs();
 		const uint32 MaxSamples  = NumInterpolationPoints;
@@ -815,9 +807,9 @@ namespace GroomBinding_RBFWeighting
 			int32 TargetSection = -1;
 			bool GlobalSamples = false;
 			const FVector3f* PositionsPointer = nullptr;
-			if (TransferedPositions.Num() == MeshLODCount)
+			if (Out.MeshPositions_Transferred.Num() == MeshLODCount)
 			{
-				PositionsPointer = TransferedPositions[MeshLODIndex].GetData();
+				PositionsPointer = Out.MeshPositions_Transferred[MeshLODIndex].GetData();
 				GlobalSamples = true;
 				TargetSection = MatchingSection;
 			}
@@ -896,6 +888,10 @@ namespace GroomBinding_RootProjection
 			FVector3f P1;
 			FVector3f P2;
 
+			FVector3f P0_NonTransfered;
+			FVector3f P1_NonTransfered;
+			FVector3f P2_NonTransfered;
+
 			FVector3f N0;
 			FVector3f N1;
 			FVector3f N2;
@@ -1265,6 +1261,11 @@ namespace GroomBinding_RootProjection
 						T.P2 = MeshLODData.GetVertexPosition(T.I2);
 					}
 
+					// Store non-transfered position as well for RBF remapping
+					T.P0_NonTransfered = MeshLODData.GetVertexPosition(T.I0);
+					T.P1_NonTransfered = MeshLODData.GetVertexPosition(T.I1);
+					T.P2_NonTransfered = MeshLODData.GetVertexPosition(T.I2);
+
 					T.UV0 = MeshLODData.GetVertexUV(T.I0, ChannelIndex);
 					T.UV1 = MeshLODData.GetVertexUV(T.I1, ChannelIndex);
 					T.UV2 = MeshLODData.GetVertexUV(T.I2, ChannelIndex);
@@ -1352,6 +1353,11 @@ namespace GroomBinding_RootProjection
 						T.P2 = MeshLODData.GetVertexPosition(T.I2);
 					}
 
+					// Store non-transfered position as well for RBF remapping
+					T.P0_NonTransfered = MeshLODData.GetVertexPosition(T.I0);
+					T.P1_NonTransfered = MeshLODData.GetVertexPosition(T.I1);
+					T.P2_NonTransfered = MeshLODData.GetVertexPosition(T.I2);
+
 					T.UV0 = MeshLODData.GetVertexUV(T.I0, ChannelIndex);
 					T.UV1 = MeshLODData.GetVertexUV(T.I1, ChannelIndex);
 					T.UV2 = MeshLODData.GetVertexUV(T.I2, ChannelIndex);
@@ -1382,6 +1388,9 @@ namespace GroomBinding_RootProjection
 			TArray<FHairStrandsMeshTrianglePositionFormat::Type> RestRootTrianglePositionBuffer;
 			RestRootTrianglePositionBuffer.SetNum(CurveCount * 3);
 
+			TArray<FHairStrandsMeshTrianglePositionFormat::Type> RestRootTrianglePositionBuffer_NonTransfered;
+			RestRootTrianglePositionBuffer_NonTransfered.SetNum(CurveCount * 3);
+
 		#if BINDING_PARALLEL_BUILDING
 			TAtomic<uint32> bIsValid(1);
 			ParallelFor(CurveCount,
@@ -1391,6 +1400,7 @@ namespace GroomBinding_RootProjection
 					&Grid,
 					&RootTriangleIndexBuffer,
 					&RestRootTrianglePositionBuffer,
+					&RestRootTrianglePositionBuffer_NonTransfered,
 					&OutRootData,
 					&bIsValid
 				] (uint32 CurveIndex)
@@ -1439,6 +1449,10 @@ namespace GroomBinding_RootProjection
 				RestRootTrianglePositionBuffer[CurveIndex * 3 + 0] = FVector4f((FVector3f)ClosestTriangle.P0, PackNormalToFloat(ClosestTriangle.N0));
 				RestRootTrianglePositionBuffer[CurveIndex * 3 + 1] = FVector4f((FVector3f)ClosestTriangle.P1, PackNormalToFloat(ClosestTriangle.N1));
 				RestRootTrianglePositionBuffer[CurveIndex * 3 + 2] = FVector4f((FVector3f)ClosestTriangle.P2, PackNormalToFloat(ClosestTriangle.N2));
+
+				RestRootTrianglePositionBuffer_NonTransfered[CurveIndex * 3 + 0] = FVector4f((FVector3f)ClosestTriangle.P0_NonTransfered, PackNormalToFloat(ClosestTriangle.N0));
+				RestRootTrianglePositionBuffer_NonTransfered[CurveIndex * 3 + 1] = FVector4f((FVector3f)ClosestTriangle.P1_NonTransfered, PackNormalToFloat(ClosestTriangle.N1));
+				RestRootTrianglePositionBuffer_NonTransfered[CurveIndex * 3 + 2] = FVector4f((FVector3f)ClosestTriangle.P2_NonTransfered, PackNormalToFloat(ClosestTriangle.N2));
 			}
 		#if BINDING_PARALLEL_BUILDING
 			);
@@ -1482,6 +1496,7 @@ namespace GroomBinding_RootProjection
 			const uint32 UniqueTriangleCount = UniqueTriangleToRootList.Num();
 			OutRootData[MeshLODIt].UniqueTriangleIndexBuffer.Reserve(UniqueTriangleCount );
 			OutRootData[MeshLODIt].RestUniqueTrianglePositionBuffer.Reserve(UniqueTriangleCount * 3);
+			OutRootData[MeshLODIt].RestUniqueTrianglePositionBuffer_TargetNonTransfered.Reserve(UniqueTriangleCount * 3);
 			for (uint32 EncodedTriangleId : UniqueTriangleToRootList)
 			{
 				auto It = UniqueTriangleToRootMap.Find(EncodedTriangleId);
@@ -1494,6 +1509,10 @@ namespace GroomBinding_RootProjection
 				OutRootData[MeshLODIt].RestUniqueTrianglePositionBuffer.Add(RestRootTrianglePositionBuffer[FirstCurveIndex * 3 + 1]);
 				OutRootData[MeshLODIt].RestUniqueTrianglePositionBuffer.Add(RestRootTrianglePositionBuffer[FirstCurveIndex * 3 + 2]);
 
+				OutRootData[MeshLODIt].RestUniqueTrianglePositionBuffer_TargetNonTransfered.Add(RestRootTrianglePositionBuffer_NonTransfered[FirstCurveIndex * 3 + 0]);
+				OutRootData[MeshLODIt].RestUniqueTrianglePositionBuffer_TargetNonTransfered.Add(RestRootTrianglePositionBuffer_NonTransfered[FirstCurveIndex * 3 + 1]);
+				OutRootData[MeshLODIt].RestUniqueTrianglePositionBuffer_TargetNonTransfered.Add(RestRootTrianglePositionBuffer_NonTransfered[FirstCurveIndex * 3 + 2]);
+
 				// Write for each root, the index of the triangle
 				const uint32 UniqueTriangleIndex = OutRootData[MeshLODIt].UniqueTriangleIndexBuffer.Num()-1;
 				for (uint32 CurveIndex : *It)
@@ -1505,6 +1524,7 @@ namespace GroomBinding_RootProjection
 			// Sanity check
 			check(OutRootData[MeshLODIt].RootToUniqueTriangleIndexBuffer.Num() == CurveCount);
 			check(OutRootData[MeshLODIt].RestUniqueTrianglePositionBuffer.Num() == UniqueTriangleCount * 3);
+			check(OutRootData[MeshLODIt].RestUniqueTrianglePositionBuffer_TargetNonTransfered.Num() == UniqueTriangleCount * 3);
 			check(OutRootData[MeshLODIt].UniqueTriangleIndexBuffer.Num() == UniqueTriangleCount);
 
 			// Update the root mesh projection data with unique valid mesh section IDs, based on the projection data
@@ -1790,7 +1810,9 @@ namespace GroomBinding_Transfer
 	bool Transfer(
 		const GroomBinding_Mesh::IMeshData* InSourceMeshData,
 		const GroomBinding_Mesh::IMeshData* InTargetMeshData,
-		TArray<TArray<FVector3f>>& OutTransferredPositions, const int32 MatchingSection)
+		TArray<TArray<FVector3f>>& OutPositions_Original, 
+		TArray<TArray<FVector3f>>& OutPositions_Transferred,
+		const int32 MatchingSection)
 	{
 
 		// 1. Insert triangles into a 2D UV grid
@@ -1867,7 +1889,8 @@ namespace GroomBinding_Transfer
 		// 2. Look for closest triangle point in UV space
 		// Make this run in parallel
 		const uint32 TargetLODCount = InTargetMeshData->GetNumLODs();
-		OutTransferredPositions.SetNum(TargetLODCount);
+		OutPositions_Original.SetNum(TargetLODCount);
+		OutPositions_Transferred.SetNum(TargetLODCount);
 		for (uint32 TargetLODIndex = 0; TargetLODIndex < TargetLODCount; ++TargetLODIndex)
 		{
 			// Check that the target SectionId is valid for the current LOD. 
@@ -1917,8 +1940,8 @@ namespace GroomBinding_Transfer
 				return false;
 			}
 
-			OutTransferredPositions[TargetLODIndex].SetNum(TargetVertexCount);
-
+			OutPositions_Original[TargetLODIndex].SetNum(TargetVertexCount);
+			OutPositions_Transferred[TargetLODIndex].SetNum(TargetVertexCount);
 #if BINDING_PARALLEL_BUILDING
 			ParallelFor(TargetVertexCount,
 				[
@@ -1927,7 +1950,8 @@ namespace GroomBinding_Transfer
 					&TargetMeshLODData,
 					TargetLODIndex,
 					&Grid,
-					&OutTransferredPositions
+					&OutPositions_Original,
+					&OutPositions_Transferred
 				] (uint32 TargetVertexIt)
 #else
 			for (uint32 TargetVertexIt = 0; TargetVertexIt < TargetVertexCount; ++TargetVertexIt)
@@ -1936,7 +1960,9 @@ namespace GroomBinding_Transfer
 				const int32 SectionIt = TargetMeshLODData.GetSectionFromVertexIndex(TargetVertexIt);
 				if (SectionIt != LocalTargetSectionId)
 				{
-					OutTransferredPositions[TargetLODIndex][TargetVertexIt] = FVector3f(0,0,0);
+					
+					OutPositions_Original[TargetLODIndex][TargetVertexIt] = FVector3f(0,0,0);
+					OutPositions_Transferred[TargetLODIndex][TargetVertexIt] = FVector3f(0,0,0);
 #if BINDING_PARALLEL_BUILDING
 					return;
 #else
@@ -1944,7 +1970,7 @@ namespace GroomBinding_Transfer
 #endif
 				}
 
-				const FVector3f Target_P    = TargetMeshLODData.GetVertexPosition(TargetVertexIt);
+				const FVector3f Target_P  = TargetMeshLODData.GetVertexPosition(TargetVertexIt);
 				const FVector2f Target_UV = TargetMeshLODData.GetVertexUV(TargetVertexIt, ChannelIndex);
 
 				// 2.1 Query closest triangles
@@ -1970,7 +1996,8 @@ namespace GroomBinding_Transfer
 					}
 				}
 				check(ClosestUVDistance < FLT_MAX);
-				OutTransferredPositions[TargetLODIndex][TargetVertexIt] = RetargetedVertexPosition;
+				OutPositions_Original[TargetLODIndex][TargetVertexIt] = Target_P;
+				OutPositions_Transferred[TargetLODIndex][TargetVertexIt] = RetargetedVertexPosition;
 			}
 #if BINDING_PARALLEL_BUILDING
 			);
@@ -2213,9 +2240,9 @@ static void BuildRootBulkData(
 
 } // namespace GroomBinding_BulkCopy
 
-  ///////////////////////////////////////////////////////////////////////////////////////////////////
+///////////////////////////////////////////////////////////////////////////////////////////////////
 // Main entry (CPU path)
-static bool InternalBuildBinding_CPU(const FGroomBindingBuilder::FInput& In, uint32 InGroupIndex, const ITargetPlatform* TargetPlatform, UGroomBindingAsset::FHairGroupPlatformData& OutPlatformData)
+bool BuildHairRootGroupData(const FGroomBindingBuilder::FInput& In, uint32 InGroupIndex, const ITargetPlatform* TargetPlatform, FHairRootGroupData& OutData)
 {
 #if WITH_EDITORONLY_DATA
 	const bool bIsAssetValid = In.GroomAsset && In.bHasValidTarget && In.GroomAsset->GetNumHairGroups() > 0;
@@ -2227,7 +2254,7 @@ static bool InternalBuildBinding_CPU(const FGroomBindingBuilder::FInput& In, uin
 	}
 
 	// 1. Build groom root data
-	FHairRootGroupData OutData;
+	OutData = FHairRootGroupData();
 	{
 		// If we're currently running on a worker thread, all this preloading
 		// stuff should have happened on the game-thread part of the build
@@ -2399,17 +2426,18 @@ static bool InternalBuildBinding_CPU(const FGroomBindingBuilder::FInput& In, uin
 		}
 	
 		uint32 WorkItemIndex = 0;
-		FScopedSlowTask SlowTask(WorkItemCount, LOCTEXT("BuildBindingData", "Building groom binding data"));
+		FScopedSlowTask SlowTask(WorkItemCount, LOCTEXT("BuildBindingData", "Building groom binding data"), true /*bInEnabled*/);
 		SlowTask.MakeDialog();
 
 		// 1.3 Transfer positions
-		TArray<TArray<FVector3f>> TransferredPositions;
 		if (bNeedTransferPosition)
 		{
 			if (!GroomBinding_Transfer::Transfer(
 				SourceMeshData.Get(),
 				TargetMeshData.Get(),
-				TransferredPositions, In.MatchingSection))
+				OutData.MeshPositions, 
+				OutData.MeshPositions_Transferred, 
+				In.MatchingSection))
 			{
 				UE_LOG(LogHairStrands, Error, TEXT("[Groom] Binding asset could not be built. Positions transfer between source and target mesh failed."));
 				return false;
@@ -2424,7 +2452,7 @@ static bool InternalBuildBinding_CPU(const FGroomBindingBuilder::FInput& In, uin
 				if (!GroomBinding_RootProjection::Project(
 					GuidesData,
 					TargetMeshData.Get(),
-					TransferredPositions,
+					OutData.MeshPositions_Transferred,
 					OutData.SimRootDatas,
 					In.MatchingSection))
 				{
@@ -2440,7 +2468,7 @@ static bool InternalBuildBinding_CPU(const FGroomBindingBuilder::FInput& In, uin
 				if (!GroomBinding_RootProjection::Project(
 					StrandsData,
 					TargetMeshData.Get(),
-					TransferredPositions,
+					OutData.MeshPositions_Transferred,
 					OutData.RenRootDatas,
 					In.MatchingSection))
 				{
@@ -2463,7 +2491,7 @@ static bool InternalBuildBinding_CPU(const FGroomBindingBuilder::FInput& In, uin
 						if (!GroomBinding_RootProjection::Project(
 							LODGuidesData,
 							TargetMeshData.Get(),
-							TransferredPositions,
+							OutData.MeshPositions_Transferred,
 							OutData.CardsRootDatas[CardsLODIt],
 							0/*In.MatchingSection*/))
 						{
@@ -2478,20 +2506,30 @@ static bool InternalBuildBinding_CPU(const FGroomBindingBuilder::FInput& In, uin
 		
 		// 1.5 RBF building
 		{
-			GroomBinding_RBFWeighting::ComputeInterpolationWeights(OutData, bNeedStrandsRoot, In.NumInterpolationPoints, In.MatchingSection, TargetMeshData.Get(), TransferredPositions);
+			GroomBinding_RBFWeighting::ComputeInterpolationWeights(OutData, bNeedStrandsRoot, In.NumInterpolationPoints, In.MatchingSection, TargetMeshData.Get());
 			SlowTask.EnterProgressFrame();
 		}
 	}
-
-	// 3. Convert data to bulk data
-	GroomBinding_BulkCopy::BuildRootBulkData(OutPlatformData, OutData);
 #endif
 	return true;
 }
 
-bool FGroomBindingBuilder::BuildBinding(const FGroomBindingBuilder::FInput& In, uint32 InGroupIndex, const ITargetPlatform* TargetPlatform, UGroomBindingAsset::FHairGroupPlatformData& Out)
+
+bool FGroomBindingBuilder::BuildBinding(const FGroomBindingBuilder::FInput& In, uint32 InGroupIndex, const ITargetPlatform* TargetPlatform, UGroomBindingAsset::FHairGroupPlatformData& OutPlatformData)
 {
-	return InternalBuildBinding_CPU(In, InGroupIndex, TargetPlatform, Out);
+	bool bSucceed = true;
+#if WITH_EDITORONLY_DATA
+	// 1. Build root Data
+	FHairRootGroupData RootData;
+	bSucceed = BuildHairRootGroupData(In, InGroupIndex, TargetPlatform, RootData);
+
+	// 2. Convert data to bulk data
+	if (bSucceed)
+	{
+		GroomBinding_BulkCopy::BuildRootBulkData(OutPlatformData, RootData);
+	}
+#endif
+	return bSucceed;
 }
 
 bool FGroomBindingBuilder::BuildBinding(class UGroomBindingAsset* BindingAsset, bool bInitResource)
diff --git a/Engine/Plugins/Runtime/HairStrands/Source/HairStrandsCore/Private/GroomBindingCommon.h b/Engine/Plugins/Runtime/HairStrands/Source/HairStrandsCore/Private/GroomBindingCommon.h
new file mode 100644
--- /dev/null
+++ b/Engine/Plugins/Runtime/HairStrands/Source/HairStrandsCore/Private/GroomBindingCommon.h
@@ -0,0 +1,21 @@
+// Copyright Epic Games, Inc. All Rights Reserved.
+
+#pragma once
+
+#include "CoreMinimal.h"
+#include "GroomBindingBuilder.h"
+
+class ITargetPlatform;
+
+/** Binding data */
+struct FHairRootGroupData
+{
+	TArray<FHairStrandsRootData>		SimRootDatas;
+	TArray<FHairStrandsRootData>		RenRootDatas;
+	TArray<TArray<FHairStrandsRootData>>CardsRootDatas;
+
+	// (Skel.) Mesh Positions for each LODs
+	TArray<TArray<FVector3f>> 			MeshPositions;
+	TArray<TArray<FVector3f>> 			MeshPositions_Transferred;
+};
+bool BuildHairRootGroupData(const FGroomBindingBuilder::FInput& In, uint32 InGroupIndex, const ITargetPlatform* TargetPlatform, FHairRootGroupData& OutData);
\ No newline at end of file
diff --git a/Engine/Plugins/Runtime/HairStrands/Source/HairStrandsCore/Private/GroomRBFDeformer.cpp b/Engine/Plugins/Runtime/HairStrands/Source/HairStrandsCore/Private/GroomRBFDeformer.cpp
--- a/Engine/Plugins/Runtime/HairStrands/Source/HairStrandsCore/Private/GroomRBFDeformer.cpp
+++ b/Engine/Plugins/Runtime/HairStrands/Source/HairStrandsCore/Private/GroomRBFDeformer.cpp
@@ -14,6 +14,7 @@
 #include "Engine/TextureDefines.h"
 #include "Engine/Texture.h"
 #include "GroomBindingBuilder.h"
+#include "GroomBindingCommon.h"
 
 #define ALLOW_CURVE_SHUFFLING 0
 
@@ -186,51 +187,11 @@ void DeformStrands(
 	});
 }
 
-// Compute the triangle positions for each curve's roots
-void ExtractUniqueTrianglePositions(
-	const FHairStrandsRootData& RestLODData,
-	const uint32 MeshLODIndex,
-	const FSkeletalMeshRenderData* InMeshRenderData, 
-	TArray<FHairStrandsMeshTrianglePositionFormat::Type>& OutDeformUniqueTrianglePositionBuffer)
-{
-	const uint32 UniqueTriangleCount = RestLODData.UniqueTriangleIndexBuffer.Num();
-	OutDeformUniqueTrianglePositionBuffer.SetNum(UniqueTriangleCount * 3);
-
-	const uint32 SectionCount = InMeshRenderData->LODRenderData[MeshLODIndex].RenderSections.Num();
-	TArray<uint32> IndexBuffer;
-	InMeshRenderData->LODRenderData[MeshLODIndex].MultiSizeIndexContainer.GetIndexBuffer(IndexBuffer);
-
-	for (uint32 UniqueTriangleIndex = 0; UniqueTriangleIndex < UniqueTriangleCount; ++UniqueTriangleIndex)
-	{
-		const uint32 PackedTriangleIndex = RestLODData.UniqueTriangleIndexBuffer[UniqueTriangleIndex];
-		uint32 TriangleIndex = 0;
-		uint32 SectionIndex = 0;
-		FHairStrandsRootUtils::UnpackTriangleIndex(PackedTriangleIndex, TriangleIndex, SectionIndex);
-
-		check(SectionIndex < SectionCount)
-		const uint32 TriangleCount = InMeshRenderData->LODRenderData[MeshLODIndex].RenderSections[SectionIndex].NumTriangles;
-		const uint32 SectionBaseIndex = InMeshRenderData->LODRenderData[MeshLODIndex].RenderSections[SectionIndex].BaseIndex;
-
-		const uint32 I0 = IndexBuffer[SectionBaseIndex + TriangleIndex * 3 + 0];
-		const uint32 I1 = IndexBuffer[SectionBaseIndex + TriangleIndex * 3 + 1];
-		const uint32 I2 = IndexBuffer[SectionBaseIndex + TriangleIndex * 3 + 2];
-
-		const FVector3f P0 = InMeshRenderData->LODRenderData[MeshLODIndex].StaticVertexBuffers.PositionVertexBuffer.VertexPosition(I0);
-		const FVector3f P1 = InMeshRenderData->LODRenderData[MeshLODIndex].StaticVertexBuffers.PositionVertexBuffer.VertexPosition(I1);
-		const FVector3f P2 = InMeshRenderData->LODRenderData[MeshLODIndex].StaticVertexBuffers.PositionVertexBuffer.VertexPosition(I2);
-
-		OutDeformUniqueTrianglePositionBuffer[UniqueTriangleIndex * 3 + 0] = P0;
-		OutDeformUniqueTrianglePositionBuffer[UniqueTriangleIndex * 3 + 1] = P1;
-		OutDeformUniqueTrianglePositionBuffer[UniqueTriangleIndex * 3 + 2] = P2;
-	}
-}
-
 TArray<FVector3f> GetDeformedHairStrandsPositions(
 	const TArray<FVector3f>& MeshVertexPositionsBuffer_Target,
 	const FHairStrandsDatas& HairStrandsData,
 	const uint32 MeshLODIndex,
 	const bool bHasRBF,
-	const FSkeletalMeshRenderData* InMeshRenderData,
 	const TArray<FHairStrandsIndexFormat::Type>& PointToCurveBuffer,
 	const FHairStrandsRootData& RestLODData)
 {
@@ -260,12 +221,8 @@ TArray<FVector3f> GetDeformedHairStrandsPositions(
 
 	// Use the vertex position of the binding, as the source asset might not have the same topology (in case the groom has been transfered from one mesh toanother using UV sharing)
 	TArray<FHairStrandsMeshTrianglePositionFormat::Type> UniqueTrianglePositionBuffer_Rest = RestLODData.RestUniqueTrianglePositionBuffer;
-	TArray<FHairStrandsMeshTrianglePositionFormat::Type> UniqueTrianglePositionBuffer_Deformed;
-	ExtractUniqueTrianglePositions(
-		RestLODData,
-		MeshLODIndex,
-		InMeshRenderData,
-		UniqueTrianglePositionBuffer_Deformed);
+	TArray<FHairStrandsMeshTrianglePositionFormat::Type> UniqueTrianglePositionBuffer_Deformed = RestLODData.RestUniqueTrianglePositionBuffer_TargetNonTransfered;
+	check(UniqueTrianglePositionBuffer_Rest.Num() == UniqueTrianglePositionBuffer_Deformed.Num());
 
 	// Deform the strands vertices with the deformed mesh samples
 	const TArray<FVector3f>& RestPosePositionBuffer = OutPositions;
@@ -657,44 +614,46 @@ static float SampleMaskTexture(const FVector2f& InUV, const FUintPoint& InResolu
 	return 1.0f;
 }
 
-void FGroomRBFDeformer::GetRBFDeformedGroomAsset(const UGroomAsset* InGroomAsset, UGroomBindingAsset* BindingAsset, FTextureSource* MaskTextureSource, const float MaskScale, UGroomAsset* OutGroomAsset)
+void FGroomRBFDeformer::GetRBFDeformedGroomAsset(const UGroomAsset* InGroomAsset, const UGroomBindingAsset* BindingAsset, FTextureSource* MaskTextureSource, const float MaskScale, UGroomAsset* OutGroomAsset, const ITargetPlatform* TargetPlatform)
 {
 #if WITH_EDITORONLY_DATA
 	if (InGroomAsset && BindingAsset && BindingAsset->GetTargetSkeletalMesh() && BindingAsset->GetSourceSkeletalMesh())
 	{
-		// Ensure the bulk data is streamed in and accessible from the CPU
-		{
-			const bool bWait = true;
-			BindingAsset->StreamInForCPUAccess(bWait);
-		}
-
-		// Get the skel. mesh render data for the current platform
-		// This similar to how we fetch skel. mesh data when building groom binding. This is requires in order to get 
-		// identical skel. mesh render data and ensure that the hair deformation is done correctly.
-		check(BindingAsset->GetGroomBindingType() == EGroomBindingMeshType::SkeletalMesh);
-		USkeletalMesh* TargetSkeletalMesh = BindingAsset->GetTargetSkeletalMesh();
-		PRAGMA_DISABLE_DEPRECATION_WARNINGS
-		FScopedSkeletalMeshRenderData TargetSkeletalMeshScopedData(TargetSkeletalMesh);
-		PRAGMA_ENABLE_DEPRECATION_WARNINGS
-		{
-			ITargetPlatform* RunningPlatform = GetTargetPlatformManagerRef().GetRunningTargetPlatform();
-			FSkinnedAssetAsyncBuildScope AsyncBuildScope(TargetSkeletalMesh);
-			USkeletalMesh::GetPlatformSkeletalMeshRenderData(RunningPlatform, TargetSkeletalMeshScopedData);
-		}
-
 		// Use the LOD0 skeletal mesh to extract the vertices used for the RBF weight computation
 		const int32 MeshLODIndex = 0;
 
-		// Get the target mesh vertices (source and target)
-		PRAGMA_DISABLE_DEPRECATION_WARNINGS
-		const FSkeletalMeshRenderData* SkeletalMeshData_Target = TargetSkeletalMeshScopedData.GetData();
-		PRAGMA_ENABLE_DEPRECATION_WARNINGS
-		TArray<FVector3f> MeshVertexPositionsBuffer_Target;
-		ExtractSkeletalVertexPosition(SkeletalMeshData_Target, MeshLODIndex, MeshVertexPositionsBuffer_Target);
-
 		// Apply RBF deformation to each group of guides and render strands
 		const int32 NumGroups = BindingAsset->GetHairGroupsPlatformData().Num();
 
+		FGroomBindingBuilder::FInput Input;
+		{
+			Input.BindingType = EGroomBindingMeshType::SkeletalMesh;
+			Input.NumInterpolationPoints = BindingAsset->GetNumInterpolationPoints();
+			Input.MatchingSection = BindingAsset->GetMatchingSection();
+			Input.bHasValidTarget = BindingAsset->HasValidTarget();
+			Input.GroomAsset = BindingAsset->GetGroom();
+			Input.SourceSkeletalMesh = BindingAsset->GetSourceSkeletalMesh();
+			Input.TargetSkeletalMesh = BindingAsset->GetTargetSkeletalMesh();
+			Input.SourceGeometryCache = nullptr;
+			Input.TargetGeometryCache = nullptr;
+		}
+
+		// Compute root datas
+		TArray<FHairRootGroupData> RootDatas;
+		{
+			if (TargetPlatform == nullptr)
+			{
+				TargetPlatform = GetTargetPlatformManagerRef().GetRunningTargetPlatform();
+			}
+
+			RootDatas.SetNum(NumGroups);
+			for (int32 GroupIt = 0; GroupIt < NumGroups; ++GroupIt)
+			{
+				const bool bSucceed = BuildHairRootGroupData(Input, GroupIt, TargetPlatform, RootDatas[GroupIt]);
+				check(bSucceed);
+			}
+		}
+
 		// Use the vertices positions from the HairDescription instead of the GroomAsset since the latter
 		// may contain decimated or auto-generated guides depending on the import settings
 		// When building the hair description groups, do not add control points at the end of each curves, 
@@ -753,6 +712,9 @@ void FGroomRBFDeformer::GetRBFDeformedGroomAsset(const UGroomAsset* InGroomAsset
 			check(BindingAsset->GetHairGroupsPlatformData()[GroupIndex].SimRootBulkDatas.IsValidIndex(MeshLODIndex));
 			check(BindingAsset->GetHairGroupsPlatformData()[GroupIndex].RenRootBulkDatas.IsValidIndex(MeshLODIndex));
 
+			// Get the deformed positions (== target mesh positions), i.e. the non-transfered position.
+			TArray<FVector3f> MeshVertexPositionsBuffer_Target = RootDatas[GroupIndex].MeshPositions[MeshLODIndex];
+
 			// Get deformed guides
 			// If the groom override the value, we output dummy value for the guides, since they won't be used
 			if (InGroomAsset->GetHairGroupsInterpolation()[GroupIndex].InterpolationSettings.GuideType != EGroomGuideType::Imported)
@@ -762,8 +724,7 @@ void FGroomRBFDeformer::GetRBFDeformedGroomAsset(const UGroomAsset* InGroomAsset
 			}
 			else
 			{
-				FHairStrandsRootData SimRootData;
-				FGroomBindingBuilder::GetRootData(SimRootData, BindingAsset->GetHairGroupsPlatformData()[GroupIndex].SimRootBulkDatas[MeshLODIndex]);
+				FHairStrandsRootData SimRootData = RootDatas[GroupIndex].SimRootDatas[MeshLODIndex];
 
 				const bool bHasRBF = InGroomAsset->IsGlobalInterpolationEnable(GroupIndex, 0/*HairLODIndex*/);
 
@@ -772,22 +733,18 @@ void FGroomRBFDeformer::GetRBFDeformedGroomAsset(const UGroomAsset* InGroomAsset
 					GuidesData,
 					MeshLODIndex,
 					bHasRBF,
-					SkeletalMeshData_Target,
 					SimRootDataPointToCurveBuffer,
 					SimRootData);
 			}
 
 			// Get deformed render strands
 			{
-				FHairStrandsRootData RenRootData;
-				FGroomBindingBuilder::GetRootData(RenRootData, BindingAsset->GetHairGroupsPlatformData()[GroupIndex].RenRootBulkDatas[MeshLODIndex]);
+				FHairStrandsRootData RenRootData = RootDatas[GroupIndex].RenRootDatas[MeshLODIndex];
 
 				// Transfer RBF weights from SimRootData to RenRootData, as RenRootData does not hold RBF sample data
 				// This is transient data.
 				{
-
-					FHairStrandsRootData SimRootData;
-					FGroomBindingBuilder::GetRootData(SimRootData, BindingAsset->GetHairGroupsPlatformData()[GroupIndex].SimRootBulkDatas[MeshLODIndex]);
+					FHairStrandsRootData SimRootData = RootDatas[GroupIndex].SimRootDatas[MeshLODIndex];
 
 					FHairStrandsRootData& SimLODData = SimRootData;
 					FHairStrandsRootData& RenLODData = RenRootData;
@@ -807,7 +764,6 @@ void FGroomRBFDeformer::GetRBFDeformedGroomAsset(const UGroomAsset* InGroomAsset
 					StrandsData,
 					MeshLODIndex,
 					bHasRBF,
-					SkeletalMeshData_Target,
 					RenRootDataPointToCurveBuffer,
 					RenRootData);
 			}
@@ -910,9 +866,10 @@ void FGroomRBFDeformer::GetRBFDeformedGroomAsset(const UGroomAsset* InGroomAsset
 					check(BindingAsset->GetHairGroupsPlatformData()[Desc.GroupIndex].SimRootBulkDatas.IsValidIndex(MeshLODIndex));
 
 					// Load Sim root data, as they contains RBF weights
-					FHairStrandsRootData SimRootData;
-					FGroomBindingBuilder::GetRootData(SimRootData, BindingAsset->GetHairGroupsPlatformData()[Desc.GroupIndex].SimRootBulkDatas[MeshLODIndex]);
+					FHairStrandsRootData SimRootData = RootDatas[Desc.GroupIndex].SimRootDatas[MeshLODIndex];
 	
+					// Get the deformed positions (== target mesh positions), i.e. the non-transfered position.
+					const TArray<FVector3f>& MeshVertexPositionsBuffer_Target = RootDatas[Desc.GroupIndex].MeshPositions[MeshLODIndex];
 					DeformStaticMeshPositions(Mesh, MeshVertexPositionsBuffer_Target, SimRootData);
 				}
 			}
@@ -933,9 +890,10 @@ void FGroomRBFDeformer::GetRBFDeformedGroomAsset(const UGroomAsset* InGroomAsset
 				check(BindingAsset->GetHairGroupsPlatformData()[Desc.GroupIndex].SimRootBulkDatas.IsValidIndex(MeshLODIndex));
 
 				// Load Sim root data, as they contains RBF weights
-				FHairStrandsRootData SimRootData;
-				FGroomBindingBuilder::GetRootData(SimRootData, BindingAsset->GetHairGroupsPlatformData()[Desc.GroupIndex].SimRootBulkDatas[MeshLODIndex]);
+				FHairStrandsRootData SimRootData = RootDatas[Desc.GroupIndex].SimRootDatas[MeshLODIndex];
 
+				// Get the deformed positions (== target mesh positions), i.e. the non-transfered position.
+				const TArray<FVector3f>& MeshVertexPositionsBuffer_Target = RootDatas[Desc.GroupIndex].MeshPositions[MeshLODIndex];
 				DeformStaticMeshPositions(Mesh, MeshVertexPositionsBuffer_Target, SimRootData);
 			}
 		}
diff --git a/Engine/Plugins/Runtime/HairStrands/Source/HairStrandsCore/Public/GroomRBFDeformer.h b/Engine/Plugins/Runtime/HairStrands/Source/HairStrandsCore/Public/GroomRBFDeformer.h
--- a/Engine/Plugins/Runtime/HairStrands/Source/HairStrandsCore/Public/GroomRBFDeformer.h
+++ b/Engine/Plugins/Runtime/HairStrands/Source/HairStrandsCore/Public/GroomRBFDeformer.h
@@ -11,10 +11,7 @@ struct FTextureSource;
 struct HAIRSTRANDSCORE_API FGroomRBFDeformer
 {
 	// Return a new GroomAsset with the RBF deformation from the BindingAsset baked into it
-	//
-	// BindingAsset is non-const because its resources will be streamed in from disk if necessary.
-	// Apart from that, it won't be modified.
-	void GetRBFDeformedGroomAsset(const UGroomAsset* InGroomAsset, UGroomBindingAsset* BindingAsset, FTextureSource* MaskTextureSource, const float MaskScale, UGroomAsset* OutGroomAsset);
+	void GetRBFDeformedGroomAsset(const UGroomAsset* InGroomAsset, const UGroomBindingAsset* BindingAsset, FTextureSource* MaskTextureSource, const float MaskScale, UGroomAsset* OutGroomAsset, const ITargetPlatform* TargetPlatform=nullptr);
 
 	static uint32 GetEntryCount(uint32 InSampleCount);
 	static uint32 GetWeightCount(uint32 InSampleCount);
diff --git a/Engine/Plugins/Runtime/HairStrands/Source/HairStrandsCore/Public/HairStrandsDatas.h b/Engine/Plugins/Runtime/HairStrands/Source/HairStrandsCore/Public/HairStrandsDatas.h
--- a/Engine/Plugins/Runtime/HairStrands/Source/HairStrandsCore/Public/HairStrandsDatas.h
+++ b/Engine/Plugins/Runtime/HairStrands/Source/HairStrandsCore/Public/HairStrandsDatas.h
@@ -864,6 +864,10 @@ struct FHairStrandsRootData
 	TArray<FHairStrandsUniqueTriangleIndexFormat::Type> UniqueTriangleIndexBuffer;
 	TArray<FHairStrandsMeshTrianglePositionFormat::Type> RestUniqueTrianglePositionBuffer;
 
+	/* Optional: Unique triangle positions before transfer. */
+	/* This data is available if triangles transfer from source to target was needed. */
+	TArray<FHairStrandsMeshTrianglePositionFormat::Type> RestUniqueTrianglePositionBuffer_TargetNonTransfered;
+
 	/* Number of samples used for the mesh interpolation */
 	uint32 SampleCount = 0;
 
