[GroomBinding] Bindings can now be built in non-editor builds, and users can specify which mesh LODs the bindings should use.

diff --git a/Engine/Plugins/Runtime/HairStrands/Source/HairStrandsCore/HairStrandsCore.Build.cs b/Engine/Plugins/Runtime/HairStrands/Source/HairStrandsCore/HairStrandsCore.Build.cs
--- a/Engine/Plugins/Runtime/HairStrands/Source/HairStrandsCore/HairStrandsCore.Build.cs
+++ b/Engine/Plugins/Runtime/HairStrands/Source/HairStrandsCore/HairStrandsCore.Build.cs
@@ -23,7 +23,8 @@ namespace UnrealBuildTool.Rules
 					"VectorVM",
 					"RHI",
 					"TraceLog",
-					"StaticMeshDescription"
+					"StaticMeshDescription",
+					"Eigen",
 				});
 
 			PublicDependencyModuleNames.AddRange(
@@ -43,7 +44,6 @@ namespace UnrealBuildTool.Rules
 				PrivateDependencyModuleNames.AddRange(
 					new string[]
 					{
-						"Eigen",
 						"DerivedDataCache",
 					});
 
diff --git a/Engine/Plugins/Runtime/HairStrands/Source/HairStrandsCore/Private/GroomBindingAsset.cpp b/Engine/Plugins/Runtime/HairStrands/Source/HairStrandsCore/Private/GroomBindingAsset.cpp
--- a/Engine/Plugins/Runtime/HairStrands/Source/HairStrandsCore/Private/GroomBindingAsset.cpp
+++ b/Engine/Plugins/Runtime/HairStrands/Source/HairStrandsCore/Private/GroomBindingAsset.cpp
@@ -9,6 +9,8 @@
 #include "GroomBindingCompiler.h"
 #include "GroomBindingBuilder.h"
 #include "HAL/LowLevelMemTracker.h"
+#include "Logging/StructuredLog.h"
+#include "Rendering/SkeletalMeshRenderData.h"
 #include "Serialization/MemoryReader.h"
 #include "UObject/UObjectIterator.h"
 #include "UObject/AnimObjectVersion.h"
@@ -110,6 +112,8 @@ static void InternalSerializePlatformData(FArchive& Ar, UObject* Owner, UGroomBi
 
 	// Cards
 	InternalSerializeCards(Ar, Owner, GroupData.CardsRootBulkDatas);
+
+	Ar << GroupData.TargetMeshMinLOD;
 }
 
 static void InternalSerializePlatformDatas(FArchive& Ar, UObject* Owner, TArray<UGroomBindingAsset::FHairGroupPlatformData>& GroupDatas, uint32 Flags)
@@ -154,12 +158,10 @@ bool FGroomBindingAsyncBuildScope::ShouldWaitOnLockedProperties(const UGroomBind
 -----------------------------------------------------------------------------*/
 void FGroomBindingAsyncBuildWorker::DoWork()
 {
-#if WITH_EDITOR
 	if (BuildContext.IsSet())
 	{
 		GroomBinding->ExecuteCacheDerivedDatas(*BuildContext);
 	}
-#endif
 }
 
 void UGroomBindingAsset::WaitUntilAsyncPropertyReleased(EGroomBindingAsyncProperties AsyncProperties, EGroomBindingAsyncPropertyLockType LockType) const
@@ -487,14 +489,18 @@ void UGroomBindingAsset::Reset()
 	bIsValid = false;
 }
 
-#if WITH_EDITORONLY_DATA
 bool UGroomBindingAsset::HasAnyDependenciesCompiling() const
 {
-	return 
+#if WITH_EDITORONLY_DATA
+	return
 		(GetTargetSkeletalMesh() && GetTargetSkeletalMesh()->IsCompiling()) ||
 		(GetSourceSkeletalMesh() && GetSourceSkeletalMesh()->IsCompiling());
+#else
+	return false;
+#endif
 }
 
+#if WITH_EDITORONLY_DATA
 void UGroomBindingAsset::InvalidateBinding(class USkeletalMesh*)
 {
 	InvalidateBinding();
@@ -508,7 +514,7 @@ void UGroomBindingAsset::InvalidateBinding()
 		return;
 	}
 
-	CacheDerivedDatas();
+	BeginCacheDerivedDatas();
 }
 
 #endif
@@ -885,22 +891,39 @@ bool UGroomBindingAsset::HasValidTarget() const
 		   (GetGroomBindingType() == EGroomBindingMeshType::GeometryCache && GetTargetGeometryCache());
 }
 
-#if WITH_EDITOR
+void UGroomBindingAsset::Build(FOnGroomBindingAssetBuildComplete CompletionDelegate)
+{
+	Build(CompletionDelegate, FOnGroomBindingAssetBuildCompleteNative());
+}
 
-void UGroomBindingAsset::Build()
+void UGroomBindingAsset::Build(const FOnGroomBindingAssetBuildCompleteNative& CompletionDelegate)
+{
+	Build(FOnGroomBindingAssetBuildComplete(), CompletionDelegate);
+}
+
+void UGroomBindingAsset::Build(const FOnGroomBindingAssetBuildComplete& DynamicCompletionDelegate, const FOnGroomBindingAssetBuildCompleteNative& NativeCompletionDelegate)
 {
 	if (GetGroom() && HasValidTarget())
 	{
+#if WITH_EDITOR
 		OnGroomBindingAssetChanged.Broadcast();
+#endif
 		Reset();
-		CacheDerivedDatas();
+		BeginCacheDerivedDatas(DynamicCompletionDelegate, NativeCompletionDelegate);
+	}
+	else
+	{
+		DynamicCompletionDelegate.ExecuteIfBound(this, EGroomBindingAssetBuildResult::Failed);
+		NativeCompletionDelegate.ExecuteIfBound(this, EGroomBindingAssetBuildResult::Failed);
 	}
 }
 
+
+#if WITH_EDITOR
 void UGroomBindingAsset::PostEditChangeProperty(FPropertyChangedEvent& PropertyChangedEvent)
 {
 	Super::PostEditChangeProperty(PropertyChangedEvent);
-	CacheDerivedDatas();
+	BeginCacheDerivedDatas();
 	OnGroomBindingAssetChanged.Broadcast();
 }
 #endif // WITH_EDITOR
@@ -949,17 +972,28 @@ namespace GroomBindingDerivedDataCacheUtils
 	}
 }
 
-static FString BuildDerivedDataKeySuffix(const UGroomBindingAsset& BindingAsset, const ITargetPlatform* TargetPlatform)
+static FString BuildDerivedDataKeySuffix(const UGroomBindingAsset& BindingAsset, const ITargetPlatform* TargetPlatform, int32 SourceMeshLOD, int32 TargetMeshMinLOD)
 {
 	FString BindingType;
 	FString SourceKey;
 	FString TargetKey;
+	FString SourceLODKey;
+	FString TargetLODKey;
 
 	if (BindingAsset.GetGroomBindingType() == EGroomBindingMeshType::SkeletalMesh)
 	{
 		// Binding type is implicitly SkeletalMesh so keep BindingType empty to prevent triggering rebuild of old binding for nothing
-		SourceKey = BindingAsset.GetSourceSkeletalMesh() ? BindingAsset.GetSourceSkeletalMesh()->BuildDerivedDataKey(TargetPlatform) : FString();
-		TargetKey = BindingAsset.GetTargetSkeletalMesh() ? BindingAsset.GetTargetSkeletalMesh()->BuildDerivedDataKey(TargetPlatform) : FString();
+		if (USkeletalMesh* SourceMesh = BindingAsset.GetSourceSkeletalMesh())
+		{
+			SourceKey = SourceMesh->BuildDerivedDataKey(TargetPlatform);
+			SourceLODKey = FString::FromInt(SourceMeshLOD);
+		}
+
+		if (USkeletalMesh* TargetMesh = BindingAsset.GetTargetSkeletalMesh())
+		{
+			TargetKey = TargetMesh->BuildDerivedDataKey(TargetPlatform);
+			TargetLODKey = FString::FromInt(TargetMeshMinLOD);
+		}
 	}
 	else
 	{
@@ -973,32 +1007,41 @@ static FString BuildDerivedDataKeySuffix(const UGroomBindingAsset& BindingAsset,
 	FString PointKey  = FString::FromInt(BindingAsset.GetNumInterpolationPoints());
 	FString SectionKey = FString::FromInt(BindingAsset.GetMatchingSection());
 
-	uint32 KeyLength  = BindingType.Len() + SourceKey.Len() + TargetKey.Len() + GroomKey.Len() + PointKey.Len() + SectionKey.Len();
+	uint32 KeyLength  = BindingType.Len() + SourceKey.Len() + TargetKey.Len() + GroomKey.Len() + PointKey.Len() + SectionKey.Len() + SourceLODKey.Len() + TargetLODKey.Len();
 
 	FString KeySuffix;
 	KeySuffix.Reserve(KeyLength);
-	KeySuffix = BindingType + SourceKey + TargetKey + GroomKey + PointKey + SectionKey;
+	KeySuffix = BindingType + SourceKey + TargetKey + GroomKey + PointKey + SectionKey + SourceLODKey + TargetLODKey;
 	return KeySuffix;
 }
 
-static FString BuildDerivedDataKeyGroup(const UGroomBindingAsset& BindingAsset, const ITargetPlatform* TargetPlatform, uint32 InGroupIndex)
-{
-	const FString DeriveDataKeySuffix = BuildDerivedDataKeySuffix(BindingAsset, TargetPlatform);
-	return GroomBindingDerivedDataCacheUtils::BuildGroomBindingDerivedDataKey(DeriveDataKeySuffix + FString(TEXT("_Group")) + FString::FromInt(InGroupIndex));
-}
-
 static FString BuildDerivedDataKeyGroup(const FString& InDeriveDataKeySuffix , uint32 InGroupIndex)
 {
 	return GroomBindingDerivedDataCacheUtils::BuildGroomBindingDerivedDataKey(InDeriveDataKeySuffix + FString(TEXT("_Group")) + FString::FromInt(InGroupIndex));
 }
 
-static TArray<FString> GetGroupDerivedDataKeys(const UGroomBindingAsset* In, const ITargetPlatform* TargetPlatform);
-static void CacheDerivedDatas(UGroomBindingAsset* In, const uint32 InGroupIndex, const FString& DerivedDataKey, bool& bOutValid, const ITargetPlatform* TargetPlatform, UGroomBindingAsset::FHairGroupPlatformData& OutPlatformData);
+static TArray<FString> GetGroupDerivedDataKeys(const UGroomBindingAsset* In, const ITargetPlatform* TargetPlatform, int32 SourceMeshLOD, int32 TargetMeshMinLOD);
+static void CacheDerivedDatas(
+	UGroomBindingAsset* In,
+	const uint32 InGroupIndex,
+	const FString& DerivedDataKey,
+	bool& bOutValid,
+	const ITargetPlatform* TargetPlatform,
+	const int32 SourceMeshLOD,
+	const int32 TargetMeshMinLOD,
+	UGroomBindingAsset::FHairGroupPlatformData& OutPlatformData);
+
+#endif // WITH_EDITORONLY_DATA
+
 
-void UGroomBindingAsset::CacheDerivedDatas()
+void UGroomBindingAsset::BeginCacheDerivedDatas(
+	const FOnGroomBindingAssetBuildComplete& DynamicCompletionDelegate,
+	const FOnGroomBindingAssetBuildCompleteNative& NativeCompletionDelegate)
 {
 	if (!GetGroom() || !GetGroom()->IsValid())
 	{
+		DynamicCompletionDelegate.ExecuteIfBound(this, EGroomBindingAssetBuildResult::Failed);
+		NativeCompletionDelegate.ExecuteIfBound(this, EGroomBindingAssetBuildResult::Failed);
 		return;
 	}
 
@@ -1013,13 +1056,11 @@ void UGroomBindingAsset::CacheDerivedDatas()
 		return;
 	}
 
-	TRACE_CPUPROFILER_EVENT_SCOPE(UGroomBindingAsset::CacheDerivedDatas);
+	TRACE_CPUPROFILER_EVENT_SCOPE(UGroomBindingAsset::BeginCacheDerivedDatas);
 
 	// Flush the rendering commands generated by the detachments.
 	FlushRenderingCommandIfUsed(this);
 
-	FGroomBindingBuildContext Context;
-#if WITH_EDITOR
 	// Acquire everything in readonly so that any thread trying to write a property locked for read-only will wait until the build is finished
 	AcquireAsyncProperty(MAX_uint64, EGroomBindingAsyncPropertyLockType::ReadOnly);
 
@@ -1028,15 +1069,34 @@ void UGroomBindingAsset::CacheDerivedDatas()
 	AcquireAsyncProperty((uint64)EGroomBindingAsyncProperties::HairGroupResources, EGroomBindingAsyncPropertyLockType::WriteOnly);
 	AcquireAsyncProperty((uint64)EGroomBindingAsyncProperties::HairGroupPlatformData, EGroomBindingAsyncPropertyLockType::WriteOnly);
 	AcquireAsyncProperty((uint64)EGroomBindingAsyncProperties::GroupInfos, EGroomBindingAsyncPropertyLockType::WriteOnly);
+	AcquireAsyncProperty((uint64)EGroomBindingAsyncProperties::SourceMeshUsedLOD, EGroomBindingAsyncPropertyLockType::WriteOnly);
+	AcquireAsyncProperty((uint64)EGroomBindingAsyncProperties::TargetMeshUsedMinLOD, EGroomBindingAsyncPropertyLockType::WriteOnly);
+
+	FGroomBindingBuildContext Context;
+	if (!TryInitializeContextForMeshes(Context))
+	{
+		DynamicCompletionDelegate.ExecuteIfBound(this, EGroomBindingAssetBuildResult::Failed);
+		NativeCompletionDelegate.ExecuteIfBound(this, EGroomBindingAssetBuildResult::Failed);
+		return;
+	}
+
+	Context.DynamicCompletionDelegate = DynamicCompletionDelegate;
+	Context.NativeCompletionDelegate = NativeCompletionDelegate;
 
 	// Dependency management should have been taken care of already. Just making sure.
 	check(GetTargetSkeletalMesh() == nullptr || GetTargetSkeletalMesh()->IsCompiling() == false);
 	check(GetSourceSkeletalMesh() == nullptr || GetSourceSkeletalMesh()->IsCompiling() == false);
 
-	// Do synchronous build until GeometryCache is thread-safe
-	const bool bIsAsyncBuildSupported = GetTargetGeometryCache() == nullptr && GetSourceGeometryCache() == nullptr;
+	const IConsoleVariable* CVarStreaming = IConsoleManager::Get().FindConsoleVariable(TEXT("r.HairStrands.Streaming"));
+	const bool bIsStrandStreamingEnabled = CVarStreaming && CVarStreaming->GetBool();
+
+	const bool bIsAsyncBuildSupported = 
+		// Do synchronous build until GeometryCache is thread-safe
+		GetTargetGeometryCache() == nullptr
+		&& GetSourceGeometryCache() == nullptr
+		// Strand streaming not supported for async build yet
+		&& !bIsStrandStreamingEnabled;
 
-	// Inline reduction is not thread-safe, prevent async build until this is fixed.
 	if (bIsAsyncBuildSupported && FGroomBindingCompilingManager::Get().IsAsyncCompilationAllowed(this))
 	{
 		FQueuedThreadPool* ThreadPool = FGroomBindingCompilingManager::Get().GetThreadPool();
@@ -1047,13 +1107,116 @@ void UGroomBindingAsset::CacheDerivedDatas()
 		FGroomBindingCompilingManager::Get().AddGroomBindings({ this });
 	}
 	else
-#endif
 	{
 		ExecuteCacheDerivedDatas(Context);
 		FinishCacheDerivedDatas(Context);
 	}
 }
 
+bool UGroomBindingAsset::TryInitializeContextForMeshes(FGroomBindingBuildContext& OutContext) const
+{
+	ITargetPlatform* RunningPlatform = nullptr;
+#if WITH_EDITOR
+	RunningPlatform = GetTargetPlatformManagerRef().GetRunningTargetPlatform();
+	check(RunningPlatform);
+#endif
+
+	if (GetGroomBindingType() == EGroomBindingMeshType::SkeletalMesh)
+	{
+		if (USkeletalMesh* SourceMesh = GetSourceSkeletalMesh())
+		{
+			OutContext.SourceMeshLOD = GetClampedSourceMeshLOD(RunningPlatform);
+
+			// Clamp SourceMeshLOD to the next best LOD that's currently loaded.
+			const FStreamableRenderResourceState& SRRState = SourceMesh->GetStreamableResourceState();
+			if (SRRState.IsValid())
+			{
+				const int32 NumValidLODs = FMath::Min(SRRState.NumRequestedLODs, SRRState.NumResidentLODs);
+				const int32 CurrentFirstLOD = SRRState.LODCountToAssetFirstLODIdx(NumValidLODs);
+				OutContext.SourceMeshLOD = FMath::Max(CurrentFirstLOD, OutContext.SourceMeshLOD);
+			}
+
+			if (FSkeletalMeshRenderData* RenderData = SourceMesh->GetResourceForRendering())
+			{
+				TRefCountPtr<FSkeletalMeshLODRenderData> LODReference;
+
+				if (RenderData->LODRenderData.IsValidIndex(OutContext.SourceMeshLOD))
+				{
+					LODReference = &RenderData->LODRenderData[OutContext.SourceMeshLOD];
+				}
+
+				if (!ensureMsgf(LODReference.IsValid(),
+					TEXT("Missing render data for mesh %s LOD %i, which is claimed to be valid by FStreamableRenderResourceState. Can't build groom binding %s."),
+					*SourceMesh->GetPathName(), OutContext.SourceMeshLOD, *GetPathName()))
+				{
+					return false;
+				}
+
+				OutContext.MeshLODReferences.Add(MoveTemp(LODReference));
+			}
+			else
+			{
+				UE_LOGFMT(LogHairStrands, Error, "Source mesh {0} has no render data. Can't build groom binding {1}.", SourceMesh->GetPathName(), GetPathName());
+				return false;
+			}
+		}
+
+		if (USkeletalMesh* TargetMesh = GetTargetSkeletalMesh())
+		{
+			OutContext.TargetMeshMinLOD = GetClampedTargetMeshMinLOD(RunningPlatform);
+
+			const FStreamableRenderResourceState& SRRState = TargetMesh->GetStreamableResourceState();
+			if (SRRState.IsValid())
+			{
+				const int32 NumValidLODs = FMath::Min(SRRState.NumRequestedLODs, SRRState.NumResidentLODs);
+				const int32 CurrentFirstLOD = SRRState.LODCountToAssetFirstLODIdx(NumValidLODs);
+				OutContext.TargetMeshMinLOD = FMath::Max(CurrentFirstLOD, OutContext.TargetMeshMinLOD);
+			}
+
+			if (FSkeletalMeshRenderData* RenderData = TargetMesh->GetResourceForRendering())
+			{
+				if (RenderData->LODRenderData.IsValidIndex(OutContext.TargetMeshMinLOD))
+				{
+					for (int32 TargetMeshLODIndex = OutContext.TargetMeshMinLOD; TargetMeshLODIndex >= 0; TargetMeshLODIndex--)
+					{
+						TRefCountPtr<FSkeletalMeshLODRenderData> LODReference = &RenderData->LODRenderData[TargetMeshLODIndex];
+
+						if (!ensureMsgf(LODReference.IsValid(),
+							TEXT("Missing render data for mesh %s LOD %i, which is claimed to be valid by FStreamableRenderResourceState. Can't build groom binding %s."),
+							*TargetMesh->GetPathName(), TargetMeshLODIndex, *GetPathName()))
+						{
+							return false;
+						}
+
+						OutContext.MeshLODReferences.Add(MoveTemp(LODReference));
+					}
+				}
+				else
+				{
+					ensureMsgf(false,
+						TEXT("Missing render data for mesh %s LOD %i, which is claimed to be valid by FStreamableRenderResourceState. Can't build groom binding %s."),
+						*TargetMesh->GetPathName(), OutContext.TargetMeshMinLOD, *GetPathName());
+
+					return false;
+				}
+			}
+			else
+			{
+				UE_LOGFMT(LogHairStrands, Error, "Target mesh {0} has no render data. Can't build groom binding {1}.", TargetMesh->GetPathName(), GetPathName());
+				return false;
+			}
+		}
+	}
+	else
+	{
+		// Geometry caches are treated as having a single LOD
+		OutContext.SourceMeshLOD = 0;
+		OutContext.TargetMeshMinLOD = 0;
+	}
+
+	return true;
+}
+
 void UGroomBindingAsset::ExecuteCacheDerivedDatas(FGroomBindingBuildContext& Context)
 {
 	TRACE_CPUPROFILER_EVENT_SCOPE(UGroomBindingAsset::ExecuteCacheDerivedDatas);
@@ -1063,18 +1226,22 @@ void UGroomBindingAsset::ExecuteCacheDerivedDatas(FGroomBindingBuildContext& Con
 	// 1. Set group count to the groom target
 	const uint32 GroupCount = GetGroom()->GetNumHairGroups();
 	GetHairGroupsPlatformData().SetNum(GroupCount);
-	CachedDerivedDataKey.SetNum(GroupCount);
 	GetGroupInfos().SetNum(GroupCount);
-
-	// 2.1 Prepare main cache key
-	// Cache derived data for the running platform.
-	ITargetPlatform* RunningPlatform = GetTargetPlatformManagerRef().GetRunningTargetPlatform();
+	ITargetPlatform* RunningPlatform = nullptr;
+#if WITH_EDITOR
+	RunningPlatform = GetTargetPlatformManagerRef().GetRunningTargetPlatform();
 	check(RunningPlatform);
-	const FString KeySuffix = BuildDerivedDataKeySuffix(*this, RunningPlatform);
+	CachedDerivedDataKey.SetNum(GroupCount);
+
 
-	// 2.2 Build the key for each group and check if any group needs to be rebuilt
-	const TArray<FString> GroupDerivedDataKeys = GetGroupDerivedDataKeys(this, RunningPlatform);
+	// 2. Build the key for each group and check if any group needs to be rebuilt
+	const TArray<FString> GroupDerivedDataKeys = GetGroupDerivedDataKeys(this, RunningPlatform, Context.SourceMeshLOD, Context.TargetMeshMinLOD);
 	const bool bAnyGroupNeedRebuild = GroupDerivedDataKeys != CachedDerivedDataKey;
+#else
+	// Note that in the non-editor path, we don't detect if the build is redundant and it's up to
+	// the user to avoid calling Build unnecessarily.
+	const bool bAnyGroupNeedRebuild = true;
+#endif
 
 	// 3. Build or retrieve from cache, binding data for each group
 	bIsValid = true;
@@ -1086,19 +1253,34 @@ void UGroomBindingAsset::ExecuteCacheDerivedDatas(FGroomBindingBuildContext& Con
 			// 1. Build data
 			bool bGroupValid = true;
 			bool bGroupReloadResource = false;
+#if WITH_EDITOR
 			if (GroupDerivedDataKeys[GroupIndex] != CachedDerivedDataKey[GroupIndex])
 			{
-				::CacheDerivedDatas(this, GroupIndex, GroupDerivedDataKeys[GroupIndex], bGroupValid, RunningPlatform, GetHairGroupsPlatformData()[GroupIndex]);
-	
-				if (bGroupValid)
-				{
-					bGroupReloadResource = true;
-					CachedDerivedDataKey[GroupIndex] = GroupDerivedDataKeys[GroupIndex] ;
-				}
-				else
-				{
-					UE_LOG(LogHairStrands, Error, TEXT("[Groom] The binding asset (%s) couldn't be built. This binding asset won't be used."), *GetName());
-				}
+				// Fetch the built data from the DDC or build it
+
+				const FString Key = GroupDerivedDataKeys[GroupIndex];
+				CacheDerivedDatas(this, GroupIndex, Key, bGroupValid, RunningPlatform, Context.SourceMeshLOD, Context.TargetMeshMinLOD, GetHairGroupsPlatformData()[GroupIndex]);
+			}
+#else
+			// Build the data directly, as there is no DDC
+			const FGroomBindingBuilder::FInput BuilderInput(this, RunningPlatform, Context.SourceMeshLOD, Context.TargetMeshMinLOD);
+			bGroupValid = FGroomBindingBuilder::BuildBinding(
+				BuilderInput,
+				GroupIndex,
+				RunningPlatform,
+				GetHairGroupsPlatformData()[GroupIndex]);
+#endif
+
+			if (bGroupValid)
+			{
+				bGroupReloadResource = true;
+#if WITH_EDITOR
+				CachedDerivedDataKey[GroupIndex] = GroupDerivedDataKeys[GroupIndex];
+#endif
+			}
+			else
+			{
+				UE_LOGFMT(LogHairStrands, Error, "[Groom] The binding asset ({0}) couldn't be built. This binding asset won't be used.", GetName());
 			}
 	
 			// 2. Release existing resources data
@@ -1135,7 +1317,16 @@ void UGroomBindingAsset::ExecuteCacheDerivedDatas(FGroomBindingBuildContext& Con
 	}
 }
 
-static void CacheDerivedDatas(UGroomBindingAsset* In, const uint32 InGroupIndex, const FString& DerivedDataKey, bool& bOutValid, const ITargetPlatform* TargetPlatform, UGroomBindingAsset::FHairGroupPlatformData& OutPlatformData)
+#if WITH_EDITORONLY_DATA
+static void CacheDerivedDatas(
+	UGroomBindingAsset* In,
+	const uint32 InGroupIndex,
+	const FString& DerivedDataKey,
+	bool& bOutValid,
+	const ITargetPlatform* TargetPlatform,
+	const int32 SourceMeshLOD,
+	const int32 TargetMeshMinLOD,
+	UGroomBindingAsset::FHairGroupPlatformData& OutPlatformData)
 {
 	{
 		bOutValid = false;
@@ -1186,22 +1377,13 @@ static void CacheDerivedDatas(UGroomBindingAsset* In, const uint32 InGroupIndex,
 
 			bOutValid = true;
 		}
+
 		if (!bHasDataInCache)
 		{
 			UE_CLOG(IsHairStrandsDDCLogEnable(), LogHairStrands, Log, TEXT("[GroomBinding/DDC] Not found (GroomBinding:%s TargetPlatform:%s Async:%s)."), *In->GetName(), *TargetPlatform->PlatformName(), IsInGameThread() ? TEXT("No") : TEXT("Yes"));
 
 			// Build groom binding data
-			FGroomBindingBuilder::FInput BuilderInput;
-			BuilderInput.BindingType = In->GetGroomBindingType();
-			BuilderInput.NumInterpolationPoints = In->GetNumInterpolationPoints();
-			BuilderInput.MatchingSection = In->GetMatchingSection();
-			BuilderInput.bHasValidTarget = In->HasValidTarget();
-			BuilderInput.GroomAsset = In->GetGroom();
-			BuilderInput.SourceSkeletalMesh = In->GetSourceSkeletalMesh();
-			BuilderInput.TargetSkeletalMesh = In->GetTargetSkeletalMesh();
-			BuilderInput.SourceGeometryCache = In->GetSourceGeometryCache();
-			BuilderInput.TargetGeometryCache = In->GetTargetGeometryCache();
-
+			FGroomBindingBuilder::FInput BuilderInput(In, TargetPlatform, SourceMeshLOD, TargetMeshMinLOD);
 			bOutValid = FGroomBindingBuilder::BuildBinding(BuilderInput, InGroupIndex, TargetPlatform, OutPlatformData);
 
 			if (bOutValid)
@@ -1233,6 +1415,7 @@ static void CacheDerivedDatas(UGroomBindingAsset* In, const uint32 InGroupIndex,
 		}
 	}
 }
+#endif // WITH_EDITORONLY_DATA
 
 void UGroomBindingAsset::FinishCacheDerivedDatas(FGroomBindingBuildContext& Context)
 {
@@ -1240,15 +1423,31 @@ void UGroomBindingAsset::FinishCacheDerivedDatas(FGroomBindingBuildContext& Cont
 
 	ReleaseAsyncProperty();
 
+	EGroomBindingAssetBuildResult Result;
 	if (!IsTemplate() && Context.bReloadResource && IsValid())
 	{
+		Result = EGroomBindingAssetBuildResult::Succeeded;
+
 		InitResource();
+
+		SetSourceMeshUsedLOD(Context.SourceMeshLOD);
+		SetTargetMeshUsedMinLOD(Context.TargetMeshMinLOD);
+	}
+	else
+	{
+		Result = EGroomBindingAssetBuildResult::Failed;
+
+		SetSourceMeshUsedLOD(INDEX_NONE);
+		SetTargetMeshUsedMinLOD(INDEX_NONE);
 	}
 
 	for (TObjectIterator<UGroomComponent> HairStrandsComponentIt; HairStrandsComponentIt; ++HairStrandsComponentIt)
 	{
-		if (HairStrandsComponentIt->BindingAsset == this ||
-			HairStrandsComponentIt->BindingAssetBeingLoaded == this) // A GroomAsset was set on the component while it was still loading
+		if (HairStrandsComponentIt->BindingAsset == this 
+#if WITH_EDITOR
+			|| HairStrandsComponentIt->BindingAssetBeingLoaded == this // A GroomAsset was set on the component while it was still loading
+#endif
+			) 
 		{
 			HairStrandsComponentIt->PostCompilation();
 
@@ -1258,14 +1457,74 @@ void UGroomBindingAsset::FinishCacheDerivedDatas(FGroomBindingBuildContext& Cont
 			}
 		}
 	}
+
+	Context.DynamicCompletionDelegate.ExecuteIfBound(this, Result);
+	Context.NativeCompletionDelegate.ExecuteIfBound(this, Result);
 }
 
-static TArray<FString> GetGroupDerivedDataKeys(const UGroomBindingAsset* In, const ITargetPlatform* TargetPlatform)
+int32 UGroomBindingAsset::GetClampedSourceMeshLOD(const ITargetPlatform* TargetPlatform) const
+{
+	const USkeletalMesh* SourceMesh = GetSourceSkeletalMesh();
+	if (GetGroomBindingType() != EGroomBindingMeshType::SkeletalMesh 
+		|| !SourceMesh)
+	{
+		return INDEX_NONE;
+	}
+
+	// BestAvailableLOD is the best (i.e. lowest index) LOD that can ever be used by this mesh.
+	// E.g. in a cooked build, this is the best LOD that was cooked.
+	// 
+	// SourceMeshLOD is clamped to it here, so that if the mesh has a per-platform min LOD,
+	// the caller of UGroomBindingAsset::Build doesn't need to take that into account. They can
+	// simply request source mesh LOD 0 and if min LOD on the mesh is 3, for example, it will
+	// silently build the groom binding using LOD 3 as the source.
+#if WITH_EDITOR
+	check(TargetPlatform);
+	const int32 BestAvailableLOD = SourceMesh->GetPlatformMinLODIdx(TargetPlatform);
+#else
+	const bool bForceLowestLODIdx = true;
+	const int32 BestAvailableLOD = SourceMesh->GetMinLodIdx(bForceLowestLODIdx);
+#endif
+
+	return FMath::Max(GetSourceMeshRequestedLOD(), BestAvailableLOD);
+}
+
+int32 UGroomBindingAsset::GetClampedTargetMeshMinLOD(const ITargetPlatform* TargetPlatform) const
+{
+	const USkeletalMesh* TargetMesh = GetTargetSkeletalMesh();
+	if (GetGroomBindingType() != EGroomBindingMeshType::SkeletalMesh
+		|| !TargetMesh)
+	{
+		return INDEX_NONE;
+	}
+
+	// BestAvailableLOD is the best (i.e. lowest index) LOD that can ever be used by this mesh.
+	// E.g. in a cooked build, this is the best LOD that was cooked.
+	// 
+	// TargetMeshMinLOD is clamped to it here, so that if the mesh has a per-platform min LOD,
+	// the caller of UGroomBindingAsset::Build doesn't need to take that into account. They can
+	// simply request target mesh LOD 0 and if min LOD on the mesh is 3, for example, it will build 
+	// the groom binding for LOD 3 onwards. LODs 0-2 would never be rendered anyway, so this is
+	// transparent to the caller.
+#if WITH_EDITOR
+	check(TargetPlatform);
+	const int32 BestAvailableLOD = TargetMesh->GetPlatformMinLODIdx(TargetPlatform);
+#else
+	const bool bForceLowestLODIdx = true;
+	const int32 BestAvailableLOD = TargetMesh->GetMinLodIdx(bForceLowestLODIdx);
+#endif
+
+	return FMath::Max(GetTargetMeshRequestedMinLOD(), BestAvailableLOD);
+}
+
+#if WITH_EDITORONLY_DATA
+
+static TArray<FString> GetGroupDerivedDataKeys(const UGroomBindingAsset* In, const ITargetPlatform* TargetPlatform, int32 SourceMeshLOD, int32 TargetMeshMinLOD)
 {
 	check(In);
 	check(TargetPlatform);
 
-	const FString KeySuffix = BuildDerivedDataKeySuffix(*In, TargetPlatform);
+	const FString KeySuffix = BuildDerivedDataKeySuffix(*In, TargetPlatform, SourceMeshLOD, TargetMeshMinLOD);
 	const uint32 GroupCount = In->GetGroupInfos().Num();
 
 	TArray<FString> Out;
@@ -1300,7 +1559,9 @@ void UGroomBindingAsset::BeginCacheForCookedPlatformData(const ITargetPlatform*
 	Super::BeginCacheForCookedPlatformData(TargetPlatform);
 
 	// 1. Build the key for each group
-	const TArray<FString> GroupDerivedDataKeys = GetGroupDerivedDataKeys(this, TargetPlatform);
+	const int32 SourceMeshLOD = GetClampedSourceMeshLOD(TargetPlatform);
+	const int32 TargetMeshMinLOD = GetClampedTargetMeshMinLOD(TargetPlatform);
+	const TArray<FString> GroupDerivedDataKeys = GetGroupDerivedDataKeys(this, TargetPlatform, SourceMeshLOD, TargetMeshMinLOD);
 
 	// 2. Find existing cached cooked data
 	UGroomBindingAsset::FCachedCookedPlatformData* TargetPlatformData = FindCachedCookedPlatformData(GroupDerivedDataKeys, CachedCookedPlatformDatas);
@@ -1316,7 +1577,16 @@ void UGroomBindingAsset::BeginCacheForCookedPlatformData(const ITargetPlatform*
 		for (uint32 GroupIndex = 0; GroupIndex < GroupCount; ++GroupIndex)
 		{
 			bool bGroupValid = true;
-			::CacheDerivedDatas(this, GroupIndex, TargetPlatformData->GroupDerivedDataKeys[GroupIndex], bGroupValid, TargetPlatform, TargetPlatformData->GroupPlatformDatas[GroupIndex]);
+			CacheDerivedDatas(
+				this, 
+				GroupIndex, 
+				TargetPlatformData->GroupDerivedDataKeys[GroupIndex], 
+				bGroupValid, 
+				TargetPlatform, 
+				SourceMeshLOD, 
+				TargetMeshMinLOD, 
+				TargetPlatformData->GroupPlatformDatas[GroupIndex]);
+
 			if (!bGroupValid)
 			{
 				UE_LOG(LogHairStrands, Error, TEXT("[Groom] The binding asset (%s) couldn't be built. This binding asset won't be used."), *GetName());
@@ -1344,7 +1614,9 @@ void UGroomBindingAsset::BeginCacheForCookedPlatformData(const ITargetPlatform*
 TArray<UGroomBindingAsset::FHairGroupPlatformData>* UGroomBindingAsset::GetCachedCookedPlatformData(const ITargetPlatform* TargetPlatform)
 {
 	// 1. Build the key for each group
-	const TArray<FString> GroupDerivedDataKeys = GetGroupDerivedDataKeys(this, TargetPlatform);
+	const int32 SourceMeshLOD = GetClampedSourceMeshLOD(TargetPlatform);
+	const int32 TargetMeshMinLOD = GetClampedTargetMeshMinLOD(TargetPlatform);
+	const TArray<FString> GroupDerivedDataKeys = GetGroupDerivedDataKeys(this, TargetPlatform, SourceMeshLOD, TargetMeshMinLOD);
 
 	// 2. Find existing cached cooked data
 	if (UGroomBindingAsset::FCachedCookedPlatformData* CachedCookedPlatformData = FindCachedCookedPlatformData(GroupDerivedDataKeys, CachedCookedPlatformDatas))
@@ -1442,6 +1714,10 @@ DEFINE_GROOMBINDING_MEMBER_NAME(SourceSkeletalMesh);
 DEFINE_GROOMBINDING_MEMBER_NAME(TargetSkeletalMesh);
 DEFINE_GROOMBINDING_MEMBER_GETTER(USkeletalMesh, *, SourceSkeletalMesh, , );
 DEFINE_GROOMBINDING_MEMBER_GETTER(USkeletalMesh, *, TargetSkeletalMesh, , );
+DEFINE_GROOMBINDING_MEMBER_ACCESSOR(int32, , SourceMeshRequestedLOD, , );
+DEFINE_GROOMBINDING_MEMBER_ACCESSOR(int32, , SourceMeshUsedLOD, , );
+DEFINE_GROOMBINDING_MEMBER_ACCESSOR(int32, , TargetMeshRequestedMinLOD, , );
+DEFINE_GROOMBINDING_MEMBER_ACCESSOR(int32, , TargetMeshUsedMinLOD, , );
 DEFINE_GROOMBINDING_MEMBER_ACCESSOR(UGeometryCache, *, SourceGeometryCache, , );
 DEFINE_GROOMBINDING_MEMBER_ACCESSOR(UGeometryCache, *, TargetGeometryCache, , );
 DEFINE_GROOMBINDING_MEMBER_ACCESSOR(int32, , NumInterpolationPoints, , );
@@ -1663,4 +1939,4 @@ void UGroomBindingAsset::ChangePlatformLevel(ERHIFeatureLevel::Type In)
 		CachedResourcesPlatformLevel = In;
 	}
 }
-#endif // WITH_EDITOR
\ No newline at end of file
+#endif // WITH_EDITOR
diff --git a/Engine/Plugins/Runtime/HairStrands/Source/HairStrandsCore/Private/GroomBindingBuilder.cpp b/Engine/Plugins/Runtime/HairStrands/Source/HairStrandsCore/Private/GroomBindingBuilder.cpp
--- a/Engine/Plugins/Runtime/HairStrands/Source/HairStrandsCore/Private/GroomBindingBuilder.cpp
+++ b/Engine/Plugins/Runtime/HairStrands/Source/HairStrandsCore/Private/GroomBindingBuilder.cpp
@@ -6,6 +6,7 @@
 #include "GeometryCacheMeshData.h"
 #include "GroomAsset.h"
 #include "GroomBindingAsset.h"
+#include "Logging/StructuredLog.h"
 #include "Rendering/SkeletalMeshLODRenderData.h"
 #include "Rendering/SkeletalMeshRenderData.h"
 #include "HairStrandsMeshProjection.h"
@@ -18,8 +19,6 @@
 #include "Engine/SkinnedAssetAsyncCompileUtils.h"
 #include "Interfaces/ITargetPlatform.h"
 
-#if WITH_EDITORONLY_DATA
-
 ///////////////////////////////////////////////////////////////////////////////////////////////////
 // Eigen for large matrix inversion
 // Just to be sure, also added this in Eigen.Build.cs
@@ -44,8 +43,6 @@ PRAGMA_DEFAULT_VISIBILITY_END
 #pragma warning(pop)
 #endif
 
-#endif
-
 // Run the binding asset building in parallel (faster)
 #define BINDING_PARALLEL_BUILDING 1
 
@@ -60,6 +57,22 @@ static FAutoConsoleVariableRef CVarHairStrandsBindingBuilderWarningEnable(TEXT("
 
 ///////////////////////////////////////////////////////////////////////////////////////////////////
 
+FGroomBindingBuilder::FInput::FInput(UGroomBindingAsset* BindingAsset, const ITargetPlatform* TargetPlatform, int32 InSourceMeshLOD, int32 InTargetMeshMinLOD)
+: BindingType(BindingAsset->GetGroomBindingType())
+, NumInterpolationPoints(BindingAsset->GetNumInterpolationPoints())
+, MatchingSection(BindingAsset->GetMatchingSection())
+, SourceMeshLOD(InSourceMeshLOD)
+, TargetMeshMinLOD(InTargetMeshMinLOD)
+, bHasValidTarget(BindingAsset->HasValidTarget())
+, GroomAsset(BindingAsset->GetGroom())
+, SourceSkeletalMesh(BindingAsset->GetSourceSkeletalMesh())
+, TargetSkeletalMesh(BindingAsset->GetTargetSkeletalMesh())
+, SourceGeometryCache(BindingAsset->GetSourceGeometryCache())
+, TargetGeometryCache(BindingAsset->GetTargetGeometryCache())
+{
+	check(BindingType == EGroomBindingMeshType::SkeletalMesh || (SourceMeshLOD == 0 && TargetMeshMinLOD == 0));
+}
+
 FString FGroomBindingBuilder::GetVersion()
 {
 	// Important to update the version when groom building changes
@@ -125,6 +138,60 @@ static float PackNormalToFloat(const FVector3f& InN)
 
 namespace GroomBinding_Mesh
 {
+////////////////////////////////////////////////////////////////////////////////////////
+// Data structures to allow curves to be read from either bulk data or mesh description
+////////////////////////////////////////////////////////////////////////////////////////
+
+struct FHairStrandRootPositions
+{
+	FHairStrandRootPositions() = default;
+	explicit FHairStrandRootPositions(const FHairStrandsDatas& Source)
+	{
+		RootPositions.Reserve(Source.GetNumCurves());
+
+		for (uint32 CurveIndex = 0; CurveIndex < Source.GetNumCurves(); CurveIndex++)
+		{
+			const uint32 PointIndex = Source.StrandsCurves.CurvesOffset[CurveIndex];
+			RootPositions.Add(Source.StrandsPoints.PointsPosition[PointIndex]);
+		}
+
+		BoundingBox = Source.BoundingBox;
+	}
+
+	explicit FHairStrandRootPositions(const FHairStrandsBulkData& Source)
+	{
+		const uint32 NumPoints = Source.GetNumPoints();
+		const uint32 NumCurves = Source.GetNumCurves();
+		RootPositions.Reserve(NumCurves);
+
+		const FVector BoundingBoxCenter = Source.GetPositionOffset();
+
+		const FHairStrandsCurveFormat::Type* Curves = reinterpret_cast<const FHairStrandsCurveFormat::Type*>(Source.Data.Curves.Data.LockReadOnly());
+		check(Source.Data.Curves.GetBulkDataSize() == NumCurves * sizeof(FHairStrandsCurveFormat::Type));
+
+		const FHairStrandsPositionFormat::Type* Positions = reinterpret_cast<const FHairStrandsPositionFormat::Type*>(Source.Data.Positions.Data.LockReadOnly());
+		check(Source.Data.Positions.GetBulkDataSize() == NumPoints * sizeof(FHairStrandsPositionFormat::Type));
+
+		for (uint32 CurveIndex = 0; CurveIndex < NumCurves; CurveIndex++)
+		{
+			const uint32 PointIndex = Curves[CurveIndex].PointOffset;
+			const FHairStrandsPositionFormat::Type& Vertex = Positions[PointIndex];
+			RootPositions.Add(FVector3f(BoundingBoxCenter + FVector(Vertex.X, Vertex.Y, Vertex.Z)));
+		}
+
+		Source.Data.Positions.Data.Unlock();
+		Source.Data.Curves.Data.Unlock();
+
+		BoundingBox = FBox3f(Source.GetBounds());
+	}
+
+	uint32 GetNumCurves() const { return RootPositions.Num(); }
+
+	TArray<FVector3f> RootPositions;
+
+	FBox3f BoundingBox = FBox3f(ForceInit);
+};
+
 //////////////////////////////////////////////////////////////////////////
 // Interfaces to query mesh data from different sources
 //////////////////////////////////////////////////////////////////////////
@@ -161,8 +228,9 @@ class IMeshData
 {
 public:
 	virtual bool IsValid() const = 0;
-	virtual uint32 GetNumLODs() const = 0;
-	virtual const IMeshLODData& GetMeshLODData(uint32 LODIndex) const = 0;
+	virtual int32 GetNumLODs() const = 0;
+	virtual bool IsLODValid(int32 LODIndex) const = 0;
+	virtual const IMeshLODData& GetMeshLODData(int32 LODIndex) const = 0;
 	virtual ~IMeshData() {}
 };
 
@@ -221,7 +289,7 @@ public:
 		LODIndex = InLODIndex;
 		MeshData = InMeshData;
 
-		if (MeshData && MeshData->LODRenderData.IsValidIndex(LODIndex))
+		if (IsSourceGeometryAvailable())
 		{
 			const uint32 SectionCount = MeshData->LODRenderData[LODIndex].RenderSections.Num();
 			Sections.Reserve(SectionCount);
@@ -301,6 +369,13 @@ public:
 		return OutSectionIndex;
 	}
 
+	bool IsSourceGeometryAvailable() const
+	{
+		return MeshData 
+			&& MeshData->LODRenderData.IsValidIndex(LODIndex)
+			&& MeshData->LODRenderData[LODIndex].StaticVertexBuffers.PositionVertexBuffer.GetVertexData() != nullptr;
+	}
+
 private:
 	int32 LODIndex = -1;
 	const FSkeletalMeshRenderData* MeshData = nullptr;
@@ -312,7 +387,7 @@ private:
 class FSkeletalMeshData : public IMeshData
 {
 public:
-	FSkeletalMeshData(const USkeletalMesh* InSkeletalMesh, const FSkeletalMeshRenderData* InRenderData) 
+	FSkeletalMeshData(const USkeletalMesh* InSkeletalMesh, const FSkeletalMeshRenderData* InRenderData, int32 MeshMinLOD) 
 	{
 		SkeletalMesh = InSkeletalMesh;
 		MeshData = InRenderData;
@@ -322,9 +397,11 @@ public:
 			{
 				LODCount = SkeletalMesh->GetLODNum();
 				MeshesLODData.Reserve(LODCount);
-				for (uint32 LODIt=0; LODIt<LODCount; ++LODIt)
+				for (int32 LODIt = 0; LODIt < LODCount; ++LODIt)
 				{
-					MeshesLODData.Add(FSkeletalMeshLODData(MeshData, LODIt));
+					// Only pass in the mesh data if we're building for this LOD
+					const FSkeletalMeshRenderData* MeshDataForLOD = LODIt >= MeshMinLOD ? MeshData : nullptr;
+					MeshesLODData.Add(FSkeletalMeshLODData(MeshDataForLOD, LODIt));
 				}
 			}
 			else
@@ -339,19 +416,29 @@ public:
 		return SkeletalMesh != nullptr && MeshData != nullptr && MeshesLODData.Num() > 0;
 	}
 
-	virtual uint32 GetNumLODs() const override
+	virtual int32 GetNumLODs() const override
 	{
 		return LODCount;
 	}
 
-	virtual const IMeshLODData& GetMeshLODData(uint32 InLODIndex) const override
+	virtual bool IsLODValid(int32 LODIndex) const override
+	{
+		if (!MeshesLODData.IsValidIndex(LODIndex))
+		{
+			return false;
+		}
+
+		return MeshesLODData[LODIndex].IsSourceGeometryAvailable();
+	}
+
+	virtual const IMeshLODData& GetMeshLODData(int32 InLODIndex) const override
 	{
 		check(MeshesLODData.IsValidIndex(InLODIndex));
 		return MeshesLODData[InLODIndex];
 	}
 
 private:
-	uint32 LODCount = 0;
+	int32 LODCount = 0;
 	const USkeletalMesh* SkeletalMesh = nullptr;
 	const FSkeletalMeshRenderData* MeshData = nullptr;
 	TArray<FSkeletalMeshLODData> MeshesLODData;
@@ -445,12 +532,18 @@ public:
 		return GeometryCache != nullptr && Sections.Num() > 0;
 	}
 
-	virtual uint32 GetNumLODs() const override
+	virtual int32 GetNumLODs() const override
 	{
 		return 1;
 	}
 
-	virtual const IMeshLODData& GetMeshLODData(uint32 LODIndex) const override
+	virtual bool IsLODValid(int32 LODIndex) const override
+	{
+		// There is only 1 LOD and it's always considered valid
+		return LODIndex == 0;
+	}
+
+	virtual const IMeshLODData& GetMeshLODData(int32 LODIndex) const override
 	{
 		return *this;
 	}
@@ -606,7 +699,7 @@ namespace GroomBinding_RBFWeighting
 		}
 	}
 }
-#if WITH_EDITORONLY_DATA
+
 namespace GroomBinding_RBFWeighting
 {
 	struct FWeightsBuilder
@@ -795,13 +888,21 @@ namespace GroomBinding_RBFWeighting
 		const bool bNeedStrandsRoot,
 		const uint32 NumInterpolationPoints, 
 		const int32 MatchingSection, 
-		const GroomBinding_Mesh::IMeshData* MeshData)
+		const GroomBinding_Mesh::IMeshData* MeshData,
+		const int32 TargetMeshMinLOD)
 	{
-		const uint32 MeshLODCount= MeshData->GetNumLODs();
+		const int32 MeshLODCount = MeshData->GetNumLODs();
 		const uint32 MaxSamples  = NumInterpolationPoints;
 
-		for (uint32 MeshLODIndex = 0; MeshLODIndex < MeshLODCount; ++MeshLODIndex)
+		for (int32 MeshLODIndex = 0; MeshLODIndex < MeshLODCount; ++MeshLODIndex)
 		{
+			if (MeshLODIndex < TargetMeshMinLOD)
+			{
+				check(!Out.SimRootDatas[MeshLODIndex].IsValid());
+				check(!bNeedStrandsRoot || !Out.RenRootDatas[MeshLODIndex].IsValid());
+				continue;
+			}
+
 			const GroomBinding_Mesh::IMeshLODData& MeshLODData = MeshData->GetMeshLODData(MeshLODIndex);
 
 			int32 TargetSection = -1;
@@ -865,8 +966,6 @@ namespace GroomBinding_RBFWeighting
 	}
 }// namespace GroomBinding_RBFWeighting
 
-#endif
-
 ///////////////////////////////////////////////////////////////////////////////////////////////////
 // Root projection
 
@@ -1173,17 +1272,18 @@ namespace GroomBinding_RootProjection
 	}
 
 	static bool Project(
-		const FHairStrandsDatas& InStrandsData,
+		const GroomBinding_Mesh::FHairStrandRootPositions& InRootPositions,
 		const GroomBinding_Mesh::IMeshData* InMeshData,
+		int32 TargetMeshMinLOD,
 		const TArray<TArray<FVector3f>>& InTransferredPositions,
 		TArray<FHairStrandsRootData>& OutRootData,
 		uint32 InMatchingSection)
 	{
 		// 2. Project root for each mesh LOD
-		const uint32 CurveCount = InStrandsData.GetNumCurves();
+		const uint32 CurveCount = InRootPositions.GetNumCurves();
 		const uint32 ChannelIndex = 0;
 		const float VoxelWorldSize = 2; //cm
-		const uint32 MeshLODCount = InMeshData->GetNumLODs();
+		const int32 MeshLODCount = InMeshData->GetNumLODs();
 		check(MeshLODCount == OutRootData.Num());
 
 		const bool bHasTransferredPosition = InTransferredPositions.Num() > 0;
@@ -1192,10 +1292,16 @@ namespace GroomBinding_RootProjection
 			check(InTransferredPositions.Num() == MeshLODCount);
 		}
 
-		for (uint32 MeshLODIt = 0; MeshLODIt < MeshLODCount; ++MeshLODIt)
+		for (int32 MeshLODIt = 0; MeshLODIt < MeshLODCount; ++MeshLODIt)
 		{
 			check(MeshLODIt == OutRootData[MeshLODIt].LODIndex);
 
+			if (MeshLODIt < TargetMeshMinLOD)
+			{
+				check(!OutRootData[MeshLODIt].IsValid());
+				continue;
+			}
+
 			// 2.1. Build a grid around the hair AABB
 			const GroomBinding_Mesh::IMeshLODData& MeshLODData = InMeshData->GetMeshLODData(MeshLODIt);
 			const TArray<uint32>& IndexBuffer = MeshLODData.GetIndexBuffer();
@@ -1279,15 +1385,15 @@ namespace GroomBinding_RootProjection
 					MeshBound += T.P2;
 
 					// Track closest point to the groom bound
-					ClosestTrianglePoint = FMath::Min(ClosestTrianglePoint, (T.P0 - InStrandsData.BoundingBox.GetCenter()).Length());
-					ClosestTrianglePoint = FMath::Min(ClosestTrianglePoint, (T.P1 - InStrandsData.BoundingBox.GetCenter()).Length());
-					ClosestTrianglePoint = FMath::Min(ClosestTrianglePoint, (T.P2 - InStrandsData.BoundingBox.GetCenter()).Length());
+					ClosestTrianglePoint = FMath::Min(ClosestTrianglePoint, (T.P0 - InRootPositions.BoundingBox.GetCenter()).Length());
+					ClosestTrianglePoint = FMath::Min(ClosestTrianglePoint, (T.P1 - InRootPositions.BoundingBox.GetCenter()).Length());
+					ClosestTrianglePoint = FMath::Min(ClosestTrianglePoint, (T.P2 - InRootPositions.BoundingBox.GetCenter()).Length());
 				}
 			}
 
 			// Take the smallest bounding box between the groom and the skeletal mesh
 			const FVector3f MeshExtent = MeshBound.Max - MeshBound.Min;
-			const FVector3f HairExtent = InStrandsData.BoundingBox.Max - InStrandsData.BoundingBox.Min;
+			const FVector3f HairExtent = InRootPositions.BoundingBox.Max - InRootPositions.BoundingBox.Min;
 			FVector3f GridMin;
 			FVector3f GridMax;
 			if (MeshExtent.Size() < HairExtent.Size())
@@ -1297,8 +1403,8 @@ namespace GroomBinding_RootProjection
 			}
 			else
 			{
-				GridMin = InStrandsData.BoundingBox.Min;
-				GridMax = InStrandsData.BoundingBox.Max;
+				GridMin = InRootPositions.BoundingBox.Min;
+				GridMax = InRootPositions.BoundingBox.Max;
 
 				// By nature, it is possible that coarser LOD have positions which ressemble only very coarsly to 
 				// LOD0. In this case we increase the hair bound to ensure that skel. mesh triangles will be intersect 
@@ -1396,7 +1502,7 @@ namespace GroomBinding_RootProjection
 			ParallelFor(CurveCount,
 				[
 					MeshLODIt,
-					&InStrandsData,
+					&InRootPositions,
 					&Grid,
 					&RootTriangleIndexBuffer,
 					&RestRootTrianglePositionBuffer,
@@ -1408,8 +1514,7 @@ namespace GroomBinding_RootProjection
 			for (uint32 CurveIndex = 0; CurveIndex < CurveCount; ++CurveIndex)
 		#endif
 			{
-				const uint32 Offset = InStrandsData.StrandsCurves.CurvesOffset[CurveIndex];
-				const FVector3f& RootP = InStrandsData.StrandsPoints.PointsPosition[Offset];
+				const FVector3f& RootP = InRootPositions.RootPositions[CurveIndex];
 				const FTriangleGrid::FCells Cells = Grid.ToCells(RootP);
 
 				if (Cells.Num() == 0)
@@ -1810,6 +1915,8 @@ namespace GroomBinding_Transfer
 	bool Transfer(
 		const GroomBinding_Mesh::IMeshData* InSourceMeshData,
 		const GroomBinding_Mesh::IMeshData* InTargetMeshData,
+		int32 TargetMeshMinLOD,
+		int32 SourceMeshLODIndex,
 		TArray<TArray<FVector3f>>& OutPositions_Original, 
 		TArray<TArray<FVector3f>>& OutPositions_Transferred,
 		const int32 MatchingSection)
@@ -1867,8 +1974,7 @@ namespace GroomBinding_Transfer
 
 		// 1. Insert triangles into a 2D UV grid
 		const uint32 ChannelIndex = 0;
-		const uint32 SourceLODIndex = 0;
-		const GroomBinding_Mesh::IMeshLODData& SourceMeshLODData = InSourceMeshData->GetMeshLODData(SourceLODIndex);
+		const GroomBinding_Mesh::IMeshLODData& SourceMeshLODData = InSourceMeshData->GetMeshLODData(SourceMeshLODIndex);
 		const bool bIsMatchingSectionValid = MatchingSection < SourceMeshLODData.GetNumSections();
 		const int32 SourceSectionId = bIsMatchingSectionValid ? MatchingSection : 0;
 		if (!bIsMatchingSectionValid && GHairStrandsBindingBuilderWarningEnable > 0)
@@ -1878,7 +1984,7 @@ namespace GroomBinding_Transfer
 		const int32 TargetSectionId = SourceSectionId;
 		FTriangleGrid2D Grid(256);
 		{
-			const bool bIsGridPopulated = BuildGrid(SourceLODIndex, SourceSectionId, TargetSectionId, ChannelIndex, Grid);
+			const bool bIsGridPopulated = BuildGrid(SourceMeshLODIndex, SourceSectionId, TargetSectionId, ChannelIndex, Grid);
 			if (!bIsGridPopulated)
 			{
 				UE_LOG(LogHairStrands, Error, TEXT("[Groom] Binding asset could not be built. The source skeletal mesh is missing or has invalid UVs."));
@@ -1888,11 +1994,18 @@ namespace GroomBinding_Transfer
 
 		// 2. Look for closest triangle point in UV space
 		// Make this run in parallel
-		const uint32 TargetLODCount = InTargetMeshData->GetNumLODs();
+		const int32 TargetLODCount = InTargetMeshData->GetNumLODs();
 		OutPositions_Original.SetNum(TargetLODCount);
 		OutPositions_Transferred.SetNum(TargetLODCount);
-		for (uint32 TargetLODIndex = 0; TargetLODIndex < TargetLODCount; ++TargetLODIndex)
+		for (int32 TargetLODIndex = 0; TargetLODIndex < TargetLODCount; ++TargetLODIndex)
 		{
+			if (TargetLODIndex < TargetMeshMinLOD)
+			{
+				OutPositions_Original[TargetLODIndex].Empty();
+				OutPositions_Transferred[TargetLODIndex].Empty();
+				continue;
+			}
+
 			// Check that the target SectionId is valid for the current LOD. 
 			// If this is not the case, then fall back to section 0 and rebuild the source triangle grid to match the same section ID (1.)
 			int32 LocalSourceSectionId = SourceSectionId;
@@ -1914,7 +2027,7 @@ namespace GroomBinding_Transfer
 
 				LocalTargetSectionId = 0;
 				LocalSourceSectionId = 0;
-				const bool bIsGridPopulated = BuildGrid(SourceLODIndex, LocalSourceSectionId, LocalTargetSectionId, ChannelIndex, Grid);
+				const bool bIsGridPopulated = BuildGrid(SourceMeshLODIndex, LocalSourceSectionId, LocalTargetSectionId, ChannelIndex, Grid);
 				if (!bIsGridPopulated)
 				{
 					UE_LOG(LogHairStrands, Error, TEXT("[Groom] Binding asset could not be built for LOD %d. The source skeletal mesh is missing or has invalid UVs."), TargetLODIndex);
@@ -2008,17 +2121,25 @@ namespace GroomBinding_Transfer
 }
 // namespace GroomBinding_Transfer
 
-static void InitHairStrandsRootData(TArray<FHairStrandsRootData>& Out, const FHairStrandsDatas* HairStrandsDatas, uint32 LODCount, uint32 NumSamples)
+static void InitHairStrandsRootData(TArray<FHairStrandsRootData>& Out, uint32 CurveCount, uint32 PointCount, uint32 LODCount, uint32 MinLOD, uint32 NumSamples)
 {
-	check(HairStrandsDatas);
-
 	Out.SetNum(LODCount);
 	uint32 LODIndex = 0;
 	for (FHairStrandsRootData& OutLOD : Out)
 	{
-		OutLOD.RootCount = HairStrandsDatas->GetNumCurves();
-		OutLOD.PointCount = HairStrandsDatas->GetNumPoints();
-		OutLOD.SampleCount = NumSamples;
+		if (LODIndex >= MinLOD)
+		{
+			OutLOD.RootCount = CurveCount;
+			OutLOD.PointCount = PointCount;
+			OutLOD.SampleCount = NumSamples;
+		}
+		else
+		{
+			OutLOD.RootCount = 0;
+			OutLOD.PointCount = 0;
+			OutLOD.SampleCount = 0;
+		}
+
 		OutLOD.MeshInterpolationWeightsBuffer.Empty();
 		OutLOD.MeshSampleIndicesBuffer.Empty();
 		OutLOD.RestSamplePositionsBuffer.Empty();
@@ -2200,42 +2321,54 @@ static void BuildRootData(
 }
 
 // Convert the root data into root bulk data
-static void BuildRootBulkData(
+static void BuildGroupPlatformData(
 	UGroomBindingAsset::FHairGroupPlatformData& Out,
 	const FHairRootGroupData& In)
 {
 	// Guides
 	{
-		const uint32 MeshLODCount = In.SimRootDatas.Num();
+		const int32 MeshLODCount = In.SimRootDatas.Num();
 		Out.SimRootBulkDatas.SetNum(MeshLODCount);
-		for (uint32 MeshLODIt = 0; MeshLODIt < MeshLODCount; ++MeshLODIt)
+		for (int32 MeshLODIt = In.TargetMeshMinLOD; MeshLODIt < MeshLODCount; ++MeshLODIt)
 		{
+			if (MeshLODIt < In.TargetMeshMinLOD)
+			{
+				continue;
+			}
+
 			BuildRootBulkData(Out.SimRootBulkDatas[MeshLODIt], In.SimRootDatas[MeshLODIt]);
 		}
 	}
 
 	// Strands
 	{
-		const uint32 MeshLODCount = In.RenRootDatas.Num();
+		const int32 MeshLODCount = In.RenRootDatas.Num();
 		Out.RenRootBulkDatas.SetNum(MeshLODCount);
-		for (uint32 MeshLODIt = 0; MeshLODIt < MeshLODCount; ++MeshLODIt)
+		for (int32 MeshLODIt = In.TargetMeshMinLOD; MeshLODIt < MeshLODCount; ++MeshLODIt)
 		{
+			if (MeshLODIt < In.TargetMeshMinLOD)
+			{
+				continue;
+			}
+
 			BuildRootBulkData(Out.RenRootBulkDatas[MeshLODIt], In.RenRootDatas[MeshLODIt]);
 		}
 	}
 
 	// Cards
-	const uint32 LODCount = In.CardsRootDatas.Num();
+	const int32 LODCount = In.CardsRootDatas.Num();
 	Out.CardsRootBulkDatas.SetNum(LODCount);
-	for (uint32 LODIt = 0; LODIt < LODCount; ++LODIt)
+	for (int32 LODIt = 0; LODIt < LODCount; ++LODIt)
 	{
-		const uint32 MeshLODCount = In.CardsRootDatas[LODIt].Num();
+		const int32 MeshLODCount = In.CardsRootDatas[LODIt].Num();
 		Out.CardsRootBulkDatas[LODIt].SetNum(MeshLODCount);
-		for (uint32 MeshLODIt = 0; MeshLODIt < MeshLODCount; ++MeshLODIt)
+		for (int32 MeshLODIt = In.TargetMeshMinLOD; MeshLODIt < MeshLODCount; ++MeshLODIt)
 		{
 			BuildRootBulkData(Out.CardsRootBulkDatas[LODIt][MeshLODIt], In.CardsRootDatas[LODIt][MeshLODIt]);
 		}
 	}
+
+	Out.TargetMeshMinLOD = In.TargetMeshMinLOD;
 }
 
 } // namespace GroomBinding_BulkCopy
@@ -2244,7 +2377,6 @@ static void BuildRootBulkData(
 // Main entry (CPU path)
 bool BuildHairRootGroupData(const FGroomBindingBuilder::FInput& In, uint32 InGroupIndex, const ITargetPlatform* TargetPlatform, FHairRootGroupData& OutData)
 {
-#if WITH_EDITORONLY_DATA
 	const bool bIsAssetValid = In.GroomAsset && In.bHasValidTarget && In.GroomAsset->GetNumHairGroups() > 0;
 	if (!bIsAssetValid)
 	{
@@ -2285,16 +2417,22 @@ bool BuildHairRootGroupData(const FGroomBindingBuilder::FInput& In, uint32 InGro
 		// will take care of finishing any groom binding that depends on the skeletal mesh being modified. So this is 
 		// safe to do asynchronously without locks on the render data. We assume that the skeletal mesh's render data is
 		// immutable once it has been built and can only be rebuilt throught a call to PreEditChange first.
-		const bool bAsyncCompiling = !IsInGameThread();
+#if WITH_EDITOR
+		const bool bUseRunningPlatform = !IsInGameThread();
+#else
+		const bool bUseRunningPlatform = true;
+#endif
 
 		// * Only for SkeletalMesh: Take scoped lock on the skeletal render mesh data during the entire groom binding building
 		// * Then use an async build scope to allow accessing skeletal mesh property safely.
 		//   If skel.meshes are nullptr, this will act as a NOP
 		USkeletalMesh* InSourceSkeletalMesh = In.SourceSkeletalMesh == In.TargetSkeletalMesh ? nullptr : In.SourceSkeletalMesh;
+#if WITH_EDITOR
 		PRAGMA_DISABLE_DEPRECATION_WARNINGS
-		FScopedSkeletalMeshRenderData SourceSkeletalMeshScopedData(bAsyncCompiling ? nullptr : InSourceSkeletalMesh);
-		FScopedSkeletalMeshRenderData TargetSkeletalMeshScopedData(bAsyncCompiling ? nullptr : In.TargetSkeletalMesh);
+		FScopedSkeletalMeshRenderData SourceSkeletalMeshScopedData(bUseRunningPlatform ? nullptr : InSourceSkeletalMesh);
+		FScopedSkeletalMeshRenderData TargetSkeletalMeshScopedData(bUseRunningPlatform ? nullptr : In.TargetSkeletalMesh);
 		PRAGMA_ENABLE_DEPRECATION_WARNINGS
+#endif
 
 		TUniquePtr<GroomBinding_Mesh::IMeshData> SourceMeshData;
 		TUniquePtr<GroomBinding_Mesh::IMeshData> TargetMeshData;
@@ -2302,50 +2440,62 @@ bool BuildHairRootGroupData(const FGroomBindingBuilder::FInput& In, uint32 InGro
 		{
 			if (InSourceSkeletalMesh)
 			{
-				if (bAsyncCompiling)
+				if (bUseRunningPlatform)
 				{
+#if WITH_EDITOR
 					const ITargetPlatform* RunningPlatform = GetTargetPlatformManagerRef().GetRunningTargetPlatform();
 					checkf(TargetPlatform == RunningPlatform, TEXT("It is only safe to query the running platform's render data asynchronously from the skeletal mesh"));
-					SourceMeshData = TUniquePtr<GroomBinding_Mesh::FSkeletalMeshData, TDefaultDelete<GroomBinding_Mesh::IMeshData>>(new GroomBinding_Mesh::FSkeletalMeshData(InSourceSkeletalMesh, InSourceSkeletalMesh->GetResourceForRendering()));
+#endif
+					SourceMeshData = TUniquePtr<GroomBinding_Mesh::FSkeletalMeshData, TDefaultDelete<GroomBinding_Mesh::IMeshData>>(new GroomBinding_Mesh::FSkeletalMeshData(InSourceSkeletalMesh, InSourceSkeletalMesh->GetResourceForRendering(), In.SourceMeshLOD));
 				}
+#if WITH_EDITOR
 				else
 				{
 					FSkinnedAssetAsyncBuildScope AsyncBuildScope(InSourceSkeletalMesh);
 					USkeletalMesh::GetPlatformSkeletalMeshRenderData(TargetPlatform, SourceSkeletalMeshScopedData);
 					PRAGMA_DISABLE_DEPRECATION_WARNINGS
-					SourceMeshData = TUniquePtr<GroomBinding_Mesh::FSkeletalMeshData, TDefaultDelete<GroomBinding_Mesh::IMeshData>>(new GroomBinding_Mesh::FSkeletalMeshData(InSourceSkeletalMesh, SourceSkeletalMeshScopedData.GetData()));
+					SourceMeshData = TUniquePtr<GroomBinding_Mesh::FSkeletalMeshData, TDefaultDelete<GroomBinding_Mesh::IMeshData>>(new GroomBinding_Mesh::FSkeletalMeshData(InSourceSkeletalMesh, SourceSkeletalMeshScopedData.GetData(), In.SourceMeshLOD));
 					PRAGMA_ENABLE_DEPRECATION_WARNINGS
 				}
+#endif
 			}
 			else
 			{
-				SourceMeshData = TUniquePtr<GroomBinding_Mesh::FSkeletalMeshData, TDefaultDelete<GroomBinding_Mesh::IMeshData>>(new GroomBinding_Mesh::FSkeletalMeshData(nullptr, nullptr));
+				SourceMeshData = TUniquePtr<GroomBinding_Mesh::FSkeletalMeshData, TDefaultDelete<GroomBinding_Mesh::IMeshData>>(new GroomBinding_Mesh::FSkeletalMeshData(nullptr, nullptr, In.SourceMeshLOD));
 			}
 
 			if (In.TargetSkeletalMesh)
 			{
-				if (bAsyncCompiling)
+				if (bUseRunningPlatform)
 				{
+#if WITH_EDITOR
 					const ITargetPlatform* RunningPlatform = GetTargetPlatformManagerRef().GetRunningTargetPlatform();
 					checkf(TargetPlatform == RunningPlatform, TEXT("It is only safe to query the running platform's render data asynchronously from the skeletal mesh"));
-					TargetMeshData = TUniquePtr<GroomBinding_Mesh::FSkeletalMeshData, TDefaultDelete<GroomBinding_Mesh::IMeshData>>(new GroomBinding_Mesh::FSkeletalMeshData(In.TargetSkeletalMesh, In.TargetSkeletalMesh->GetResourceForRendering()));
+#endif
+					TargetMeshData = TUniquePtr<GroomBinding_Mesh::FSkeletalMeshData, TDefaultDelete<GroomBinding_Mesh::IMeshData>>(new GroomBinding_Mesh::FSkeletalMeshData(In.TargetSkeletalMesh, In.TargetSkeletalMesh->GetResourceForRendering(), In.TargetMeshMinLOD));
 				}
+#if WITH_EDITOR
 				else
 				{
 					FSkinnedAssetAsyncBuildScope AsyncBuildScope(In.TargetSkeletalMesh);
 					USkeletalMesh::GetPlatformSkeletalMeshRenderData(TargetPlatform, TargetSkeletalMeshScopedData);
 					PRAGMA_DISABLE_DEPRECATION_WARNINGS
-					TargetMeshData = TUniquePtr<GroomBinding_Mesh::FSkeletalMeshData, TDefaultDelete<GroomBinding_Mesh::IMeshData>>(new GroomBinding_Mesh::FSkeletalMeshData(In.TargetSkeletalMesh, TargetSkeletalMeshScopedData.GetData()));
+					TargetMeshData = TUniquePtr<GroomBinding_Mesh::FSkeletalMeshData, TDefaultDelete<GroomBinding_Mesh::IMeshData>>(new GroomBinding_Mesh::FSkeletalMeshData(In.TargetSkeletalMesh, TargetSkeletalMeshScopedData.GetData(), In.TargetMeshMinLOD));
 					PRAGMA_ENABLE_DEPRECATION_WARNINGS
 				}
+#endif
 			}
 			else
 			{
-				TargetMeshData = TUniquePtr<GroomBinding_Mesh::FSkeletalMeshData, TDefaultDelete<GroomBinding_Mesh::IMeshData>>(new GroomBinding_Mesh::FSkeletalMeshData(nullptr, nullptr));
+				TargetMeshData = TUniquePtr<GroomBinding_Mesh::FSkeletalMeshData, TDefaultDelete<GroomBinding_Mesh::IMeshData>>(new GroomBinding_Mesh::FSkeletalMeshData(nullptr, nullptr, In.TargetMeshMinLOD));
 			}
 		}
 		else
 		{
+			// Geometry caches only have one LOD
+			check(In.SourceMeshLOD == 0);
+			check(In.TargetMeshMinLOD == 0);
+
 			In.TargetGeometryCache->ConditionalPostLoad();
 			if (In.SourceGeometryCache)
 			{
@@ -2361,12 +2511,31 @@ bool BuildHairRootGroupData(const FGroomBindingBuilder::FInput& In, uint32 InGro
 			UE_LOG(LogHairStrands, Error, TEXT("[Groom] Binding asset could not be built. Target mesh is not valid."));
 			return false;
 		}
+
 		const uint32 GroupCount = In.GroomAsset->GetNumHairGroups();
 		const uint32 MeshLODCount = TargetMeshData->GetNumLODs();
 
 		check(InGroupIndex < uint32(In.GroomAsset->GetHairGroupsPlatformData().Num()));
 
+		if (SourceMeshData->IsValid() && !SourceMeshData->IsLODValid(In.SourceMeshLOD))
+		{
+			UE_LOGFMT(LogHairStrands, Error, "[Groom] Binding asset could not be built. Source mesh LOD {0} is not valid.", In.SourceMeshLOD);
+			return false;
+		}
+
+		for (uint32 MeshLODIndex = In.TargetMeshMinLOD; MeshLODIndex < MeshLODCount; MeshLODIndex++)
+		{
+			if (!TargetMeshData->IsLODValid(MeshLODIndex))
+			{
+				UE_LOGFMT(LogHairStrands, Error, "[Groom] Binding asset could not be built. Target mesh LOD {0} is not valid.", MeshLODIndex);
+				return false;
+			}
+		}
+
 		// Check if root data are needs for strands
+		//
+		// For now we process all groom LODs here. In future we could skip any LODs that will be
+		// stripped on the target platform.
 		bool bNeedStrandsRoot = false;
 		for (const FHairLODSettings& LODSettings : In.GroomAsset->GetHairGroupsLOD()[InGroupIndex].LODs)
 		{
@@ -2376,22 +2545,42 @@ bool BuildHairRootGroupData(const FGroomBindingBuilder::FInput& In, uint32 InGro
 				break;
 			}
 		}
-		
+
+		const FHairGroupPlatformData& GroupData = In.GroomAsset->GetHairGroupsPlatformData()[InGroupIndex];
+
 		// 1.1 Build guide/strands data
-		FHairStrandsDatas StrandsData;
-		FHairStrandsDatas GuidesData;
-		In.GroomAsset->GetHairStrandsDatas(InGroupIndex, StrandsData, GuidesData);
+#if WITH_EDITORONLY_DATA
+		FHairStrandsDatas StrandsDataSource;
+		FHairStrandsDatas GuidesDataSource;
+		if (!In.GroomAsset->GetHairStrandsDatas(InGroupIndex, StrandsDataSource, GuidesDataSource))
+		{
+			UE_LOGFMT(LogHairStrands, Error, "[Groom] Binding asset could not be built. Failed to read from {0}", In.GroomAsset->GetFullName());
+			return false;
+		}
+#else
+		const FHairStrandsBulkData& StrandsDataSource = GroupData.Strands.BulkData;
+		const FHairStrandsBulkData& GuidesDataSource = GroupData.Guides.BulkData;
+
+		if (!GroupData.Strands.HasValidData())
+		{
+			// Strands have been cooked out
+			bNeedStrandsRoot = false;
+		}
+#endif
 
 		// 1.2 Init root data for guides/strands/cards
-		const FHairGroupPlatformData& GroupData = In.GroomAsset->GetHairGroupsPlatformData()[InGroupIndex];
 		{
 			// Guides
-			InitHairStrandsRootData(OutData.SimRootDatas, &GuidesData, MeshLODCount, In.NumInterpolationPoints);
+			const uint32 GuidesCurveCount = GuidesDataSource.GetNumCurves();
+			const uint32 GuidesPointCount = GuidesDataSource.GetNumPoints();
+			InitHairStrandsRootData(OutData.SimRootDatas, GuidesCurveCount, GuidesPointCount, MeshLODCount, In.TargetMeshMinLOD, In.NumInterpolationPoints);
 
 			// Strands
 			if (bNeedStrandsRoot)
 			{
-				InitHairStrandsRootData(OutData.RenRootDatas, &StrandsData, MeshLODCount, In.NumInterpolationPoints);
+				const uint32 StrandsCurveCount = StrandsDataSource.GetNumCurves();
+				const uint32 StrandsPointCount = StrandsDataSource.GetNumPoints();
+				InitHairStrandsRootData(OutData.RenRootDatas, StrandsCurveCount, StrandsPointCount, MeshLODCount, In.TargetMeshMinLOD, In.NumInterpolationPoints);
 			}
 
 			// Cards
@@ -2401,11 +2590,23 @@ bool BuildHairRootGroupData(const FGroomBindingBuilder::FInput& In, uint32 InGro
 			{
 				if (GroupData.Cards.IsValid(CardsLODIt))
 				{
+#if WITH_EDITORONLY_DATA
 					FHairStrandsDatas LODGuidesData;
 					const bool bIsValid = In.GroomAsset->GetHairCardsGuidesDatas(InGroupIndex, CardsLODIt, LODGuidesData);
+#else
+					const bool bIsValid = GroupData.Cards.LODs[CardsLODIt].BulkData.IsValid();
+#endif
 					if (bIsValid)
 					{
-						InitHairStrandsRootData(OutData.CardsRootDatas[CardsLODIt], &LODGuidesData, MeshLODCount, In.NumInterpolationPoints);
+#if WITH_EDITORONLY_DATA
+						const uint32 CardsCurveCount = LODGuidesData.GetNumCurves();
+						const uint32 CardsPointCount = LODGuidesData.GetNumPoints(); 
+#else
+						const uint32 CardsCurveCount = GroupData.Cards.LODs[CardsLODIt].GuideBulkData.GetNumCurves();
+						const uint32 CardsPointCount = GroupData.Cards.LODs[CardsLODIt].GuideBulkData.GetNumPoints();
+#endif
+
+						InitHairStrandsRootData(OutData.CardsRootDatas[CardsLODIt], CardsCurveCount, CardsPointCount, MeshLODCount, In.TargetMeshMinLOD, In.NumInterpolationPoints);
 					}
 				}
 			}
@@ -2435,6 +2636,8 @@ bool BuildHairRootGroupData(const FGroomBindingBuilder::FInput& In, uint32 InGro
 			if (!GroomBinding_Transfer::Transfer(
 				SourceMeshData.Get(),
 				TargetMeshData.Get(),
+				In.TargetMeshMinLOD,
+				In.SourceMeshLOD,
 				OutData.MeshPositions, 
 				OutData.MeshPositions_Transferred, 
 				In.MatchingSection))
@@ -2449,9 +2652,12 @@ bool BuildHairRootGroupData(const FGroomBindingBuilder::FInput& In, uint32 InGro
 		{
 			// Guides
 			{
+				const GroomBinding_Mesh::FHairStrandRootPositions GuidesData(GuidesDataSource);
+
 				if (!GroomBinding_RootProjection::Project(
 					GuidesData,
 					TargetMeshData.Get(),
+					In.TargetMeshMinLOD,
 					OutData.MeshPositions_Transferred,
 					OutData.SimRootDatas,
 					In.MatchingSection))
@@ -2465,9 +2671,12 @@ bool BuildHairRootGroupData(const FGroomBindingBuilder::FInput& In, uint32 InGro
 			// Strands
 			if (bNeedStrandsRoot)
 			{
+				const GroomBinding_Mesh::FHairStrandRootPositions StrandsData(StrandsDataSource);
+
 				if (!GroomBinding_RootProjection::Project(
 					StrandsData,
 					TargetMeshData.Get(),
+					In.TargetMeshMinLOD,
 					OutData.MeshPositions_Transferred,
 					OutData.RenRootDatas,
 					In.MatchingSection))
@@ -2484,13 +2693,21 @@ bool BuildHairRootGroupData(const FGroomBindingBuilder::FInput& In, uint32 InGro
 			{
 				if (In.GroomAsset->GetHairGroupsPlatformData()[InGroupIndex].Cards.IsValid(CardsLODIt))
 				{
-					FHairStrandsDatas LODGuidesData;
-					const bool bIsValid = In.GroomAsset->GetHairCardsGuidesDatas(InGroupIndex, CardsLODIt, LODGuidesData);
+#if WITH_EDITORONLY_DATA
+					FHairStrandsDatas LODGuidesDataSource;
+					const bool bIsValid = In.GroomAsset->GetHairCardsGuidesDatas(InGroupIndex, CardsLODIt, LODGuidesDataSource);
+#else
+					const bool bIsValid = GroupData.Cards.LODs[CardsLODIt].IsValid();
+					const FHairStrandsBulkData& LODGuidesDataSource = GroupData.Cards.LODs[CardsLODIt].GuideBulkData;
+#endif
 					if (bIsValid)
 					{
+						const GroomBinding_Mesh::FHairStrandRootPositions LODGuidesData(LODGuidesDataSource);
+
 						if (!GroomBinding_RootProjection::Project(
 							LODGuidesData,
 							TargetMeshData.Get(),
+							In.TargetMeshMinLOD,
 							OutData.MeshPositions_Transferred,
 							OutData.CardsRootDatas[CardsLODIt],
 							0/*In.MatchingSection*/))
@@ -2506,36 +2723,38 @@ bool BuildHairRootGroupData(const FGroomBindingBuilder::FInput& In, uint32 InGro
 		
 		// 1.5 RBF building
 		{
-			GroomBinding_RBFWeighting::ComputeInterpolationWeights(OutData, bNeedStrandsRoot, In.NumInterpolationPoints, In.MatchingSection, TargetMeshData.Get());
+			GroomBinding_RBFWeighting::ComputeInterpolationWeights(OutData, bNeedStrandsRoot, In.NumInterpolationPoints, In.MatchingSection, TargetMeshData.Get(), In.TargetMeshMinLOD);
 			SlowTask.EnterProgressFrame();
 		}
 	}
-#endif
+
+	OutData.TargetMeshMinLOD = In.TargetMeshMinLOD;
 	return true;
 }
 
-
-bool FGroomBindingBuilder::BuildBinding(const FGroomBindingBuilder::FInput& In, uint32 InGroupIndex, const ITargetPlatform* TargetPlatform, UGroomBindingAsset::FHairGroupPlatformData& OutPlatformData)
+bool FGroomBindingBuilder::BuildBinding(
+	const FGroomBindingBuilder::FInput& In,
+	uint32 InGroupIndex,
+	const ITargetPlatform* TargetPlatform,
+	UGroomBindingAsset::FHairGroupPlatformData& OutPlatformData)
 {
-	bool bSucceed = true;
-#if WITH_EDITORONLY_DATA
 	// 1. Build root Data
 	FHairRootGroupData RootData;
-	bSucceed = BuildHairRootGroupData(In, InGroupIndex, TargetPlatform, RootData);
+	bool bSucceed = BuildHairRootGroupData(In, InGroupIndex, TargetPlatform, RootData);
 
 	// 2. Convert data to bulk data
 	if (bSucceed)
 	{
-		GroomBinding_BulkCopy::BuildRootBulkData(OutPlatformData, RootData);
+		GroomBinding_BulkCopy::BuildGroupPlatformData(OutPlatformData, RootData);
 	}
-#endif
+
 	return bSucceed;
 }
 
 bool FGroomBindingBuilder::BuildBinding(class UGroomBindingAsset* BindingAsset, bool bInitResource)
 {
 #if WITH_EDITORONLY_DATA
-	BindingAsset->CacheDerivedDatas();
+	BindingAsset->BeginCacheDerivedDatas();
 #endif
 	return true;
 }
@@ -2543,7 +2762,7 @@ bool FGroomBindingBuilder::BuildBinding(class UGroomBindingAsset* BindingAsset,
 bool FGroomBindingBuilder::BuildBinding(class UGroomBindingAsset* BindingAsset, uint32 InGroupIndex)
 {
 #if WITH_EDITORONLY_DATA
-	BindingAsset->CacheDerivedDatas();
+	BindingAsset->BeginCacheDerivedDatas();
 #endif
 	return true;
 }
diff --git a/Engine/Plugins/Runtime/HairStrands/Source/HairStrandsCore/Private/GroomBindingCommon.h b/Engine/Plugins/Runtime/HairStrands/Source/HairStrandsCore/Private/GroomBindingCommon.h
--- a/Engine/Plugins/Runtime/HairStrands/Source/HairStrandsCore/Private/GroomBindingCommon.h
+++ b/Engine/Plugins/Runtime/HairStrands/Source/HairStrandsCore/Private/GroomBindingCommon.h
@@ -17,5 +17,8 @@ struct FHairRootGroupData
 	// (Skel.) Mesh Positions for each LODs
 	TArray<TArray<FVector3f>> 			MeshPositions;
 	TArray<TArray<FVector3f>> 			MeshPositions_Transferred;
+
+	// The index of the best LOD on the target mesh that is supported by this data
+	int32 TargetMeshMinLOD = INDEX_NONE;
 };
 bool BuildHairRootGroupData(const FGroomBindingBuilder::FInput& In, uint32 InGroupIndex, const ITargetPlatform* TargetPlatform, FHairRootGroupData& OutData);
\ No newline at end of file
diff --git a/Engine/Plugins/Runtime/HairStrands/Source/HairStrandsCore/Private/GroomBindingCompiler.cpp b/Engine/Plugins/Runtime/HairStrands/Source/HairStrandsCore/Private/GroomBindingCompiler.cpp
--- a/Engine/Plugins/Runtime/HairStrands/Source/HairStrandsCore/Private/GroomBindingCompiler.cpp
+++ b/Engine/Plugins/Runtime/HairStrands/Source/HairStrandsCore/Private/GroomBindingCompiler.cpp
@@ -2,12 +2,9 @@
 
 #include "GroomBindingCompiler.h"
 
-#if WITH_EDITOR
-
 #include "AssetCompilingManager.h"
 #include "GroomBindingAsset.h"
 #include "ObjectCacheContext.h"
-#include "Settings/EditorExperimentalSettings.h"
 #include "GroomComponent.h"
 #include "Misc/QueuedThreadPoolWrapper.h"
 #include "EngineModule.h"
@@ -17,15 +14,18 @@
 #include "Containers/Set.h"
 #include "SkeletalMeshCompiler.h"
 #include "Misc/IQueuedWork.h"
-#include "LevelEditor.h"
-#include "SLevelViewport.h"
 #include "Components/PrimitiveComponent.h"
 #include "ContentStreaming.h"
 #include "EngineUtils.h"
 #include "ProfilingDebugging/CountersTrace.h"
 
+#if WITH_EDITOR
+#include "Settings/EditorExperimentalSettings.h"
+#endif
+
 #define LOCTEXT_NAMESPACE "GroomBindingCompiler"
 
+#if WITH_EDITOR
 static AsyncCompilationHelpers::FAsyncCompilationStandardCVars CVarAsyncGroomBindingAssetStandard(
 	TEXT("GroomBinding"),
 	TEXT("groom bindings"),
@@ -35,6 +35,7 @@ static AsyncCompilationHelpers::FAsyncCompilationStandardCVars CVarAsyncGroomBin
 			FGroomBindingCompilingManager::Get().FinishAllCompilation();
 		}
 	));
+#endif
 
 namespace GroomBindingAssetCompilingManagerImpl
 {
@@ -45,12 +46,13 @@ namespace GroomBindingAssetCompilingManagerImpl
 		if (!bIsInitialized)
 		{
 			bIsInitialized = true;
-			
+#if WITH_EDITOR
 			AsyncCompilationHelpers::EnsureInitializedCVars(
 				TEXT("groombinding"),
 				CVarAsyncGroomBindingAssetStandard.AsyncCompilation,
 				CVarAsyncGroomBindingAssetStandard.AsyncCompilationMaxConcurrency,
 				GET_MEMBER_NAME_CHECKED(UEditorExperimentalSettings, bEnableAsyncGroomBindingCompilation));
+#endif
 		}
 	}
 }
@@ -155,11 +157,15 @@ TArrayView<FName> FGroomBindingCompilingManager::GetDependentTypeNames() const
 {
 	// GroomBindings may have dependencies to SkeletalMesh so we prefer processing them before we get called
 	// so a single FinishAllCompilation is able to finish everything in a single pass.
+#if WITH_EDITOR
 	static FName DependentTypeNames[] = 
 	{
 		FSkinnedAssetCompilingManager::GetStaticAssetTypeName(),
 	};
 	return TArrayView<FName>(DependentTypeNames);
+#else
+	return TArrayView<FName>();
+#endif
 }
 
 int32 FGroomBindingCompilingManager::GetNumRemainingAssets() const
@@ -182,12 +188,14 @@ FQueuedThreadPool* FGroomBindingCompilingManager::GetThreadPool() const
 		// GroomBinding assets will be scheduled on the asset thread pool, where concurrency limits might by dynamically adjusted depending on memory constraints.
 		GGroomBindingAssetThreadPool = new FQueuedThreadPoolDynamicWrapper(FAssetCompilingManager::Get().GetThreadPool(), -1, PriorityMapper);
 
+#if WITH_EDITOR
 		AsyncCompilationHelpers::BindThreadPoolToCVar(
 			GGroomBindingAssetThreadPool,
 			CVarAsyncGroomBindingAssetStandard.AsyncCompilation,
 			CVarAsyncGroomBindingAssetStandard.AsyncCompilationResume,
 			CVarAsyncGroomBindingAssetStandard.AsyncCompilationMaxConcurrency
 		);
+#endif
 	}
 
 	return GGroomBindingAssetThreadPool;
@@ -230,7 +238,11 @@ bool FGroomBindingCompilingManager::IsAsyncCompilationEnabled() const
 		return false;
 	}
 
+#if WITH_EDITOR
 	return CVarAsyncGroomBindingAssetStandard.AsyncCompilation.GetValueOnAnyThread() != 0;
+#else
+	return true;
+#endif
 }
 
 TRACE_DECLARE_INT_COUNTER(QueuedGroomBindingAssetCompilation, TEXT("AsyncCompilation/QueuedGroomBinding"));
@@ -290,6 +302,7 @@ void FGroomBindingCompilingManager::PostCompilation(UGroomBindingAsset* GroomBin
 			}
 		}
 
+#if WITH_EDITOR
 		// Calling this delegate during app exit might be quite dangerous and lead to crash
 		// if the content browser wants to refresh a thumbnail it might try to load a package
 		// which will then fail due to various reasons related to the editor shutting down.
@@ -301,6 +314,7 @@ void FGroomBindingCompilingManager::PostCompilation(UGroomBindingAsset* GroomBin
 			FPropertyChangedEvent EmptyPropertyChangedEvent(nullptr);
 			FCoreUObjectDelegates::OnObjectPropertyChanged.Broadcast(GroomBindingAsset, EmptyPropertyChangedEvent);
 		}
+#endif
 	}
 }
 
@@ -372,7 +386,9 @@ void FGroomBindingCompilingManager::FinishCompilation(TArrayView<UGroomBindingAs
 		}
 	}
 
+#if WITH_EDITOR
 	FSkinnedAssetCompilingManager::Get().FinishCompilation(Dependencies.Array());
+#endif
 	
 	// Now that dependencies have finished, we can launch the compilations for the groom bindings.
 	SchedulePendingCompilations();
@@ -413,8 +429,10 @@ void FGroomBindingCompilingManager::FinishCompilation(TArrayView<UGroomBindingAs
 		AsyncCompilationHelpers::FinishCompilation(
 			[&CompilableGroomBindingAsset](int32 Index) -> AsyncCompilationHelpers::ICompilable& { return CompilableGroomBindingAsset[Index]; },
 			CompilableGroomBindingAsset.Num(),
+#if WITH_EDITOR
 			LOCTEXT("GroomBindings", "Groom Bindings"),
 			LogHairStrands,
+#endif
 			[this](AsyncCompilationHelpers::ICompilable* Object)
 			{
 				UGroomBindingAsset* GroomBindingAsset = static_cast<FCompilableGroomBindingAsset*>(Object)->GroomBindingAsset.Get();
@@ -446,7 +464,7 @@ void FGroomBindingCompilingManager::SchedulePendingCompilations()
 	// Call CacheDerivedDatas again so it's scheduled for real this time.
 	for (UGroomBindingAsset* GroomBindingAsset : ReadyToSchedule)
 	{
-		GroomBindingAsset->CacheDerivedDatas();
+		GroomBindingAsset->BeginCacheDerivedDatas();
 	}
 }
 
@@ -562,6 +580,4 @@ void FGroomBindingCompilingManager::ProcessAsyncTasks(bool bLimitExecutionTime)
 	UpdateCompilationNotification();
 }
 
-#endif // #if WITH_EDITOR
-
 #undef LOCTEXT_NAMESPACE
diff --git a/Engine/Plugins/Runtime/HairStrands/Source/HairStrandsCore/Private/GroomComponent.cpp b/Engine/Plugins/Runtime/HairStrands/Source/HairStrandsCore/Private/GroomComponent.cpp
--- a/Engine/Plugins/Runtime/HairStrands/Source/HairStrandsCore/Private/GroomComponent.cpp
+++ b/Engine/Plugins/Runtime/HairStrands/Source/HairStrandsCore/Private/GroomComponent.cpp
@@ -1571,8 +1571,6 @@ void UGroomComponent::SwitchSimulationLOD(const int32 PreviousLOD, const int32 C
 	}
 }
 
-#if WITH_EDITOR
-
 bool UGroomComponent::IsCompiling() const
 {
 	return BindingAsset && BindingAsset->IsCompiling();
@@ -1580,17 +1578,16 @@ bool UGroomComponent::IsCompiling() const
 
 void UGroomComponent::PostCompilation()
 {
+#if WITH_EDITOR
 	if (GroomAssetBeingLoaded && GroomAssetBeingLoaded->IsValid())
 	{
 		// Re-set the assets on the component now that they are loaded
 		SetGroomAsset(GroomAssetBeingLoaded, BindingAssetBeingLoaded);
 	}
-
+#endif
 	InitIfDependenciesReady();
 }
 
-#endif // WITH_EDITOR
-
 void UGroomComponent::SetGroomAsset(UGroomAsset* Asset)
 {
 	SetGroomAsset(Asset, BindingAsset);
@@ -1977,13 +1974,11 @@ FPrimitiveSceneProxy* UGroomComponent::CreateSceneProxy()
 		return nullptr;
 	}
 
-#if WITH_EDITOR
 	// This will be recreated when the binding asset compilation finishes
 	if (BindingAsset && BindingAsset->IsCompiling())
 	{
 		return nullptr;
 	}
-#endif
 
 	if (CheckPSOPrecachingAndBoostPriority() && GetPSOPrecacheProxyCreationStrategy() == EPSOPrecacheProxyCreationStrategy::DelayUntilPSOPrecached)
 	{
@@ -2863,6 +2858,16 @@ void UGroomComponent::InitResources(bool bIsBindingReloading)
 			HairGroupInstance->HairGroupPublicData->SetLODGeometryTypes(LODGeometryTypes);
 			HairGroupInstance->HairGroupPublicData->bAutoLOD = GroomAsset->GetLODMode() == EGroomLODMode::Auto;
 			HairGroupInstance->HairGroupPublicData->AutoLODBias = GroomAsset->GetAutoLODBias();
+
+			HairGroupInstance->HairGroupPublicData->BindingMinMeshLOD = INDEX_NONE;
+			if (LocalBindingAsset)
+			{
+				const TArray<UGroomBindingAsset::FHairGroupPlatformData>& BindingPlatformData = LocalBindingAsset->GetHairGroupsPlatformData();
+				if (BindingPlatformData.IsValidIndex(GroupIt))
+				{
+					HairGroupInstance->HairGroupPublicData->BindingMinMeshLOD = BindingPlatformData[GroupIt].TargetMeshMinLOD;
+				}
+			}
 		}
 
 		FHairResourceName ResourceName(GetFName(), GroupIt);
diff --git a/Engine/Plugins/Runtime/HairStrands/Source/HairStrandsCore/Private/GroomManager.cpp b/Engine/Plugins/Runtime/HairStrands/Source/HairStrandsCore/Private/GroomManager.cpp
--- a/Engine/Plugins/Runtime/HairStrands/Source/HairStrandsCore/Private/GroomManager.cpp
+++ b/Engine/Plugins/Runtime/HairStrands/Source/HairStrandsCore/Private/GroomManager.cpp
@@ -2644,7 +2644,8 @@ static bool SelectValidLOD(
 		}
 	}
 
-	if (!bIsVisible)
+	const bool bDoesBindingSupportMeshLOD = MeshLODIndex >= Instance->HairGroupPublicData->BindingMinMeshLOD;
+	if (!bIsVisible || !bDoesBindingSupportMeshLOD)
 	{
 		GeometryType = EHairGeometryType::NoneGeometry;
 	}
@@ -2778,7 +2779,7 @@ static bool SelectValidLOD(
 	// (which would required extra data loading) change the geometry to be invalid, so that it don't get processed this frame.
 	const bool bIsLODValid = PrevHairLODIndex >= 0;
 	const bool bHasMeshLODChanged = PrevMeshLODIndex != MeshLODIndex && (Instance->Guides.bHasGlobalInterpolation || Instance->BindingType == EHairBindingType::Skinning);
-	if (!bIsLODDataReady && (!bIsLODValid || bHasMeshLODChanged))
+	if (!bIsLODDataReady && (!bIsLODValid || bHasMeshLODChanged || !bDoesBindingSupportMeshLOD))
 	{
 		return false;
 	}
diff --git a/Engine/Plugins/Runtime/HairStrands/Source/HairStrandsCore/Private/GroomRBFDeformer.cpp b/Engine/Plugins/Runtime/HairStrands/Source/HairStrandsCore/Private/GroomRBFDeformer.cpp
--- a/Engine/Plugins/Runtime/HairStrands/Source/HairStrandsCore/Private/GroomRBFDeformer.cpp
+++ b/Engine/Plugins/Runtime/HairStrands/Source/HairStrandsCore/Private/GroomRBFDeformer.cpp
@@ -630,6 +630,8 @@ void FGroomRBFDeformer::GetRBFDeformedGroomAsset(const UGroomAsset* InGroomAsset
 			Input.BindingType = EGroomBindingMeshType::SkeletalMesh;
 			Input.NumInterpolationPoints = BindingAsset->GetNumInterpolationPoints();
 			Input.MatchingSection = BindingAsset->GetMatchingSection();
+			Input.SourceMeshLOD = 0;
+			Input.TargetMeshMinLOD = 0;
 			Input.bHasValidTarget = BindingAsset->HasValidTarget();
 			Input.GroomAsset = BindingAsset->GetGroom();
 			Input.SourceSkeletalMesh = BindingAsset->GetSourceSkeletalMesh();
diff --git a/Engine/Plugins/Runtime/HairStrands/Source/HairStrandsCore/Private/GroomResources.cpp b/Engine/Plugins/Runtime/HairStrands/Source/HairStrandsCore/Private/GroomResources.cpp
--- a/Engine/Plugins/Runtime/HairStrands/Source/HairStrandsCore/Private/GroomResources.cpp
+++ b/Engine/Plugins/Runtime/HairStrands/Source/HairStrandsCore/Private/GroomResources.cpp
@@ -1724,10 +1724,14 @@ FHairStrandsLODDeformedRootResource::FHairStrandsLODDeformedRootResource(FHairSt
 ,	CurveType(InCurveType)
 {
 	check(InRestResources);
-	check(InRestResources->IsDataValid(InMeshLODIndex));
-	FHairStrandsLODRestRootResource* LODData = InRestResources->GetLOD(InMeshLODIndex);
-	RootCount = LODData->BulkData.Header.RootCount;
-	SampleCount = LODData->BulkData.Header.SampleCount;
+
+	if (InRestResources->IsDataValid(InMeshLODIndex))
+	{
+		FHairStrandsLODRestRootResource* LODData = InRestResources->GetLOD(InMeshLODIndex);
+		RootCount = LODData->BulkData.Header.RootCount;
+		SampleCount = LODData->BulkData.Header.SampleCount;
+	}
+
 	Status = EStatus::Invalid;
 }
 
diff --git a/Engine/Plugins/Runtime/HairStrands/Source/HairStrandsCore/Private/HairStrandsCore.cpp b/Engine/Plugins/Runtime/HairStrands/Source/HairStrandsCore/Private/HairStrandsCore.cpp
--- a/Engine/Plugins/Runtime/HairStrands/Source/HairStrandsCore/Private/HairStrandsCore.cpp
+++ b/Engine/Plugins/Runtime/HairStrands/Source/HairStrandsCore/Private/HairStrandsCore.cpp
@@ -25,9 +25,7 @@ void FHairStrandsCore::StartupModule()
 {
 	RegisterBookmarkFunction(ProcessHairStrandsBookmark);
 
-#if WITH_EDITOR
 	FAssetCompilingManager::Get().RegisterManager(&FGroomBindingCompilingManager::Get());
-#endif
 
 	// Maps virtual shader source directory /Plugin/FX/Niagara to the plugin's actual Shaders directory.
 	FString PluginShaderDir = FPaths::Combine(IPluginManager::Get().FindPlugin(TEXT("HairStrands"))->GetBaseDir(), TEXT("Shaders"));
@@ -37,9 +35,7 @@ void FHairStrandsCore::StartupModule()
 
 void FHairStrandsCore::ShutdownModule()
 {
-#if WITH_EDITOR
 	FAssetCompilingManager::Get().UnregisterManager(&FGroomBindingCompilingManager::Get());
-#endif
 
 	SetGroomEnabled(false);
 }
diff --git a/Engine/Plugins/Runtime/HairStrands/Source/HairStrandsCore/Public/GroomBindingAsset.h b/Engine/Plugins/Runtime/HairStrands/Source/HairStrandsCore/Public/GroomBindingAsset.h
--- a/Engine/Plugins/Runtime/HairStrands/Source/HairStrandsCore/Public/GroomBindingAsset.h
+++ b/Engine/Plugins/Runtime/HairStrands/Source/HairStrandsCore/Public/GroomBindingAsset.h
@@ -8,6 +8,7 @@
 #include "HairDescription.h"
 #include "HairStrandsDatas.h"
 #include "RenderResource.h"
+#include "Rendering/SkeletalMeshLODRenderData.h"
 #include "GroomResources.h"
 #include "HairStrandsInterface.h"
 #include "Async/AsyncWork.h"
@@ -58,14 +59,18 @@ enum class EGroomBindingAsyncProperties : uint64
 	GroomBindingType = 1 << 0,
 	Groom = 1 << 1,
 	SourceSkeletalMesh = 1 << 2,
-	TargetSkeletalMesh = 1 << 3,
-	SourceGeometryCache = 1 << 4,
-	TargetGeometryCache = 1 << 5,
-	NumInterpolationPoints = 1 << 6,
-	MatchingSection = 1 << 7,
-	GroupInfos = 1 << 8,
-	HairGroupResources = 1 << 9,
-	HairGroupPlatformData = 1 << 10,
+	SourceMeshRequestedLOD = 1 << 3,
+	SourceMeshUsedLOD = 1 << 4,
+	TargetSkeletalMesh = 1 << 5,
+	TargetMeshRequestedMinLOD = 1 << 6,
+	TargetMeshUsedMinLOD = 1 << 7,
+	SourceGeometryCache = 1 << 8,
+	TargetGeometryCache = 1 << 9,
+	NumInterpolationPoints = 1 << 10,
+	MatchingSection = 1 << 11,
+	GroupInfos = 1 << 12,
+	HairGroupResources = 1 << 13,
+	HairGroupPlatformData = 1 << 14,
 	All = MAX_uint64
 };
 
@@ -95,6 +100,16 @@ private:
 	static thread_local const UGroomBindingAsset* Asset;
 };
 
+UENUM()
+enum class EGroomBindingAssetBuildResult : uint8
+{
+	Succeeded,
+	Failed
+};
+
+DECLARE_DYNAMIC_DELEGATE_TwoParams(FOnGroomBindingAssetBuildComplete, UGroomBindingAsset*, GroomBinding, EGroomBindingAssetBuildResult, Result);
+DECLARE_DELEGATE_TwoParams(FOnGroomBindingAssetBuildCompleteNative, UGroomBindingAsset*, EGroomBindingAssetBuildResult);
+
 struct FGroomBindingBuildContext
 {
 	FGroomBindingBuildContext() = default;
@@ -105,6 +120,15 @@ struct FGroomBindingBuildContext
 	FGroomBindingBuildContext(FGroomBindingBuildContext&&) = default;
 	FGroomBindingBuildContext& operator=(FGroomBindingBuildContext&&) = default;
 
+	// All mesh LODs needed to build the binding are referenced here to prevent them from being
+	// streamed out during the build.
+	TArray<TRefCountPtr<FSkeletalMeshLODRenderData>> MeshLODReferences;
+
+	FOnGroomBindingAssetBuildComplete DynamicCompletionDelegate;
+	FOnGroomBindingAssetBuildCompleteNative NativeCompletionDelegate;
+
+	int32 SourceMeshLOD = INDEX_NONE;
+	int32 TargetMeshMinLOD = INDEX_NONE;
 	bool bReloadResource = false;
 };
 
@@ -194,6 +218,26 @@ public:
 	UFUNCTION(BlueprintGetter) USkeletalMesh* GetSourceSkeletalMesh() const;
 	UFUNCTION(BlueprintSetter) void SetSourceSkeletalMesh(USkeletalMesh* InSkeletalMesh);
 
+private:
+	UE_DEPRECATED(5.5, "Please do not access this member directly; use UGroomBindingAsset::GetSourceMeshRequestedLOD() or UGroomBindingAsset::SetSourceMeshRequestedLOD().")
+	UPROPERTY(EditAnywhere, BlueprintGetter = GetSourceMeshRequestedLOD, BlueprintSetter = SetSourceMeshRequestedLOD, Category = "BuildSettings")
+	int32 SourceMeshRequestedLOD;
+
+public:
+	static FName GetSourceMeshRequestedLODMemberName();
+	UFUNCTION(BlueprintGetter) int32 GetSourceMeshRequestedLOD() const;
+	UFUNCTION(BlueprintSetter) void SetSourceMeshRequestedLOD(int32 InSourceMeshRequestedLOD);
+
+private:
+	UE_DEPRECATED(5.5, "Please do not access this member directly; use UGroomBindingAsset::GetSourceMeshUsedLOD().")
+	UPROPERTY(VisibleAnywhere, BlueprintGetter = GetSourceMeshUsedLOD, Category = "BuildSettings")
+	int32 SourceMeshUsedLOD = INDEX_NONE;
+
+	UFUNCTION(BlueprintSetter) void SetSourceMeshUsedLOD(int32 InSourceMeshUsedLOD);
+public:
+	static FName GetSourceMeshUsedLODMemberName();
+	UFUNCTION(BlueprintGetter) int32 GetSourceMeshUsedLOD() const;
+
 private:
 	/** Skeletal mesh on which the groom is attached to. */
 	UE_DEPRECATED(5.3, "Please do not access this member directly; use UGroomBindingAsset::GetTargetSkeletalMesh() or UGroomBindingAsset::SetTargetSkeletalMesh().")
@@ -204,6 +248,26 @@ public:
 	static FName GetTargetSkeletalMeshMemberName();
 	UFUNCTION(BlueprintGetter) USkeletalMesh* GetTargetSkeletalMesh() const;
 	UFUNCTION(BlueprintSetter) void SetTargetSkeletalMesh(USkeletalMesh* InSkeletalMesh);
+	
+private:
+	UE_DEPRECATED(5.3, "Please do not access this member directly; use UGroomBindingAsset::GetTargetMeshRequestedMinLOD() or UGroomBindingAsset::SetTargetMeshRequestedMinLOD().")
+	UPROPERTY(EditAnywhere, BlueprintGetter = GetTargetMeshRequestedMinLOD, BlueprintSetter = SetTargetMeshRequestedMinLOD, Category = "BuildSettings")
+	int32 TargetMeshRequestedMinLOD;
+
+public:
+	static FName GetTargetMeshRequestedMinLODMemberName();
+	UFUNCTION(BlueprintGetter) int32 GetTargetMeshRequestedMinLOD() const;
+	UFUNCTION(BlueprintSetter) void SetTargetMeshRequestedMinLOD(int32 InTargetMeshRequestedMinLOD);
+
+private:
+	UE_DEPRECATED(5.5, "Please do not access this member directly; use UGroomBindingAsset::GetTargetMeshUsedMinLOD().")
+	UPROPERTY(VisibleAnywhere, BlueprintGetter = GetTargetMeshUsedMinLOD, Category = "BuildSettings")
+	int32 TargetMeshUsedMinLOD = INDEX_NONE;
+
+	UFUNCTION(BlueprintSetter) void SetTargetMeshUsedMinLOD(int32 InTargetMeshUsedMinLOD);
+public:
+	static FName GetTargetMeshUsedMinLODMemberName();
+	UFUNCTION(BlueprintGetter) int32 GetTargetMeshUsedMinLOD() const;
 
 private:
 	UE_DEPRECATED(5.3, "Please do not access this member directly; use UGroomBindingAsset::GetSourceGeometryCache() or UGroomBindingAsset::SetSourceGeometryCache().")
@@ -283,6 +347,9 @@ public:
 		TArray<FHairStrandsRootBulkData>		 SimRootBulkDatas;
 		TArray<FHairStrandsRootBulkData>		 RenRootBulkDatas;
 		TArray<TArray<FHairStrandsRootBulkData>> CardsRootBulkDatas;
+
+		// The minimum mesh LOD that this binding data can support
+		int32 TargetMeshMinLOD = -1;
 	};
 
 private:
@@ -361,7 +428,6 @@ public:
 	/** Return true if the binding asset is valid, i.e., correctly built and loaded. */
 	bool IsValid() const { return bIsValid;  }
 
-#if WITH_EDITOR
 private:
 	/** Used as a bit-field indicating which properties are read by async compilation. */
 	std::atomic<uint64> AccessedProperties;
@@ -390,6 +456,8 @@ private:
 
 	void ExecuteCacheDerivedDatas(FGroomBindingBuildContext& Context);
 	void FinishCacheDerivedDatas(FGroomBindingBuildContext& Context);
+	int32 GetClampedSourceMeshLOD(const ITargetPlatform* TargetPlatform) const;
+	int32 GetClampedTargetMeshMinLOD(const ITargetPlatform* TargetPlatform) const;
 
 public:
 	/** IInterface_AsyncCompilation begin*/
@@ -398,7 +466,7 @@ public:
 		return AsyncTask != nullptr || AccessedProperties.load(std::memory_order_relaxed) != 0;
 	}
 	/** IInterface_AsyncCompilation end*/
-
+#if WITH_EDITOR
 	FOnGroomBindingAssetChanged OnGroomBindingAssetChanged;
 
 	void RecreateResources();
@@ -409,7 +477,6 @@ private:
 	void WaitUntilAsyncPropertyReleased(EGroomBindingAsyncProperties AsyncProperties, EGroomBindingAsyncPropertyLockType LockType) const;
 	void AcquireAsyncProperty(uint64 AsyncProperties = MAX_uint64, EGroomBindingAsyncPropertyLockType LockType = EGroomBindingAsyncPropertyLockType::ReadWrite)
 	{
-#if WITH_EDITOR
 		if ((LockType & EGroomBindingAsyncPropertyLockType::ReadOnly) == EGroomBindingAsyncPropertyLockType::ReadOnly)
 		{
 			AccessedProperties |= AsyncProperties;
@@ -419,12 +486,10 @@ private:
 		{
 			ModifiedProperties |= AsyncProperties;
 		}
-#endif
 	}
 
 	void ReleaseAsyncProperty(uint64 AsyncProperties = MAX_uint64, EGroomBindingAsyncPropertyLockType LockType = EGroomBindingAsyncPropertyLockType::ReadWrite)
 	{
-#if WITH_EDITOR
 		if ((LockType & EGroomBindingAsyncPropertyLockType::ReadOnly) == EGroomBindingAsyncPropertyLockType::ReadOnly)
 		{
 			AccessedProperties &= ~AsyncProperties;
@@ -434,19 +499,37 @@ private:
 		{
 			ModifiedProperties &= ~AsyncProperties;
 		}
-#endif
 	}
 
 	static void FlushRenderingCommandIfUsed(const UGroomBindingAsset* In);
 public:
-#if WITH_EDITORONLY_DATA
-	/** Build/rebuild a binding asset */
-	void Build();
+	/**
+	 * Build/rebuild a binding asset.
+	 * 
+	 * Avoid calling Build again while a build is already in progress on this asset, as this will
+	 * cause the Game Thread to block until the earlier build is finished.
+	 * 
+	 * Similarly, the Game Thread will block for a short time if a Groom Component is already using 
+	 * this binding asset, so if performance is critical avoid calling Build on bindings that are 
+	 * in use.
+	 */
+	UFUNCTION(BlueprintCallable, Category = "Groom")
+	void Build(FOnGroomBindingAssetBuildComplete CompletionDelegate);
+	void Build(const FOnGroomBindingAssetBuildCompleteNative& CompletionDelegate = FOnGroomBindingAssetBuildCompleteNative());
+	void Build(const FOnGroomBindingAssetBuildComplete& DynamicCompletionDelegate, const FOnGroomBindingAssetBuildCompleteNative& NativeCompletionDelegate);
 
-	void CacheDerivedDatas();
+	// Internal use only
+	void BeginCacheDerivedDatas(
+		const FOnGroomBindingAssetBuildComplete& DynamicCompletionDelegate = FOnGroomBindingAssetBuildComplete(),
+		const FOnGroomBindingAssetBuildCompleteNative& NativeCompletionDelegate = FOnGroomBindingAssetBuildCompleteNative());
 
 	bool HasAnyDependenciesCompiling() const;
 
+private:
+	bool TryInitializeContextForMeshes(FGroomBindingBuildContext& OutContext) const;
+
+public:
+#if WITH_EDITORONLY_DATA
 	virtual void BeginCacheForCookedPlatformData(const ITargetPlatform* TargetPlatform);
 	virtual void ClearAllCachedCookedPlatformData();
 	TArray<FHairGroupPlatformData>* GetCachedCookedPlatformData(const ITargetPlatform* TargetPlatform);
@@ -459,7 +542,7 @@ public:
 		TArray<FString> GroupDerivedDataKeys;
 		TArray<FHairGroupPlatformData> GroupPlatformDatas;
 	};
-private:
+
 	TArray<FCachedCookedPlatformData*> CachedCookedPlatformDatas;
 
 	void RegisterGroomDelegates();
diff --git a/Engine/Plugins/Runtime/HairStrands/Source/HairStrandsCore/Public/GroomBindingBuilder.h b/Engine/Plugins/Runtime/HairStrands/Source/HairStrandsCore/Public/GroomBindingBuilder.h
--- a/Engine/Plugins/Runtime/HairStrands/Source/HairStrandsCore/Public/GroomBindingBuilder.h
+++ b/Engine/Plugins/Runtime/HairStrands/Source/HairStrandsCore/Public/GroomBindingBuilder.h
@@ -12,34 +12,51 @@
 #include "GroomBindingAsset.h"
 
 class ITargetPlatform;
+class UGroomBindingAsset;
 
 struct HAIRSTRANDSCORE_API FGroomBindingBuilder
 {
 	struct FInput
 	{
+		FInput() = default;
+		FInput(UGroomBindingAsset* BindingAsset, const ITargetPlatform* TargetPlatform, int32 InSourceMeshLOD, int32 InTargetMeshMinLOD);
+
 		EGroomBindingMeshType BindingType = EGroomBindingMeshType::SkeletalMesh;
 		int32 NumInterpolationPoints = 0;
 		int32 MatchingSection = 0;
+
+		// These must be initialized to valid LODs.
+		// 
+		// When binding to a geometry cache, they must be set to 0.
+		int32 SourceMeshLOD = INDEX_NONE;
+		int32 TargetMeshMinLOD = INDEX_NONE;
+
 		bool bHasValidTarget = false;
 		UGroomAsset* GroomAsset = nullptr;
 		USkeletalMesh* SourceSkeletalMesh = nullptr;
 		USkeletalMesh* TargetSkeletalMesh = nullptr;
-		UGeometryCache* SourceGeometryCache;
-		UGeometryCache* TargetGeometryCache;
+		UGeometryCache* SourceGeometryCache = nullptr;
+		UGeometryCache* TargetGeometryCache = nullptr;
 	};
 
 	static FString GetVersion();
 
 	// Build binding asset data
 	UE_DEPRECATED(5.4, "Please do not access this funciton; but rather call BindingAsset->CacheDerivedDatas()")
-	static bool BuildBinding(class UGroomBindingAsset* BindingAsset, bool bInitResource);
+	static bool BuildBinding(UGroomBindingAsset* BindingAsset, bool bInitResource);
 
 	// Build binding asset data
 	UE_DEPRECATED(5.4, "Please do not access this funciton; but rather call BindingAsset->CacheDerivedDatas()")
-	static bool BuildBinding(class UGroomBindingAsset* BindingAsset, uint32 InGroupIndex);
+	static bool BuildBinding(UGroomBindingAsset* BindingAsset, uint32 InGroupIndex);
 
 	// Build binding asset data
-	static bool BuildBinding(const FInput& In, uint32 InGroupIndex, const ITargetPlatform* TargetPlatform, UGroomBindingAsset::FHairGroupPlatformData& Out);
+	//
+	// The caller must ensure that the referenced mesh LODs stay loaded until this function returns.
+	static bool BuildBinding(
+		const FInput& In,
+		uint32 InGroupIndex,
+		const ITargetPlatform* TargetPlatform,
+		UGroomBindingAsset::FHairGroupPlatformData& Out);
 
 	// Extract root data from bulk data
 	static void GetRootData(FHairStrandsRootData& Out, const FHairStrandsRootBulkData& In);
diff --git a/Engine/Plugins/Runtime/HairStrands/Source/HairStrandsCore/Public/GroomBindingCompiler.h b/Engine/Plugins/Runtime/HairStrands/Source/HairStrandsCore/Public/GroomBindingCompiler.h
--- a/Engine/Plugins/Runtime/HairStrands/Source/HairStrandsCore/Public/GroomBindingCompiler.h
+++ b/Engine/Plugins/Runtime/HairStrands/Source/HairStrandsCore/Public/GroomBindingCompiler.h
@@ -9,8 +9,6 @@
 #include "IAssetCompilingManager.h"
 #include "AsyncCompilationHelpers.h"
 
-#if WITH_EDITOR
-
 class UGroomBindingAsset;
 class UPrimitiveComponent;
 class FQueuedThreadPool;
@@ -104,5 +102,3 @@ private:
 	void OnPostReachabilityAnalysis();
 	FDelegateHandle PostReachabilityAnalysisHandle;
 };
-
-#endif // #if WITH_EDITOR
\ No newline at end of file
diff --git a/Engine/Plugins/Runtime/HairStrands/Source/HairStrandsCore/Public/GroomComponent.h b/Engine/Plugins/Runtime/HairStrands/Source/HairStrandsCore/Public/GroomComponent.h
--- a/Engine/Plugins/Runtime/HairStrands/Source/HairStrandsCore/Public/GroomComponent.h
+++ b/Engine/Plugins/Runtime/HairStrands/Source/HairStrandsCore/Public/GroomComponent.h
@@ -174,19 +174,18 @@ public:
 	void ValidateMaterials(bool bMapCheck) const;
 	void Invalidate();
 	void InvalidateAndRecreate();
-	void PostCompilation();
-	
-	//~ Begin IInterface_AsyncCompilation Interface.
-	virtual bool IsCompiling() const override;
-	//~ End IInterface_AsyncCompilation Interface.
-#endif
 
-#if WITH_EDITOR
 	virtual void PreFeatureLevelChange(ERHIFeatureLevel::Type PendingFeatureLevel) override;
 	void HandlePlatformPreviewChanged(ERHIFeatureLevel::Type InFeatureLevel);
 	void HandleFeatureLevelChanged(ERHIFeatureLevel::Type InFeatureLevel);
 #endif
 
+	void PostCompilation();
+	
+	//~ Begin IInterface_AsyncCompilation Interface.
+	virtual bool IsCompiling() const override;
+	//~ End IInterface_AsyncCompilation Interface.
+
 	/* Accessor function for changing Groom asset from blueprint/sequencer */
 	UFUNCTION(BlueprintCallable, Category = "Groom")
 	void SetGroomAsset(UGroomAsset* Asset);
diff --git a/Engine/Plugins/Runtime/HairStrands/Source/HairStrandsCore/Public/GroomResources.h b/Engine/Plugins/Runtime/HairStrands/Source/HairStrandsCore/Public/GroomResources.h
--- a/Engine/Plugins/Runtime/HairStrands/Source/HairStrandsCore/Public/GroomResources.h
+++ b/Engine/Plugins/Runtime/HairStrands/Source/HairStrandsCore/Public/GroomResources.h
@@ -474,7 +474,23 @@ struct FHairStrandsRestRootResource
 	const FHairStrandsLODRestRootResource* GetLOD(int32 InMeshLODIndex) const { check(LODs.IsValidIndex(InMeshLODIndex)); return LODs[InMeshLODIndex]; }
 
 	uint32 GetLODCount() const { return LODs.Num(); }
-	uint32 GetRootCount()const { return LODs.Num() > 0 && LODs[0] != nullptr? LODs[0]->GetRootCount() : 0u; }
+	uint32 GetRootCount() const
+	{ 
+		// Find the first valid root count and return that
+		for (const FHairStrandsLODRestRootResource* LOD : LODs)
+		{
+			if (LOD)
+			{
+				const uint32 RootCount = LOD->GetRootCount();
+				if (RootCount > 0)
+				{
+					return RootCount;
+				}
+			}
+		}
+
+		return 0u;
+	}
 
 	/* Return the memory size for GPU resources */
 	uint32 GetResourcesSize() const 
diff --git a/Engine/Plugins/Runtime/HairStrands/Source/HairStrandsEditor/Private/GroomBindingActions.cpp b/Engine/Plugins/Runtime/HairStrands/Source/HairStrandsEditor/Private/GroomBindingActions.cpp
--- a/Engine/Plugins/Runtime/HairStrands/Source/HairStrandsEditor/Private/GroomBindingActions.cpp
+++ b/Engine/Plugins/Runtime/HairStrands/Source/HairStrandsEditor/Private/GroomBindingActions.cpp
@@ -47,7 +47,7 @@ void ExecuteRebuildBindingAsset(const FToolMenuContext& InContext)
 						BindingAsset->GetSourceGeometryCache()->ConditionalPostLoad();
 					}
 				}
-				BindingAsset->CacheDerivedDatas();
+				BindingAsset->BeginCacheDerivedDatas();
 				BindingAsset->GetOutermost()->MarkPackageDirty();
 			}
 		}
diff --git a/Engine/Source/Runtime/Engine/Classes/Interfaces/Interface_AsyncCompilation.h b/Engine/Source/Runtime/Engine/Classes/Interfaces/Interface_AsyncCompilation.h
--- a/Engine/Source/Runtime/Engine/Classes/Interfaces/Interface_AsyncCompilation.h
+++ b/Engine/Source/Runtime/Engine/Classes/Interfaces/Interface_AsyncCompilation.h
@@ -24,6 +24,9 @@ class IInterface_AsyncCompilation
 #if WITH_EDITOR
 	/**	Returns whether or not the asset is currently being compiled */
 	virtual bool IsCompiling() const = 0;
+#else
+	/**	Returns whether or not the asset is currently being compiled */
+	virtual bool IsCompiling() const { return false; }
 #endif
 };
 
diff --git a/Engine/Source/Runtime/Engine/Private/AsyncCompilationHelpers.cpp b/Engine/Source/Runtime/Engine/Private/AsyncCompilationHelpers.cpp
--- a/Engine/Source/Runtime/Engine/Private/AsyncCompilationHelpers.cpp
+++ b/Engine/Source/Runtime/Engine/Private/AsyncCompilationHelpers.cpp
@@ -3,22 +3,21 @@
 #include "AsyncCompilationHelpers.h"
 #include "UObject/UnrealType.h"
 
-#if WITH_EDITOR
-
 #include "AssetCompilingManager.h"
 #include "Misc/DelayedAutoRegister.h"
 #include "Misc/ScopedSlowTask.h"
 #include "Misc/QueuedThreadPool.h"
 #include "ObjectCacheContext.h"
 #include "Misc/CommandLine.h"
-#include "Settings/EditorExperimentalSettings.h"
 #include "Misc/QueuedThreadPoolWrapper.h"
 #include "HAL/PlatformStackWalk.h"
 #include "ShaderCompiler.h"
 #include "ProfilingDebugging/StallDetector.h"
 #include "UObject/UObjectThreadContext.h"
 
-#endif // #if WITH_EDITOR
+#if WITH_EDITOR
+#include "Settings/EditorExperimentalSettings.h"
+#endif
 
 DEFINE_LOG_CATEGORY(LogAsyncCompilation);
 
@@ -84,13 +83,17 @@ static FAutoConsoleCommand CVarAsyncAssetDumpStallStacks(
 	})
 );
 
+#endif // WITH_EDITOR
+
 namespace AsyncCompilationHelpers
 {
 	void FinishCompilation(
 		TFunctionRef<ICompilable& (int32 Index)> Getter,
 		int32 Num,
+#if WITH_EDITOR
 		const FText& AssetType,
 		const FLogCategoryBase& LogCategory,
+#endif
 		TFunctionRef<void(ICompilable*)> PostCompileSingle
 	)
 	{
@@ -104,6 +107,7 @@ namespace AsyncCompilationHelpers
 		FObjectCacheContextScope ObjectCacheScope;
 
 		uint64 StartTime = FPlatformTime::Cycles64();
+#if WITH_EDITOR
 		TOptional<FScopedSlowTask> SlowTask;
 
 		// Do not create a progress during PostLoad or garbage collect as ticking slate might
@@ -117,6 +121,17 @@ namespace AsyncCompilationHelpers
 			SlowTask->MakeDialogDelayed(1.0f, false /*bShowCancelButton*/, false /*bAllowInPIE*/);
 		}
 
+		auto FormatProgress =
+			[&AssetType](int32 Done, int32 Total, FName ObjectName)
+			{
+				FFormatNamedArguments Args;
+				Args.Add(TEXT("AssetType"), AssetType.ToLower());
+				Args.Add(TEXT("Done"), Done);
+				Args.Add(TEXT("Total"), Total);
+				Args.Add(TEXT("ObjectName"), FText::FromName(ObjectName));
+				return FText::Format(LOCTEXT("WaitingOnFinishCompilationWithCount", "Waiting for {AssetType} to be ready {Done}/{Total} ({ObjectName}) ..."), Args);
+			};
+#endif // WITH_EDITOR
 		// Reschedule everything to be executed at blocking priority, it bypasses the pause mechanism to ensure forward progress since we're waiting.
 		// Important to keep using the asset compiling manager's thread pool to benefit from it's memory management and avoid OOM.
 		FQueuedThreadPool* AssetThreadPool = FAssetCompilingManager::Get().GetThreadPool();
@@ -125,17 +140,6 @@ namespace AsyncCompilationHelpers
 			Getter(Index).Reschedule(AssetThreadPool, EQueuedWorkPriority::Blocking);
 		}
 
-		auto FormatProgress =
-			[&AssetType](int32 Done, int32 Total, FName ObjectName)
-		{
-			FFormatNamedArguments Args;
-			Args.Add(TEXT("AssetType"), AssetType.ToLower());
-			Args.Add(TEXT("Done"), Done);
-			Args.Add(TEXT("Total"), Total);
-			Args.Add(TEXT("ObjectName"), FText::FromName(ObjectName));
-			return FText::Format(LOCTEXT("WaitingOnFinishCompilationWithCount", "Waiting for {AssetType} to be ready {Done}/{Total} ({ObjectName}) ..."), Args);
-		};
-
 		const double MaxProcessingTimeSeconds = 0.016f;
 		double StartTimeSeconds = FPlatformTime::Seconds();
 		int32 NumDone = 0;
@@ -164,12 +168,13 @@ namespace AsyncCompilationHelpers
 						PostCompileSingle(&Job);
 						NumDone++;
 
+#if WITH_EDITOR
 						if (SlowTask.IsSet())
 						{
 							FText Progress = FormatProgress(NumDone, Num, Job.GetName());
 							SlowTask->EnterProgressFrame(1.0f, Progress);
 						}
-
+#endif // WITH_EDITOR
 						continue;
 					}
 				}
@@ -190,6 +195,7 @@ namespace AsyncCompilationHelpers
 			float ElapsedTimeSeconds = float(FPlatformTime::Seconds() - StartTimeSeconds);
 			if (!IncompleteJob.WaitCompletionWithTimeout(MaxProcessingTimeSeconds - ElapsedTimeSeconds))
 			{
+#if WITH_EDITOR
 				// Progress the slow task so the editor may remain responsive
 				if (SlowTask.IsSet())
 				{
@@ -214,14 +220,17 @@ namespace AsyncCompilationHelpers
 					const bool bBlockOnGlobalShaderCompletion = false;
 					GShaderCompilingManager->ProcessAsyncResults(bLimitExecutionTime, bBlockOnGlobalShaderCompletion);
 				}
-
+#endif // WITH_EDITOR
 				StartTimeSeconds = FPlatformTime::Seconds();
 			}
 		}
 
+#if WITH_EDITOR
 		SaveStallStack(FPlatformTime::Cycles64() - StartTime);
+#endif
 	}
 
+#if WITH_EDITOR
 	void EnsureInitializedCVars(
 		const TCHAR* InName,
 		TAutoConsoleVariable<int32>& InCVarAsyncCompilation,
@@ -446,8 +455,7 @@ namespace AsyncCompilationHelpers
 			UE_LOG(LogAsyncCompilation, Display, TEXT("Async Compilation Stall Stack: (Total Count: %llu, Total Time: %s)"), TotalCount, *FPlatformTime::PrettyTime(FPlatformTime::ToSeconds64(TotalCycles)));
 		}
 	}
+#endif // WITH_EDITOR
 }
 
-#endif // #if WITH_EDITOR
-
 #undef LOCTEXT_NAMESPACE
diff --git a/Engine/Source/Runtime/Engine/Public/AsyncCompilationHelpers.h b/Engine/Source/Runtime/Engine/Public/AsyncCompilationHelpers.h
--- a/Engine/Source/Runtime/Engine/Public/AsyncCompilationHelpers.h
+++ b/Engine/Source/Runtime/Engine/Public/AsyncCompilationHelpers.h
@@ -26,23 +26,10 @@ private:
 	FTextFormat AssetNameFormat;
 };
 
-#if WITH_EDITOR
-
 class FQueuedThreadPoolWrapper;
 
 namespace AsyncCompilationHelpers
 {
-	class FAsyncCompilationStandardCVars
-	{
-	public:
-		TAutoConsoleVariable<int32> AsyncCompilation;
-		TAutoConsoleVariable<int32> AsyncCompilationMaxConcurrency;
-		FAutoConsoleCommand         AsyncCompilationFinishAll;
-		TAutoConsoleVariable<int32> AsyncCompilationResume;
-
-		ENGINE_API FAsyncCompilationStandardCVars(const TCHAR* AssetType, const TCHAR* AssetTypePluralLowerCase, const FConsoleCommandDelegate& FinishAllCommand);
-	};
-
 	class ICompilable
 	{
 	public:
@@ -88,11 +75,26 @@ namespace AsyncCompilationHelpers
 	ENGINE_API void FinishCompilation(
 		TFunctionRef<ICompilable& (int32 Index)> Getter,
 		int32 Num,
+#if WITH_EDITOR
 		const FText& AssetType,
 		const FLogCategoryBase& LogCategory,
+#endif
 		TFunctionRef<void(ICompilable*)> PostCompileSingle
 	);
 
+#if WITH_EDITOR
+
+	class FAsyncCompilationStandardCVars
+	{
+	public:
+		TAutoConsoleVariable<int32> AsyncCompilation;
+		TAutoConsoleVariable<int32> AsyncCompilationMaxConcurrency;
+		FAutoConsoleCommand         AsyncCompilationFinishAll;
+		TAutoConsoleVariable<int32> AsyncCompilationResume;
+
+		ENGINE_API FAsyncCompilationStandardCVars(const TCHAR* AssetType, const TCHAR* AssetTypePluralLowerCase, const FConsoleCommandDelegate& FinishAllCommand);
+	};
+
 	ENGINE_API void EnsureInitializedCVars(
 		const TCHAR* InName,
 		TAutoConsoleVariable<int32>& InCVarAsyncCompilation,
@@ -108,6 +110,6 @@ namespace AsyncCompilationHelpers
 
 	ENGINE_API void SaveStallStack(uint64 Cycles);
 	ENGINE_API void DumpStallStacks();
-}
 
-#endif // #if WITH_EDITOR
+#endif // WITH_EDITOR
+}
diff --git a/Engine/Source/Runtime/Renderer/Public/HairStrandsInterface.h b/Engine/Source/Runtime/Renderer/Public/HairStrandsInterface.h
--- a/Engine/Source/Runtime/Renderer/Public/HairStrandsInterface.h
+++ b/Engine/Source/Runtime/Renderer/Public/HairStrandsInterface.h
@@ -387,6 +387,9 @@ public:
 
 	TArray<EHairBindingType> BindingTypes;
 
+	// The minimum mesh LOD that this data can support
+	int32 BindingMinMeshLOD = -1;
+
 	// Data change every frame by the groom proxy based on views data
 	float MeshLODIndex = -1;
 	float LODIndex = -1;		// Current LOD used for all views
