Fix cooking issue with groom binding on certain platforms.

The skel. mesh DDC is not stable, and can mutate. This causes the indexing of the groom binding's cached cooked data to be incorrect. Ultimately, the cooking will fails due to missing cooked data for groom binding.
This CL forces the skel. mesh data to be cached when calling UGroomBindingAsset::BeginCacheForCookedPlatformData() prior to building the binding DDC key. This way, the Skel. mesh DDC key will be up to date and groom binding caching will be correctly indexed, and could be retrieved later when serializing the data. This change only impacts cooking (no impact for editor/runtime part).

diff --git a/Engine/Plugins/Runtime/HairStrands/Source/HairStrandsCore/Private/GroomBindingAsset.cpp b/Engine/Plugins/Runtime/HairStrands/Source/HairStrandsCore/Private/GroomBindingAsset.cpp
--- a/Engine/Plugins/Runtime/HairStrands/Source/HairStrandsCore/Private/GroomBindingAsset.cpp
+++ b/Engine/Plugins/Runtime/HairStrands/Source/HairStrandsCore/Private/GroomBindingAsset.cpp
@@ -29,6 +29,7 @@
 #include "DerivedDataCacheInterface.h"
 #include "Serialization/LargeMemoryReader.h"
 #include "Serialization/LargeMemoryWriter.h"
+#include "Engine/SkinnedAssetAsyncCompileUtils.h"
 #endif
 
 LLM_DECLARE_TAG(Groom);
@@ -958,6 +959,54 @@ void UpdateGroomBindingAssetInfos(UGroomBindingAsset* In)
 // * Release : Engine\Source\Runtime\Core\Private\UObject\UE5ReleaseStreamObjectVersion.cpp
 // * ...
 
+static void ForceSkeletalMeshDataCachingForCookingOnly(const ITargetPlatform* TargetPlatform, USkeletalMesh* InSourceSkeletalMesh, USkeletalMesh* InTargetSkeletalMesh)
+{
+	// If a skeletal mesh build or preedit is called while we're async compiling groom bindings, the groom binding compiler
+	// will take care of finishing any groom binding that depends on the skeletal mesh being modified. So this is 
+	// safe to do asynchronously without locks on the render data. We assume that the skeletal mesh's render data is
+	// immutable once it has been built and can only be rebuilt throught a call to PreEditChange first.
+	const bool bAsyncCompiling = !IsInGameThread();
+
+	// * Only for SkeletalMesh: Take scoped lock on the skeletal render mesh data during the entire groom binding building
+	// * Then use an async build scope to allow accessing skeletal mesh property safely.
+	//   If skel.meshes are nullptr, this will act as a NOP
+	InSourceSkeletalMesh = InSourceSkeletalMesh == InTargetSkeletalMesh ? nullptr : InSourceSkeletalMesh;
+	PRAGMA_DISABLE_DEPRECATION_WARNINGS
+	FScopedSkeletalMeshRenderData SourceSkeletalMeshScopedData(bAsyncCompiling ? nullptr : InSourceSkeletalMesh);
+	FScopedSkeletalMeshRenderData TargetSkeletalMeshScopedData(bAsyncCompiling ? nullptr : InTargetSkeletalMesh);
+	PRAGMA_ENABLE_DEPRECATION_WARNINGS
+
+	if (InSourceSkeletalMesh)
+	{
+		if (bAsyncCompiling)
+		{
+			const ITargetPlatform* RunningPlatform = GetTargetPlatformManagerRef().GetRunningTargetPlatform();
+			checkf(TargetPlatform == RunningPlatform, TEXT("It is only safe to query the running platform's render data asynchronously from the skeletal mesh"));
+			InSourceSkeletalMesh->GetResourceForRendering();
+		}
+		else
+		{
+			FSkinnedAssetAsyncBuildScope AsyncBuildScope(InSourceSkeletalMesh);
+			USkeletalMesh::GetPlatformSkeletalMeshRenderData(TargetPlatform, SourceSkeletalMeshScopedData);
+		}
+	}
+
+	if (InTargetSkeletalMesh)
+	{
+		if (bAsyncCompiling)
+		{
+			const ITargetPlatform* RunningPlatform = GetTargetPlatformManagerRef().GetRunningTargetPlatform();
+			checkf(TargetPlatform == RunningPlatform, TEXT("It is only safe to query the running platform's render data asynchronously from the skeletal mesh"));
+			InTargetSkeletalMesh->GetResourceForRendering();
+		}
+		else
+		{
+			FSkinnedAssetAsyncBuildScope AsyncBuildScope(InTargetSkeletalMesh);
+			USkeletalMesh::GetPlatformSkeletalMeshRenderData(TargetPlatform, TargetSkeletalMeshScopedData);
+		}
+	}
+}
+
 namespace GroomBindingDerivedDataCacheUtils
 {
 	const FString& GetGroomBindingDerivedDataVersion()
@@ -972,7 +1021,7 @@ namespace GroomBindingDerivedDataCacheUtils
 	}
 }
 
-static FString BuildDerivedDataKeySuffix(const UGroomBindingAsset& BindingAsset, const ITargetPlatform* TargetPlatform, int32 SourceMeshLOD, int32 TargetMeshMinLOD)
+static FString BuildDerivedDataKeySuffix(const UGroomBindingAsset& BindingAsset, const ITargetPlatform* TargetPlatform, int32 SourceMeshLOD, int32 TargetMeshMinLOD, bool bForceSkelMeshDataCachingForCookingOnly=false)
 {
 	FString BindingType;
 	FString SourceKey;
@@ -982,6 +1031,12 @@ static FString BuildDerivedDataKeySuffix(const UGroomBindingAsset& BindingAsset,
 
 	if (BindingAsset.GetGroomBindingType() == EGroomBindingMeshType::SkeletalMesh)
 	{
+		// Force skel. mesh data caching, in order to produce later an valid/up-to-date DDC key
+		if (bForceSkelMeshDataCachingForCookingOnly)
+		{
+			ForceSkeletalMeshDataCachingForCookingOnly(TargetPlatform, BindingAsset.GetSourceSkeletalMesh(), BindingAsset.GetTargetSkeletalMesh());
+		}
+
 		// Binding type is implicitly SkeletalMesh so keep BindingType empty to prevent triggering rebuild of old binding for nothing
 		if (USkeletalMesh* SourceMesh = BindingAsset.GetSourceSkeletalMesh())
 		{
@@ -1020,7 +1075,7 @@ static FString BuildDerivedDataKeyGroup(const FString& InDeriveDataKeySuffix , u
 	return GroomBindingDerivedDataCacheUtils::BuildGroomBindingDerivedDataKey(InDeriveDataKeySuffix + FString(TEXT("_Group")) + FString::FromInt(InGroupIndex));
 }
 
-static TArray<FString> GetGroupDerivedDataKeys(const UGroomBindingAsset* In, const ITargetPlatform* TargetPlatform, int32 SourceMeshLOD, int32 TargetMeshMinLOD);
+static TArray<FString> GetGroupDerivedDataKeys(const UGroomBindingAsset* In, const ITargetPlatform* TargetPlatform, int32 SourceMeshLOD, int32 TargetMeshMinLOD, bool bForceSkelMeshDataCaching=false);
 static void CacheDerivedDatas(
 	UGroomBindingAsset* In,
 	const uint32 InGroupIndex,
@@ -1519,12 +1574,12 @@ int32 UGroomBindingAsset::GetClampedTargetMeshMinLOD(const ITargetPlatform* Targ
 
 #if WITH_EDITORONLY_DATA
 
-static TArray<FString> GetGroupDerivedDataKeys(const UGroomBindingAsset* In, const ITargetPlatform* TargetPlatform, int32 SourceMeshLOD, int32 TargetMeshMinLOD)
+static TArray<FString> GetGroupDerivedDataKeys(const UGroomBindingAsset* In, const ITargetPlatform* TargetPlatform, int32 SourceMeshLOD, int32 TargetMeshMinLOD, bool bForceSkelMeshDataCachingForCookingOnly)
 {
 	check(In);
 	check(TargetPlatform);
 
-	const FString KeySuffix = BuildDerivedDataKeySuffix(*In, TargetPlatform, SourceMeshLOD, TargetMeshMinLOD);
+	const FString KeySuffix = BuildDerivedDataKeySuffix(*In, TargetPlatform, SourceMeshLOD, TargetMeshMinLOD, bForceSkelMeshDataCachingForCookingOnly );
 	const uint32 GroupCount = In->GetGroupInfos().Num();
 
 	TArray<FString> Out;
@@ -1561,7 +1616,7 @@ void UGroomBindingAsset::BeginCacheForCookedPlatformData(const ITargetPlatform*
 	// 1. Build the key for each group
 	const int32 SourceMeshLOD = GetClampedSourceMeshLOD(TargetPlatform);
 	const int32 TargetMeshMinLOD = GetClampedTargetMeshMinLOD(TargetPlatform);
-	const TArray<FString> GroupDerivedDataKeys = GetGroupDerivedDataKeys(this, TargetPlatform, SourceMeshLOD, TargetMeshMinLOD);
+	const TArray<FString> GroupDerivedDataKeys = GetGroupDerivedDataKeys(this, TargetPlatform, SourceMeshLOD, TargetMeshMinLOD, true /*bForceSkelMeshDataCachingForCookingOnly */);
 
 	// 2. Find existing cached cooked data
 	UGroomBindingAsset::FCachedCookedPlatformData* TargetPlatformData = FindCachedCookedPlatformData(GroupDerivedDataKeys, CachedCookedPlatformDatas);
