Bug fix for reading strand positions, resolving the frame difference issue between the query and the readback.

Work done by [REDACTED_PERSON].

diff --git a/Engine/Plugins/Runtime/HairStrands/Source/HairStrandsCore/Private/GroomBuilder.cpp b/Engine/Plugins/Runtime/HairStrands/Source/HairStrandsCore/Private/GroomBuilder.cpp
--- a/Engine/Plugins/Runtime/HairStrands/Source/HairStrandsCore/Private/GroomBuilder.cpp
+++ b/Engine/Plugins/Runtime/HairStrands/Source/HairStrandsCore/Private/GroomBuilder.cpp
@@ -3887,55 +3887,55 @@ IMPLEMENT_GLOBAL_SHADER(FHairStrandCopyPositionCS, "/Engine/Private/HairStrands/
 static FRDGBufferRef ReadPositions(
 	FRDGBuilder& GraphBuilder,
 	FGlobalShaderMap* ShaderMap,
-	const UGroomComponent* Component,
-	uint32 GroupIt)
+	const FHairGroupInstance* Instance)
 {
+	check(Instance);
+
 	FRDGBufferRef OutPositions = nullptr;
-	if (const FHairGroupInstance* Instance = Component->GetGroupInstance(GroupIt))
+	const uint32 NumVertices = Instance->HairGroupPublicData->GetActiveStrandsPointCount();
+
+	FRDGBufferSRVRef InPositions = nullptr;
+	FRDGBufferSRVRef InPositionOffset = nullptr;
+	if (Instance->Strands.DeformedResource)
 	{
-		const uint32 NumVertices = Instance->Strands.GetData().GetNumPoints();
+		InPositions = Register(GraphBuilder, Instance->Strands.DeformedResource->GetBuffer(FHairStrandsDeformedResource::EFrameType::Current), ERDGImportedBufferFlags::CreateSRV).SRV;
+		InPositionOffset = Register(GraphBuilder, Instance->Strands.DeformedResource->GetPositionOffsetBuffer(FHairStrandsDeformedResource::EFrameType::Current), ERDGImportedBufferFlags::CreateSRV).SRV;
+	}
+	else
+	{
+		InPositions = Register(GraphBuilder, Instance->Strands.RestResource->PositionBuffer, ERDGImportedBufferFlags::CreateSRV).SRV;
+		InPositionOffset = Register(GraphBuilder, Instance->Strands.RestResource->PositionOffsetBuffer, ERDGImportedBufferFlags::CreateSRV).SRV;
+	}
 
-		FRDGBufferSRVRef InPositions = nullptr;
-		FRDGBufferSRVRef InPositionOffset = nullptr;
-		if (Instance->Strands.DeformedResource)
-		{
-			InPositions = Register(GraphBuilder, Instance->Strands.DeformedResource->GetBuffer(FHairStrandsDeformedResource::EFrameType::Current), ERDGImportedBufferFlags::CreateSRV).SRV;
-			InPositionOffset = Register(GraphBuilder, Instance->Strands.DeformedResource->GetPositionOffsetBuffer(FHairStrandsDeformedResource::EFrameType::Current), ERDGImportedBufferFlags::CreateSRV).SRV;
-		}
-		else
-		{
-			InPositions = Register(GraphBuilder, Instance->Strands.RestResource->PositionBuffer, ERDGImportedBufferFlags::CreateSRV).SRV;
-			InPositionOffset = Register(GraphBuilder, Instance->Strands.RestResource->PositionOffsetBuffer, ERDGImportedBufferFlags::CreateSRV).SRV;
-		}
+	OutPositions = GraphBuilder.CreateBuffer(FRDGBufferDesc::CreateBufferDesc(sizeof(FVector4f), NumVertices), TEXT("StrandsPositionForReadback"));
 
-		OutPositions = GraphBuilder.CreateBuffer(FRDGBufferDesc::CreateBufferDesc(sizeof(FVector4f), NumVertices), TEXT("StrandsPositionForReadback"));
+	FHairStrandCopyPositionCS::FParameters* Parameters = GraphBuilder.AllocParameters<FHairStrandCopyPositionCS::FParameters>();
+	Parameters->MaxVertexCount = NumVertices;
+	Parameters->InPositionOffset = InPositionOffset;
+	Parameters->InPositions = InPositions;
+	Parameters->OutPositions = GraphBuilder.CreateUAV(OutPositions, PF_A32B32G32R32F);
 
-		FHairStrandCopyPositionCS::FParameters* Parameters = GraphBuilder.AllocParameters<FHairStrandCopyPositionCS::FParameters>();
-		Parameters->MaxVertexCount = NumVertices;
-		Parameters->InPositionOffset = InPositionOffset;
-		Parameters->InPositions = InPositions;
-		Parameters->OutPositions = GraphBuilder.CreateUAV(OutPositions, PF_A32B32G32R32F);
+	FIntVector DispatchCount = FComputeShaderUtils::GetGroupCount(NumVertices, 128);
 
-		FIntVector DispatchCount = FComputeShaderUtils::GetGroupCount(NumVertices, 128);
+	TShaderMapRef<FHairStrandCopyPositionCS> ComputeShader(ShaderMap);
+	FComputeShaderUtils::AddPass(
+		GraphBuilder,
+		RDG_EVENT_NAME("HairStrands::CopyPositions"),
+		ComputeShader,
+		Parameters,
+		DispatchCount);
 
-		TShaderMapRef<FHairStrandCopyPositionCS> ComputeShader(ShaderMap);
-		FComputeShaderUtils::AddPass(
-			GraphBuilder,
-			RDG_EVENT_NAME("HairStrands::CopyPositions"),
-			ComputeShader,
-			Parameters,
-			DispatchCount);
-	}
 	return OutPositions;
 }
 
 struct FStrandsPositionData
 {
 	TUniquePtr<FRHIGPUBufferReadback> GPUPositions;
-	uint32 NumBytes = 0;
 	uint32 GroupIt = 0;
 	uint32 GroupCount = 0;
 	uint32 ReadbackDelay = 0;
+	uint32 CurveCount = 0;
+	uint32 PointCount = 0;
 	const UGroomComponent* Component = nullptr;
 	FStrandsPositionOutput* Output;
 };
@@ -3961,26 +3961,27 @@ void RunHairStrandsPositionQueries(FRDGBuilder& GraphBuilder, FGlobalShaderMap*
 			if (Q->GPUPositions->IsReady() || Q->ReadbackDelay == 0)
 			{
 				const FHairGroupInstance* Instance = Q->Component->GetGroupInstance(Q->GroupIt);
-				const uint32 CurveCount = Instance->Strands.RestResource->BulkData.GetNumCurves();
-				const uint32 PointCount = Instance->Strands.RestResource->BulkData.GetNumPoints();
 
 				check(Q->GroupIt < uint32(Q->Output->Groups.Num()));
 				FStrandsPositionOutput::FGroup& Group = Q->Output->Groups[Q->GroupIt];
-				Group.Reserve(CurveCount);
+				Group.Reserve(Q->CurveCount);
 
 				FStrandsPositionOutput::FStrand Strand;
-				const FVector4f* Positions = (const FVector4f*)Q->GPUPositions->Lock(Q->NumBytes);
-				for (uint32 PointIt=0; PointIt<PointCount; ++PointIt)
+				const FVector4f* Positions = (const FVector4f*)Q->GPUPositions->Lock(sizeof(FVector4f) * Q->PointCount);
+				for (uint32 PointIt = 0; PointIt < Q->PointCount; ++PointIt)
 				{
 					const FVector4f P = Positions[PointIt];
-					const bool bIsFirstPoint = P.W == 1;
-					if (bIsFirstPoint && Strand.Num() > 0)
+					check(P.W != HAIR_CONTROLPOINT_START || Strand.Num() == 0);
+					check(P.W != HAIR_CONTROLPOINT_INSIDE || Strand.Num() > 0);
+
+					Strand.Add(FVector3f(P.X, P.Y, P.Z));
+					if (P.W == HAIR_CONTROLPOINT_END)
 					{
 						Group.Add(Strand);
 						Strand.SetNum(0);
 					}
-					Strand.Add(FVector3f(P.X, P.Y, P.Z));
 				}
+				check(Q->CurveCount == Group.Num());
 				Q->GPUPositions->Unlock();
 
 				--Q->Output->Status;
@@ -4005,12 +4006,22 @@ void RunHairStrandsPositionQueries(FRDGBuilder& GraphBuilder, FGlobalShaderMap*
 		while (GStrandsPositionQueries.Dequeue(Q))
 		{
 			check(Q->Component);
-			
-			FRDGBufferRef Positions = ReadPositions(GraphBuilder, ShaderMap, Q->Component, Q->GroupIt);
-			FRHIGPUBufferReadback* GPUPositions = Q->GPUPositions.Get();
-			AddEnqueueCopyPass(GraphBuilder, GPUPositions, Positions, Positions->Desc.GetSize());
-			Q->NumBytes = Positions->Desc.GetSize();
-			GStrandsPositionReadbacks.Enqueue(Q);
+			if (const FHairGroupInstance* Instance = Q->Component->GetGroupInstance(Q->GroupIt))
+			{
+				Q->CurveCount = Instance->HairGroupPublicData->GetActiveStrandsCurveCount();
+				Q->PointCount = Instance->HairGroupPublicData->GetActiveStrandsPointCount();
+
+				FRDGBufferRef Positions = ReadPositions(GraphBuilder, ShaderMap, Instance);
+				AddEnqueueCopyPass(GraphBuilder, Q->GPUPositions.Get(), Positions, Positions->Desc.GetSize());
+				GStrandsPositionReadbacks.Enqueue(Q);
+			}
+			else
+			{
+				// If the instance is missing, Q->Output->Groups[GroupIt] remains an empty array
+				// in this case, skip the readback to prevent the query from remaining unfinished
+				--Q->Output->Status;
+				check(Q->Output->Status >= 0);
+			}
 		}
 	}
 }
@@ -4027,11 +4038,12 @@ bool RequestStrandsPosition(const UGroomComponent* Component, FStrandsPositionOu
 	for (uint32 GroupIt = 0; GroupIt < GroupCount; ++GroupIt)
 	{
 		FStrandsPositionData* Data = new FStrandsPositionData();
-		Data->NumBytes = 0;
 		Data->GroupIt = GroupIt;
 		Data->GroupCount = GroupCount;
 		Data->Component = Component;
 		Data->ReadbackDelay = 3;
+		Data->CurveCount = 0;
+		Data->PointCount = 0;
 		Data->GPUPositions = MakeUnique<FRHIGPUBufferReadback>(TEXT("Readback.Positions"));
 		Data->Output = Output;
 		GStrandsPositionQueries.Enqueue(Data);
