Fix Hair crash when switching groom asset

diff --git a/Engine/Plugins/Runtime/HairStrands/Source/HairStrandsCore/Private/Niagara/NiagaraDataInterfaceHairStrands.cpp b/Engine/Plugins/Runtime/HairStrands/Source/HairStrandsCore/Private/Niagara/NiagaraDataInterfaceHairStrands.cpp
--- a/Engine/Plugins/Runtime/HairStrands/Source/HairStrandsCore/Private/Niagara/NiagaraDataInterfaceHairStrands.cpp
+++ b/Engine/Plugins/Runtime/HairStrands/Source/HairStrandsCore/Private/Niagara/NiagaraDataInterfaceHairStrands.cpp
@@ -558,7 +558,7 @@ void FNDIHairStrandsData::Update(
 				FMatrix44d BoneTransformDouble = BoneTransform.ToMatrixWithScale();
 				const FMatrix44d WorldTransformDouble = WorldTransform.ToMatrixWithScale();
 
-				if(DeltaSeconds != 0.0f && (TickCount > GHairSimulationMaxDelay))
+				if(DeltaSeconds != 0.0f && !ForceReset)
 				{
 					const FMatrix44d PreviousBoneTransformDouble = PreviousBoneTransform.ToMatrixWithScale();
 					const FMatrix44d DeltaTransformDouble =  BoneTransformDouble * PreviousBoneTransformDouble.Inverse();
@@ -635,7 +635,6 @@ bool FNDIHairStrandsData::Init(UNiagaraDataInterfaceHairStrands* Interface, FNia
 
 			BeginInitResource(HairStrandsBuffer);
 
-			TickCount = 0;
 			ForceReset = true;
 		}
 	}
@@ -863,19 +862,26 @@ void UNiagaraDataInterfaceHairStrands::DestroyPerInstanceData(void* PerInstanceD
 bool UNiagaraDataInterfaceHairStrands::PerInstanceTick(void* PerInstanceData, FNiagaraSystemInstance* SystemInstance, float InDeltaSeconds)
 {
 	FNDIHairStrandsData* InstanceData = static_cast<FNDIHairStrandsData*>(PerInstanceData);
-	InstanceData->TickCount = FMath::Min(GHairSimulationMaxDelay + 1, InstanceData->TickCount + 1);
 
 	FNDIHairStrandsInfo InfoData;
 	ExtractDatasAndResources(SystemInstance, InfoData);
 
 	if (SourceComponent != nullptr)
 	{
-		if (SourceComponent->bResetSimulation || RequiresSimulationReset(SystemInstance, InstanceData->SkeletalMeshes))
-			
+		InstanceData->ForceReset = SourceComponent->bResetSimulation || RequiresSimulationReset(SystemInstance, InstanceData->SkeletalMeshes);
+		if (InstanceData->ForceReset)
 		{
-			InstanceData->TickCount = 0;
+			FNDIHairStrandsBuffer* LocalStrandsBuffer = InstanceData->HairStrandsBuffer;
+			ENQUEUE_RENDER_COMMAND(FNiagaraDIDestroyInstanceData) (
+				[LocalStrandsBuffer](FRHICommandListImmediate& CmdList)
+				{
+					if(LocalStrandsBuffer)
+					{
+						LocalStrandsBuffer->bShouldReset = true;
+					}
+				}
+			);
 		}
-		InstanceData->ForceReset = SourceComponent->bResetSimulation;
 	}
 	InstanceData->Update(this, InfoData, InDeltaSeconds);
 	return false;
@@ -3445,10 +3451,14 @@ void UNiagaraDataInterfaceHairStrands::SetShaderParameters(const FNiagaraDataInt
 
 		// Simulation setup (we update the rest configuration based on the deformed positions 
 		// if in restupdate mode or if we are resetting the sim and using RBF transfer since the rest positions are not matching the physics asset)
-		const int32 NeedResetValue = (ProxyData->TickCount <= GHairSimulationMaxDelay) || !HairStrandsBuffer->bValidGeometryType || (HairStrandsBuffer->CurrentMeshLOD != MeshLODIndex);
+		const bool bShouldReset = HairStrandsBuffer->bShouldReset || !HairStrandsBuffer->bValidGeometryType || (HairStrandsBuffer->CurrentMeshLOD != MeshLODIndex);
+		HairStrandsBuffer->ResetCount = bShouldReset ? 0 : FMath::Min(GHairSimulationMaxDelay + 1, HairStrandsBuffer->ResetCount+ 1);
+		
+		const int32 NeedResetValue = HairStrandsBuffer->ResetCount <= GHairSimulationMaxDelay;
 		const int32 RestUpdateValue = GHairSimulationRestUpdate || (NeedResetValue && ProxyData->bSkinningTransfer) ? 1 : 0;
 		const int32 LocalSimulationValue = ProxyData->LocalSimulation;
 		
+		HairStrandsBuffer->bShouldReset = false;
 		HairStrandsBuffer->bValidGeometryType = true;
 		HairStrandsBuffer->CurrentMeshLOD = MeshLODIndex;
 
diff --git a/Engine/Plugins/Runtime/HairStrands/Source/HairStrandsCore/Public/Niagara/NiagaraDataInterfaceHairStrands.h b/Engine/Plugins/Runtime/HairStrands/Source/HairStrandsCore/Public/Niagara/NiagaraDataInterfaceHairStrands.h
--- a/Engine/Plugins/Runtime/HairStrands/Source/HairStrandsCore/Public/Niagara/NiagaraDataInterfaceHairStrands.h
+++ b/Engine/Plugins/Runtime/HairStrands/Source/HairStrandsCore/Public/Niagara/NiagaraDataInterfaceHairStrands.h
@@ -65,6 +65,12 @@ struct FNDIHairStrandsBuffer : public FRenderResource
 	/** True if the internal resources (BoundingBoxBuffer/ParamsScaleBuffer) needs to be built */
 	bool bNeedResouces = false;
 
+	/** Boolean to trigger the reset */
+	bool bShouldReset = false;
+	
+	/** Counter to reset the simulation once triggered */
+	int32 ResetCount = 0;
+
 	// For debug only
 	//FRHIGPUBufferReadback* ReadbackBuffer = nullptr;
 };
@@ -108,7 +114,6 @@ struct FNDIHairStrandsData
 		HairGroupInstance = nullptr;
 		HairGroupIndex = -1;
 
-		TickCount = 0;
 		ForceReset = true;
 
 		NumStrands = 0;
@@ -184,9 +189,8 @@ struct FNDIHairStrandsData
 				HairGroupInstance = GetHairGroupInstance(HairGroupInstSource.Get(), HairGroupIndex);
 			}
 
-			TickCount = OtherDatas->TickCount;
 			ForceReset = OtherDatas->ForceReset;
-
+			
 			NumStrands = OtherDatas->NumStrands;
 			StrandsSize = OtherDatas->StrandsSize;
 
@@ -268,9 +272,6 @@ struct FNDIHairStrandsData
 	/** Strand size */
 	int32 StrandsSize;
 
-	/** Tick Count*/
-	int32 TickCount;
-
 	/** Force reset simulation */
 	bool ForceReset;
 
