Fix cook of groom binding asset on certain target platform, due to skel. mesh updating their DDC when building their data.

The initial fix (CL 36826968) caused a deadlock as it involved taking a lock on the skel mesh, and was reverted in CL 36936042. Instead, this new CL recomputes the binding asset platform data a second time if the skel. mesh DDC key is not identical before and after the groom binding building.

diff --git a/Engine/Plugins/Runtime/HairStrands/Source/HairStrandsCore/Private/GroomBindingAsset.cpp b/Engine/Plugins/Runtime/HairStrands/Source/HairStrandsCore/Private/GroomBindingAsset.cpp
--- a/Engine/Plugins/Runtime/HairStrands/Source/HairStrandsCore/Private/GroomBindingAsset.cpp
+++ b/Engine/Plugins/Runtime/HairStrands/Source/HairStrandsCore/Private/GroomBindingAsset.cpp
@@ -1624,45 +1624,60 @@ void UGroomBindingAsset::BeginCacheForCookedPlatformData(const ITargetPlatform*
 	// 3. If the target cooked data does not already exist, we build it
 	if (TargetPlatformData == nullptr && GetGroom() != nullptr)
 	{
-		// 3.1 Build cooked derived data
-		const uint32 GroupCount = GroupDerivedDataKeys.Num();
-		TargetPlatformData = new FCachedCookedPlatformData();
-		TargetPlatformData->GroupDerivedDataKeys = GroupDerivedDataKeys;
-		TargetPlatformData->GroupPlatformDatas.SetNum(GroupCount);
-		for (uint32 GroupIndex = 0; GroupIndex < GroupCount; ++GroupIndex)
+		auto BuildTargetPlatformData = [TargetPlatform, SourceMeshLOD, TargetMeshMinLOD](const TArray<FString>& InGroupDerivedDataKeys, UGroomBindingAsset* BindingAsset)
 		{
-			bool bGroupValid = true;
-			CacheDerivedDatas(
-				this, 
-				GroupIndex, 
-				TargetPlatformData->GroupDerivedDataKeys[GroupIndex], 
-				bGroupValid, 
-				TargetPlatform, 
-				SourceMeshLOD, 
-				TargetMeshMinLOD, 
-				TargetPlatformData->GroupPlatformDatas[GroupIndex]);
-
-			if (!bGroupValid)
+			// 3.a Build cooked derived data
+			const uint32 GroupCount = InGroupDerivedDataKeys.Num();
+			UGroomBindingAsset::FCachedCookedPlatformData* NewTargetPlatformData = new FCachedCookedPlatformData();
+			NewTargetPlatformData->GroupDerivedDataKeys = InGroupDerivedDataKeys;
+			NewTargetPlatformData->GroupPlatformDatas.SetNum(GroupCount);
+			for (uint32 GroupIndex = 0; GroupIndex < GroupCount; ++GroupIndex)
 			{
-				UE_LOG(LogHairStrands, Error, TEXT("[Groom] The binding asset (%s) couldn't be built. This binding asset won't be used."), *GetName());
+				bool bGroupValid = true;
+				CacheDerivedDatas(
+					BindingAsset, 
+					GroupIndex, 
+					NewTargetPlatformData->GroupDerivedDataKeys[GroupIndex], 
+					bGroupValid, 
+					TargetPlatform, 
+					SourceMeshLOD, 
+					TargetMeshMinLOD, 
+					NewTargetPlatformData->GroupPlatformDatas[GroupIndex]);
+
+				if (!bGroupValid)
+				{
+					UE_LOG(LogHairStrands, Error, TEXT("[Groom] The binding asset (%s) couldn't be built. This binding asset won't be used."), *BindingAsset->GetName());
+				}
 			}
-		}
 
-		// 3.2 Place cooked derived data into their bulk data. 
-		// This is done only for strands, which support DDC streaming
-		// When cooking data, force loading of *all* bulk data prior to saving them
-		// Note: bFillBulkdata is true for filling in the bulkdata container prior to serialization. This also forces the resources loading 
-		// from the 'start' (i.e., without offset)
-		for (uint32 GroupIndex = 0; GroupIndex < GroupCount; ++GroupIndex)
-		{
-			TArray<FHairStrandsRootBulkData>& RenRootBulkDatas = TargetPlatformData->GroupPlatformDatas[GroupIndex].RenRootBulkDatas;
-			for (int32 MeshLODIndex = 0, MeshLODCount = RenRootBulkDatas.Num(); MeshLODIndex < MeshLODCount; ++MeshLODIndex)
+			// 3.b Place cooked derived data into their bulk data. 
+			// This is done only for strands, which support DDC streaming
+			// When cooking data, force loading of *all* bulk data prior to saving them
+			// Note: bFillBulkdata is true for filling in the bulkdata container prior to serialization. This also forces the resources loading 
+			// from the 'start' (i.e., without offset)
+			for (uint32 GroupIndex = 0; GroupIndex < GroupCount; ++GroupIndex)
 			{
-				FHairStreamingRequest R; R.Request(HAIR_MAX_NUM_CURVE_PER_GROUP, HAIR_MAX_NUM_POINT_PER_GROUP, RenRootBulkDatas[MeshLODIndex], true /*bWait*/, true /*bFillBulkdata*/, true /*bWarmCache*/, GetFName());
+				TArray<FHairStrandsRootBulkData>& RenRootBulkDatas = NewTargetPlatformData->GroupPlatformDatas[GroupIndex].RenRootBulkDatas;
+				for (int32 MeshLODIndex = 0, MeshLODCount = RenRootBulkDatas.Num(); MeshLODIndex < MeshLODCount; ++MeshLODIndex)
+				{
+					FHairStreamingRequest R; R.Request(HAIR_MAX_NUM_CURVE_PER_GROUP, HAIR_MAX_NUM_POINT_PER_GROUP, RenRootBulkDatas[MeshLODIndex], true /*bWait*/, true /*bFillBulkdata*/, true /*bWarmCache*/, BindingAsset->GetFName());
+				}
 			}
-		}
 
-		CachedCookedPlatformDatas.Add(TargetPlatformData);
+			BindingAsset->CachedCookedPlatformDatas.Add(NewTargetPlatformData);
+		};
+
+		// 3.1 Build binding data for the target platform
+		BuildTargetPlatformData(GroupDerivedDataKeys, this);
+
+		// 3.2 Check if the derived data keys are unchanged. If they have change, recompute the target platform binding data again.
+		// This can happen when a skel. mesh asset hasn't be resaved correct. While being cached, the skel. mesh data will update/transform 
+		// its DDC key, causing the binding data to be cached as the incorrect DDC key.
+		const TArray<FString> UpdatedGroupDerivedDataKeys = GetGroupDerivedDataKeys(this, TargetPlatform, SourceMeshLOD, TargetMeshMinLOD, false /*bForceSkelMeshDataCachingForCookingOnly */);
+		if (UpdatedGroupDerivedDataKeys != GroupDerivedDataKeys)
+		{
+			BuildTargetPlatformData(UpdatedGroupDerivedDataKeys, this);
+		}
 	}
 }
 
