Fix groom binding cooking on target platform.

The MinLOD value used for computing the binding DDC key changes/mutates once the skel. mesh data rebuild. When caching the groom binding data, we check if the DDC key has changed to cache the data a second time with the correct DDC key. The MinLOD value was not correctly updated.

diff --git a/Engine/Plugins/Runtime/HairStrands/Source/HairStrandsCore/Private/GroomBindingAsset.cpp b/Engine/Plugins/Runtime/HairStrands/Source/HairStrandsCore/Private/GroomBindingAsset.cpp
--- a/Engine/Plugins/Runtime/HairStrands/Source/HairStrandsCore/Private/GroomBindingAsset.cpp
+++ b/Engine/Plugins/Runtime/HairStrands/Source/HairStrandsCore/Private/GroomBindingAsset.cpp
@@ -1673,7 +1673,9 @@ void UGroomBindingAsset::BeginCacheForCookedPlatformData(const ITargetPlatform*
 		// 3.2 Check if the derived data keys are unchanged. If they have change, recompute the target platform binding data again.
 		// This can happen when a skel. mesh asset hasn't be resaved correct. While being cached, the skel. mesh data will update/transform 
 		// its DDC key, causing the binding data to be cached as the incorrect DDC key.
-		const TArray<FString> UpdatedGroupDerivedDataKeys = GetGroupDerivedDataKeys(this, TargetPlatform, SourceMeshLOD, TargetMeshMinLOD, false /*bForceSkelMeshDataCachingForCookingOnly */);
+		const int32 UpdatedSourceMeshLOD = GetClampedSourceMeshLOD(TargetPlatform);
+		const int32 UpdatedTargetMeshMinLOD = GetClampedTargetMeshMinLOD(TargetPlatform);
+		const TArray<FString> UpdatedGroupDerivedDataKeys = GetGroupDerivedDataKeys(this, TargetPlatform, UpdatedSourceMeshLOD, UpdatedTargetMeshMinLOD, false /*bForceSkelMeshDataCachingForCookingOnly */);
 		if (UpdatedGroupDerivedDataKeys != GroupDerivedDataKeys)
 		{
 			BuildTargetPlatformData(UpdatedGroupDerivedDataKeys, this);
