Fixed crash in HairStrandsInterpolation on GPUs where GRHIMinimumWaveSize is detected as 64

diff --git a/Engine/Plugins/Runtime/HairStrands/Source/HairStrandsCore/Private/HairStrandsInterpolation.cpp b/Engine/Plugins/Runtime/HairStrands/Source/HairStrandsCore/Private/HairStrandsInterpolation.cpp
--- a/Engine/Plugins/Runtime/HairStrands/Source/HairStrandsCore/Private/HairStrandsInterpolation.cpp
+++ b/Engine/Plugins/Runtime/HairStrands/Source/HairStrandsCore/Private/HairStrandsInterpolation.cpp
@@ -709,7 +709,7 @@ void AddHairStrandsInterpolationPass(
 	Parameters->DispatchCountX = DispatchInfo.DispatchCount.X;
 
 	const bool bSingleGuidePermutation = bUseSingleGuide || GetHairStrandsForceSingleGuideInterpolation();
-	const bool bWaveOps = GRHISupportsWaveOperations && GRHIMaximumWaveSize >= 32 && FHairInterpolationCS::DoesSupportsWaveOps(InPlatform, DispatchInfo.PointPerCurve) != ERHIFeatureSupport::Unsupported;
+	const bool bWaveOps = GRHISupportsWaveOperations && GRHIMinimumWaveSize <= 32 && GRHIMaximumWaveSize >= 32 && FHairInterpolationCS::DoesSupportsWaveOps(InPlatform, DispatchInfo.PointPerCurve) != ERHIFeatureSupport::Unsupported;
 
 	FHairInterpolationCS::FPermutationDomain PermutationVector;
 	PermutationVector.Set<FHairInterpolationCS::FDynamicGeometry>(DynamicGeometryType);
