* Remove ForceSkeletalMeshDataCachingForCookingOnly which is no longer used (preflight tentative fix for addressing cook issue)
* Merge valid LOD gathering and MinLOD finding into GetClampedSkeletalMeshLOD GetSkeletalMeshLOD
* Fix target LOD referencing when building groom data in cook build.

diff --git a/Engine/Plugins/Runtime/HairStrands/Source/HairStrandsCore/Private/GroomBindingAsset.cpp b/Engine/Plugins/Runtime/HairStrands/Source/HairStrandsCore/Private/GroomBindingAsset.cpp
--- a/Engine/Plugins/Runtime/HairStrands/Source/HairStrandsCore/Private/GroomBindingAsset.cpp
+++ b/Engine/Plugins/Runtime/HairStrands/Source/HairStrandsCore/Private/GroomBindingAsset.cpp
@@ -945,54 +945,6 @@ void UpdateGroomBindingAssetInfos(UGroomBindingAsset* In)
 // * Release : Engine\Source\Runtime\Core\Private\UObject\UE5ReleaseStreamObjectVersion.cpp
 // * ...
 
-static void ForceSkeletalMeshDataCachingForCookingOnly(const ITargetPlatform* TargetPlatform, USkeletalMesh* InSourceSkeletalMesh, USkeletalMesh* InTargetSkeletalMesh)
-{
-	// If a skeletal mesh build or preedit is called while we're async compiling groom bindings, the groom binding compiler
-	// will take care of finishing any groom binding that depends on the skeletal mesh being modified. So this is 
-	// safe to do asynchronously without locks on the render data. We assume that the skeletal mesh's render data is
-	// immutable once it has been built and can only be rebuilt throught a call to PreEditChange first.
-	const bool bAsyncCompiling = !IsInGameThread();
-
-	// * Only for SkeletalMesh: Take scoped lock on the skeletal render mesh data during the entire groom binding building
-	// * Then use an async build scope to allow accessing skeletal mesh property safely.
-	//   If skel.meshes are nullptr, this will act as a NOP
-	InSourceSkeletalMesh = InSourceSkeletalMesh == InTargetSkeletalMesh ? nullptr : InSourceSkeletalMesh;
-	PRAGMA_DISABLE_DEPRECATION_WARNINGS
-	FScopedSkeletalMeshRenderData SourceSkeletalMeshScopedData(bAsyncCompiling ? nullptr : InSourceSkeletalMesh);
-	FScopedSkeletalMeshRenderData TargetSkeletalMeshScopedData(bAsyncCompiling ? nullptr : InTargetSkeletalMesh);
-	PRAGMA_ENABLE_DEPRECATION_WARNINGS
-
-	if (InSourceSkeletalMesh)
-	{
-		if (bAsyncCompiling)
-		{
-			const ITargetPlatform* RunningPlatform = GetTargetPlatformManagerRef().GetRunningTargetPlatform();
-			checkf(TargetPlatform == RunningPlatform, TEXT("It is only safe to query the running platform's render data asynchronously from the skeletal mesh"));
-			InSourceSkeletalMesh->GetResourceForRendering();
-		}
-		else
-		{
-			FSkinnedAssetAsyncBuildScope AsyncBuildScope(InSourceSkeletalMesh);
-			USkeletalMesh::GetPlatformSkeletalMeshRenderData(TargetPlatform, SourceSkeletalMeshScopedData);
-		}
-	}
-
-	if (InTargetSkeletalMesh)
-	{
-		if (bAsyncCompiling)
-		{
-			const ITargetPlatform* RunningPlatform = GetTargetPlatformManagerRef().GetRunningTargetPlatform();
-			checkf(TargetPlatform == RunningPlatform, TEXT("It is only safe to query the running platform's render data asynchronously from the skeletal mesh"));
-			InTargetSkeletalMesh->GetResourceForRendering();
-		}
-		else
-		{
-			FSkinnedAssetAsyncBuildScope AsyncBuildScope(InTargetSkeletalMesh);
-			USkeletalMesh::GetPlatformSkeletalMeshRenderData(TargetPlatform, TargetSkeletalMeshScopedData);
-		}
-	}
-}
-
 namespace GroomBindingDerivedDataCacheUtils
 {
 	const FString& GetGroomBindingDerivedDataVersion()
@@ -1007,7 +959,7 @@ namespace GroomBindingDerivedDataCacheUtils
 	}
 }
 
-static FString BuildDerivedDataKeySuffix(const UGroomBindingAsset& BindingAsset, const ITargetPlatform* TargetPlatform, int32 SourceMeshLOD, int32 TargetMeshMinLOD, bool bForceSkelMeshDataCachingForCookingOnly=false)
+static FString BuildDerivedDataKeySuffix(const UGroomBindingAsset& BindingAsset, const ITargetPlatform* TargetPlatform, int32 SourceMeshLOD, int32 TargetMeshMinLOD)
 {
 	FString BindingType;
 	FString SourceKey;
@@ -1017,12 +969,6 @@ static FString BuildDerivedDataKeySuffix(const UGroomBindingAsset& BindingAsset,
 
 	if (BindingAsset.GetGroomBindingType() == EGroomBindingMeshType::SkeletalMesh)
 	{
-		// Force skel. mesh data caching, in order to produce later an valid/up-to-date DDC key
-		if (bForceSkelMeshDataCachingForCookingOnly)
-		{
-			ForceSkeletalMeshDataCachingForCookingOnly(TargetPlatform, BindingAsset.GetSourceSkeletalMesh(), BindingAsset.GetTargetSkeletalMesh());
-		}
-
 		// Binding type is implicitly SkeletalMesh so keep BindingType empty to prevent triggering rebuild of old binding for nothing
 		if (USkeletalMesh* SourceMesh = BindingAsset.GetSourceSkeletalMesh())
 		{
@@ -1061,7 +1007,7 @@ static FString BuildDerivedDataKeyGroup(const FString& InDeriveDataKeySuffix , u
 	return GroomBindingDerivedDataCacheUtils::BuildGroomBindingDerivedDataKey(InDeriveDataKeySuffix + FString(TEXT("_Group")) + FString::FromInt(InGroupIndex));
 }
 
-static TArray<FString> GetGroupDerivedDataKeys(const UGroomBindingAsset* In, const ITargetPlatform* TargetPlatform, int32 SourceMeshLOD, int32 TargetMeshMinLOD, bool bForceSkelMeshDataCaching=false);
+static TArray<FString> GetGroupDerivedDataKeys(const UGroomBindingAsset* In, const ITargetPlatform* TargetPlatform, int32 SourceMeshLOD, int32 TargetMeshMinLOD);
 static void CacheDerivedDatas(
 	UGroomBindingAsset* In,
 	const uint32 InGroupIndex,
@@ -1154,98 +1100,113 @@ void UGroomBindingAsset::BeginCacheDerivedDatas(
 	}
 }
 
-bool UGroomBindingAsset::TryInitializeContextForMeshes(FGroomBindingBuildContext& OutContext) const
+static int32 GetClampedSkeletalMeshLOD(const USkeletalMesh* SkeletalMesh, int32 RequestedMinLOD, const ITargetPlatform* TargetPlatform)
 {
-	ITargetPlatform* RunningPlatform = nullptr;
-#if WITH_EDITOR
-	RunningPlatform = GetTargetPlatformManagerRef().GetRunningTargetPlatform();
-	check(RunningPlatform);
-#endif
-
-	if (GetGroomBindingType() == EGroomBindingMeshType::SkeletalMesh)
+	if (!SkeletalMesh)
 	{
-		if (USkeletalMesh* SourceMesh = GetSourceSkeletalMesh())
-		{
-			OutContext.SourceMeshLOD = GetClampedSourceMeshLOD(RunningPlatform);
+		return INDEX_NONE;
+	}
 
-			// Clamp SourceMeshLOD to the next best LOD that's currently loaded.
-			const FStreamableRenderResourceState& SRRState = SourceMesh->GetStreamableResourceState();
-			if (SRRState.IsValid())
-			{
-				const int32 NumValidLODs = FMath::Min(SRRState.NumRequestedLODs, SRRState.NumResidentLODs);
-				const int32 CurrentFirstLOD = SRRState.LODCountToAssetFirstLODIdx(NumValidLODs);
-				OutContext.SourceMeshLOD = FMath::Max(CurrentFirstLOD, OutContext.SourceMeshLOD);
-			}
+	// BestAvailableLOD is the best (i.e. lowest index) LOD that can ever be used by this mesh.
+	// E.g. in a cooked build, this is the best LOD that was cooked.
+	// 
+	// RequestedMinLOD is clamped to it here, so that if the mesh has a per-platform min LOD,
+	// the caller of UGroomBindingAsset::Build doesn't need to take that into account. They can
+	// simply set RequestedMinLOD to 0 and if min LOD on the mesh is 3, for example, it will
+	// silently build the groom binding using LOD 3 as the source.
+	#if WITH_EDITOR
+	check(TargetPlatform);
+	const int32 BestAvailableLOD = SkeletalMesh->GetPlatformMinLODIdx(TargetPlatform);
+	#else
+	const bool bForceLowestLODIdx = true;
+	const int32 BestAvailableLOD = SkeletalMesh->GetMinLodIdx(bForceLowestLODIdx);
+	#endif
+
+	return FMath::Max(RequestedMinLOD, BestAvailableLOD);
+}
 
-			if (FSkeletalMeshRenderData* RenderData = SourceMesh->GetResourceForRendering())
+static bool GetSkeletalMeshLOD(
+	const UGroomBindingAsset* Binding,
+	const USkeletalMesh* SkeletalMesh, 
+	const int32 RequestedMinLOD, 
+	const bool bRequestSingleLOD,
+	ITargetPlatform* InPlatform, 
+	int32& OutMeshMinLOD, 
+	TArray<TRefCountPtr<FSkeletalMeshLODRenderData>>& OutMeshLODReferences)
+{	
+	check(SkeletalMesh);
+	OutMeshMinLOD = GetClampedSkeletalMeshLOD(SkeletalMesh, RequestedMinLOD, InPlatform);
+
+	// Clamp SourceMeshLOD to the next best LOD that's currently loaded.
+	const FStreamableRenderResourceState& SRRState = SkeletalMesh->GetStreamableResourceState();
+	if (SRRState.IsValid())
+	{
+		const int32 NumValidLODs = FMath::Min(SRRState.NumRequestedLODs, SRRState.NumResidentLODs);
+		const int32 CurrentFirstLOD = SRRState.LODCountToAssetFirstLODIdx(NumValidLODs);
+		OutMeshMinLOD = FMath::Max(CurrentFirstLOD, OutMeshMinLOD);
+	}
+
+	if (FSkeletalMeshRenderData* RenderData = SkeletalMesh->GetResourceForRendering())
+	{
+		if (RenderData->LODRenderData.IsValidIndex(OutMeshMinLOD))
+		{
+			const int32 LODCount = bRequestSingleLOD ? 1u : SkeletalMesh->GetLODNum();
+			for (int32 MeshLODIndex = OutMeshMinLOD; MeshLODIndex < LODCount; ++MeshLODIndex)
 			{
 				TRefCountPtr<FSkeletalMeshLODRenderData> LODReference;
-
-				if (RenderData->LODRenderData.IsValidIndex(OutContext.SourceMeshLOD))
+				if (RenderData->LODRenderData.IsValidIndex(MeshLODIndex))
 				{
-					LODReference = &RenderData->LODRenderData[OutContext.SourceMeshLOD];
+					LODReference = &RenderData->LODRenderData[MeshLODIndex];
 				}
 
 				if (!ensureMsgf(LODReference.IsValid(),
 					TEXT("Missing render data for mesh %s LOD %i, which is claimed to be valid by FStreamableRenderResourceState. Can't build groom binding %s."),
-					*SourceMesh->GetPathName(), OutContext.SourceMeshLOD, *GetPathName()))
+					*SkeletalMesh->GetPathName(), MeshLODIndex, *Binding->GetPathName()))
 				{
 					return false;
 				}
 
-				OutContext.MeshLODReferences.Add(MoveTemp(LODReference));
-			}
-			else
-			{
-				UE_LOGFMT(LogHairStrands, Error, "Source mesh {0} has no render data. Can't build groom binding {1}.", SourceMesh->GetPathName(), GetPathName());
-				return false;
+				OutMeshLODReferences.Add(MoveTemp(LODReference));
 			}
 		}
-
-		if (USkeletalMesh* TargetMesh = GetTargetSkeletalMesh())
+		else
 		{
-			OutContext.TargetMeshMinLOD = GetClampedTargetMeshMinLOD(RunningPlatform);
+			ensureMsgf(false,
+			TEXT("Missing render data for mesh %s LOD %i, which is claimed to be valid by FStreamableRenderResourceState. Can't build groom binding %s."),
+			*SkeletalMesh->GetPathName(), OutMeshMinLOD, *Binding->GetPathName());
 
-			const FStreamableRenderResourceState& SRRState = TargetMesh->GetStreamableResourceState();
-			if (SRRState.IsValid())
-			{
-				const int32 NumValidLODs = FMath::Min(SRRState.NumRequestedLODs, SRRState.NumResidentLODs);
-				const int32 CurrentFirstLOD = SRRState.LODCountToAssetFirstLODIdx(NumValidLODs);
-				OutContext.TargetMeshMinLOD = FMath::Max(CurrentFirstLOD, OutContext.TargetMeshMinLOD);
-			}
+			return false;
+		}
+	}
+	else
+	{
+		UE_LOGFMT(LogHairStrands, Error, "Skeletal mesh {0} has no render data. Can't build groom binding {1}.", SkeletalMesh->GetPathName(), Binding->GetPathName());
+		return false;
+	}
 
-			if (FSkeletalMeshRenderData* RenderData = TargetMesh->GetResourceForRendering())
-			{
-				if (RenderData->LODRenderData.IsValidIndex(OutContext.TargetMeshMinLOD))
-				{
-					for (int32 TargetMeshLODIndex = OutContext.TargetMeshMinLOD; TargetMeshLODIndex >= 0; TargetMeshLODIndex--)
-					{
-						TRefCountPtr<FSkeletalMeshLODRenderData> LODReference = &RenderData->LODRenderData[TargetMeshLODIndex];
+	return true;
+}
 
-						if (!ensureMsgf(LODReference.IsValid(),
-							TEXT("Missing render data for mesh %s LOD %i, which is claimed to be valid by FStreamableRenderResourceState. Can't build groom binding %s."),
-							*TargetMesh->GetPathName(), TargetMeshLODIndex, *GetPathName()))
-						{
-							return false;
-						}
+bool UGroomBindingAsset::TryInitializeContextForMeshes(FGroomBindingBuildContext& OutContext) const
+{
+	if (GetGroomBindingType() == EGroomBindingMeshType::SkeletalMesh)
+	{
+		ITargetPlatform* RunningPlatform = nullptr;
+		#if WITH_EDITOR
+		RunningPlatform = GetTargetPlatformManagerRef().GetRunningTargetPlatform();
+		check(RunningPlatform);
+		#endif
 
-						OutContext.MeshLODReferences.Add(MoveTemp(LODReference));
-					}
-				}
-				else
-				{
-					ensureMsgf(false,
-						TEXT("Missing render data for mesh %s LOD %i, which is claimed to be valid by FStreamableRenderResourceState. Can't build groom binding %s."),
-						*TargetMesh->GetPathName(), OutContext.TargetMeshMinLOD, *GetPathName());
+		if (USkeletalMesh* SourceMesh = GetSourceSkeletalMesh())
+		{
+			const bool bSucceed = GetSkeletalMeshLOD(this, SourceMesh, GetSourceMeshRequestedLOD(), true /*bSingleLOD*/, RunningPlatform, OutContext.SourceMeshLOD, OutContext.MeshLODReferences);
+			if (!bSucceed) { return false; }
+		}
 
-					return false;
-				}
-			}
-			else
-			{
-				UE_LOGFMT(LogHairStrands, Error, "Target mesh {0} has no render data. Can't build groom binding {1}.", TargetMesh->GetPathName(), GetPathName());
-				return false;
-			}
+		if (USkeletalMesh* TargetMesh = GetTargetSkeletalMesh())
+		{
+			const bool bSucceed = GetSkeletalMeshLOD(this, TargetMesh, GetTargetMeshRequestedMinLOD(), false/*bSingleLOD*/, RunningPlatform, OutContext.TargetMeshMinLOD, OutContext.MeshLODReferences);
+			if (!bSucceed) { return false; }
 		}
 	}
 	else
@@ -1505,67 +1466,30 @@ void UGroomBindingAsset::FinishCacheDerivedDatas(FGroomBindingBuildContext& Cont
 
 int32 UGroomBindingAsset::GetClampedSourceMeshLOD(const ITargetPlatform* TargetPlatform) const
 {
-	const USkeletalMesh* SourceMesh = GetSourceSkeletalMesh();
-	if (GetGroomBindingType() != EGroomBindingMeshType::SkeletalMesh 
-		|| !SourceMesh)
+	if (GetGroomBindingType() == EGroomBindingMeshType::SkeletalMesh)
 	{
-		return INDEX_NONE;
+		return GetClampedSkeletalMeshLOD(GetSourceSkeletalMesh(), GetSourceMeshRequestedLOD(), TargetPlatform);
 	}
-
-	// BestAvailableLOD is the best (i.e. lowest index) LOD that can ever be used by this mesh.
-	// E.g. in a cooked build, this is the best LOD that was cooked.
-	// 
-	// SourceMeshLOD is clamped to it here, so that if the mesh has a per-platform min LOD,
-	// the caller of UGroomBindingAsset::Build doesn't need to take that into account. They can
-	// simply request source mesh LOD 0 and if min LOD on the mesh is 3, for example, it will
-	// silently build the groom binding using LOD 3 as the source.
-#if WITH_EDITOR
-	check(TargetPlatform);
-	const int32 BestAvailableLOD = SourceMesh->GetPlatformMinLODIdx(TargetPlatform);
-#else
-	const bool bForceLowestLODIdx = true;
-	const int32 BestAvailableLOD = SourceMesh->GetMinLodIdx(bForceLowestLODIdx);
-#endif
-
-	return FMath::Max(GetSourceMeshRequestedLOD(), BestAvailableLOD);
+	return INDEX_NONE;
 }
 
 int32 UGroomBindingAsset::GetClampedTargetMeshMinLOD(const ITargetPlatform* TargetPlatform) const
 {
-	const USkeletalMesh* TargetMesh = GetTargetSkeletalMesh();
-	if (GetGroomBindingType() != EGroomBindingMeshType::SkeletalMesh
-		|| !TargetMesh)
+	if (GetGroomBindingType() == EGroomBindingMeshType::SkeletalMesh)
 	{
-		return INDEX_NONE;
+		return GetClampedSkeletalMeshLOD(GetTargetSkeletalMesh(), GetTargetMeshRequestedMinLOD(), TargetPlatform);
 	}
-
-	// BestAvailableLOD is the best (i.e. lowest index) LOD that can ever be used by this mesh.
-	// E.g. in a cooked build, this is the best LOD that was cooked.
-	// 
-	// TargetMeshMinLOD is clamped to it here, so that if the mesh has a per-platform min LOD,
-	// the caller of UGroomBindingAsset::Build doesn't need to take that into account. They can
-	// simply request target mesh LOD 0 and if min LOD on the mesh is 3, for example, it will build 
-	// the groom binding for LOD 3 onwards. LODs 0-2 would never be rendered anyway, so this is
-	// transparent to the caller.
-#if WITH_EDITOR
-	check(TargetPlatform);
-	const int32 BestAvailableLOD = TargetMesh->GetPlatformMinLODIdx(TargetPlatform);
-#else
-	const bool bForceLowestLODIdx = true;
-	const int32 BestAvailableLOD = TargetMesh->GetMinLodIdx(bForceLowestLODIdx);
-#endif
-
-	return FMath::Max(GetTargetMeshRequestedMinLOD(), BestAvailableLOD);
+	return INDEX_NONE;
 }
 
 #if WITH_EDITORONLY_DATA
 
-static TArray<FString> GetGroupDerivedDataKeys(const UGroomBindingAsset* In, const ITargetPlatform* TargetPlatform, int32 SourceMeshLOD, int32 TargetMeshMinLOD, bool bForceSkelMeshDataCachingForCookingOnly)
+static TArray<FString> GetGroupDerivedDataKeys(const UGroomBindingAsset* In, const ITargetPlatform* TargetPlatform, int32 SourceMeshLOD, int32 TargetMeshMinLOD)
 {
 	check(In);
 	check(TargetPlatform);
 
-	const FString KeySuffix = BuildDerivedDataKeySuffix(*In, TargetPlatform, SourceMeshLOD, TargetMeshMinLOD, bForceSkelMeshDataCachingForCookingOnly );
+	const FString KeySuffix = BuildDerivedDataKeySuffix(*In, TargetPlatform, SourceMeshLOD, TargetMeshMinLOD);
 	const uint32 GroupCount = In->GetGroupInfos().Num();
 
 	TArray<FString> Out;
@@ -1602,7 +1526,7 @@ void UGroomBindingAsset::BeginCacheForCookedPlatformData(const ITargetPlatform*
 	// 1. Build the key for each group
 	const int32 SourceMeshLOD = GetClampedSourceMeshLOD(TargetPlatform);
 	const int32 TargetMeshMinLOD = GetClampedTargetMeshMinLOD(TargetPlatform);
-	const TArray<FString> GroupDerivedDataKeys = GetGroupDerivedDataKeys(this, TargetPlatform, SourceMeshLOD, TargetMeshMinLOD, false /*bForceSkelMeshDataCachingForCookingOnly */);
+	const TArray<FString> GroupDerivedDataKeys = GetGroupDerivedDataKeys(this, TargetPlatform, SourceMeshLOD, TargetMeshMinLOD);
 
 	// 2. Find existing cached cooked data
 	UGroomBindingAsset::FCachedCookedPlatformData* TargetPlatformData = FindCachedCookedPlatformData(GroupDerivedDataKeys, CachedCookedPlatformDatas);
@@ -1610,7 +1534,7 @@ void UGroomBindingAsset::BeginCacheForCookedPlatformData(const ITargetPlatform*
 	// 3. If the target cooked data does not already exist, we build it
 	if (TargetPlatformData == nullptr && GetGroom() != nullptr)
 	{
-		auto BuildTargetPlatformData = [TargetPlatform, SourceMeshLOD, TargetMeshMinLOD](const TArray<FString>& InGroupDerivedDataKeys, UGroomBindingAsset* BindingAsset)
+		auto BuildTargetPlatformData = [TargetPlatform](const TArray<FString>& InGroupDerivedDataKeys, UGroomBindingAsset* BindingAsset, int32 InSourceMeshLOD, int32 InTargetMeshMinLOD)
 		{
 			// 3.a Build cooked derived data
 			const uint32 GroupCount = InGroupDerivedDataKeys.Num();
@@ -1626,8 +1550,8 @@ void UGroomBindingAsset::BeginCacheForCookedPlatformData(const ITargetPlatform*
 					NewTargetPlatformData->GroupDerivedDataKeys[GroupIndex], 
 					bGroupValid, 
 					TargetPlatform, 
-					SourceMeshLOD, 
-					TargetMeshMinLOD, 
+					InSourceMeshLOD, 
+					InTargetMeshMinLOD, 
 					NewTargetPlatformData->GroupPlatformDatas[GroupIndex]);
 
 				if (!bGroupValid)
@@ -1654,17 +1578,17 @@ void UGroomBindingAsset::BeginCacheForCookedPlatformData(const ITargetPlatform*
 		};
 
 		// 3.1 Build binding data for the target platform
-		BuildTargetPlatformData(GroupDerivedDataKeys, this);
+		BuildTargetPlatformData(GroupDerivedDataKeys, this, SourceMeshLOD, TargetMeshMinLOD);
 
 		// 3.2 Check if the derived data keys are unchanged. If they have change, recompute the target platform binding data again.
 		// This can happen when a skel. mesh asset hasn't be resaved correct. While being cached, the skel. mesh data will update/transform 
 		// its DDC key, causing the binding data to be cached as the incorrect DDC key.
 		const int32 UpdatedSourceMeshLOD = GetClampedSourceMeshLOD(TargetPlatform);
 		const int32 UpdatedTargetMeshMinLOD = GetClampedTargetMeshMinLOD(TargetPlatform);
-		const TArray<FString> UpdatedGroupDerivedDataKeys = GetGroupDerivedDataKeys(this, TargetPlatform, UpdatedSourceMeshLOD, UpdatedTargetMeshMinLOD, false /*bForceSkelMeshDataCachingForCookingOnly */);
+		const TArray<FString> UpdatedGroupDerivedDataKeys = GetGroupDerivedDataKeys(this, TargetPlatform, UpdatedSourceMeshLOD, UpdatedTargetMeshMinLOD);
 		if (UpdatedGroupDerivedDataKeys != GroupDerivedDataKeys)
 		{
-			BuildTargetPlatformData(UpdatedGroupDerivedDataKeys, this);
+			BuildTargetPlatformData(UpdatedGroupDerivedDataKeys, this, UpdatedSourceMeshLOD, UpdatedTargetMeshMinLOD);
 		}
 	}
 }
