Fix incorrect LOD count introduced in CL 37389364.

diff --git a/Engine/Plugins/Runtime/HairStrands/Source/HairStrandsCore/Private/GroomBindingAsset.cpp b/Engine/Plugins/Runtime/HairStrands/Source/HairStrandsCore/Private/GroomBindingAsset.cpp
--- a/Engine/Plugins/Runtime/HairStrands/Source/HairStrandsCore/Private/GroomBindingAsset.cpp
+++ b/Engine/Plugins/Runtime/HairStrands/Source/HairStrandsCore/Private/GroomBindingAsset.cpp
@@ -1150,7 +1150,7 @@ static bool GetSkeletalMeshLOD(
 	{
 		if (RenderData->LODRenderData.IsValidIndex(OutMeshMinLOD))
 		{
-			const int32 LODCount = bRequestSingleLOD ? 1u : SkeletalMesh->GetLODNum();
+			const int32 LODCount = bRequestSingleLOD ? FMath::Min(OutMeshMinLOD+1, SkeletalMesh->GetLODNum()) : SkeletalMesh->GetLODNum();
 			for (int32 MeshLODIndex = OutMeshMinLOD; MeshLODIndex < LODCount; ++MeshLODIndex)
 			{
 				TRefCountPtr<FSkeletalMeshLODRenderData> LODReference;
