[AsyncGroomBindings] Fix potential crash when dependencies of groom bindings are getting GCed

Submitted on behalf of danny.couture

diff --git a/Engine/Plugins/Runtime/HairStrands/Source/HairStrandsCore/Private/GroomBindingCompiler.cpp b/Engine/Plugins/Runtime/HairStrands/Source/HairStrandsCore/Private/GroomBindingCompiler.cpp
--- a/Engine/Plugins/Runtime/HairStrands/Source/HairStrandsCore/Private/GroomBindingCompiler.cpp
+++ b/Engine/Plugins/Runtime/HairStrands/Source/HairStrandsCore/Private/GroomBindingCompiler.cpp
@@ -103,42 +103,58 @@ void FGroomBindingCompilingManager::DetachDependencies(UGroomBindingAsset* Groom
 
 void FGroomBindingCompilingManager::OnPostReachabilityAnalysis()
 {
-	if (GetNumRemainingJobs())
-	{
-		TRACE_CPUPROFILER_EVENT_SCOPE(FGroomBindingCompilingManager::CancelUnreachableGroomBindings);
+	TRACE_CPUPROFILER_EVENT_SCOPE(FGroomBindingCompilingManager::CancelUnreachableGroomBindings);
 
-		TArray<UGroomBindingAsset*> PendingAssets;
-		PendingAssets.Reserve(GetNumRemainingJobs());
+	TArray<UGroomBindingAsset*> PendingAssets;
+	PendingAssets.Reserve(GetNumRemainingJobs());
 		
-		// Compilation has not started yet... just remove it from our pending list
-		for (auto Iterator = GroomBindingWithPendingDependencies.CreateIterator(); Iterator; ++Iterator)
+	// Compilation has not started yet... just remove it from our pending list
+	for (auto Iterator = GroomBindingWithPendingDependencies.CreateIterator(); Iterator; ++Iterator)
+	{
+		UGroomBindingAsset* GroomBinding = *Iterator;
+		if (GroomBinding->IsUnreachable())
 		{
-			UGroomBindingAsset* GroomBinding = *Iterator;
-			if (GroomBinding->IsUnreachable())
+			Iterator.RemoveCurrent();
+		}
+	}
+
+	for (auto Iterator = RegisteredGroomBindingAssets.CreateIterator(); Iterator; ++Iterator)
+	{
+		UGroomBindingAsset* GroomBinding = *Iterator;
+		if (GroomBinding->IsUnreachable())
+		{
+			UE_LOG(LogHairStrands, Verbose, TEXT("Cancelling groom binding %s async compilation because it's being garbage collected"), *GroomBinding->GetName());
+
+			if (GroomBinding->TryCancelAsyncTasks())
 			{
 				Iterator.RemoveCurrent();
+				DetachDependencies(GroomBinding);
+			}
+			else
+			{
+				PendingAssets.Add(GroomBinding);
 			}
 		}
+	}
 
-		for (auto Iterator = RegisteredGroomBindingAssets.CreateIterator(); Iterator; ++Iterator)
+	for (auto Iterator = RegisteredSkeletalMeshes.CreateIterator(); Iterator; ++Iterator)
+	{
+		if (Iterator->Key->IsUnreachable() || Iterator->Value->IsUnreachable())
 		{
-			UGroomBindingAsset* GroomBinding = *Iterator;
-			if (GroomBinding->IsUnreachable())
-			{
-				UE_LOG(LogHairStrands, Verbose, TEXT("Cancelling groom binding %s async compilation because it's being garbage collected"), *GroomBinding->GetName());
+			Iterator.RemoveCurrent();
+		}
+	}
 
-				if (GroomBinding->TryCancelAsyncTasks())
-				{
-					Iterator.RemoveCurrent();
-					DetachDependencies(GroomBinding);
-				}
-				else
-				{
-					PendingAssets.Add(GroomBinding);
-				}
-			}
+	for (auto Iterator = RegisteredGroomAssets.CreateIterator(); Iterator; ++Iterator)
+	{
+		if (Iterator->Key->IsUnreachable() || Iterator->Value->IsUnreachable())
+		{
+			Iterator.RemoveCurrent();
 		}
+	}
 
+	if (PendingAssets.Num())
+	{
 		FinishCompilation(PendingAssets);
 	}
 }
