Fix groom binding cooking error.

The RequestedMeshLOD value queried from a skeletal mesh, has inconsistent value between the moment the data is built/cached, and the moment it is queried back. This causes the DDC key to mismatch and prevent to retrieve the data during the groom binding serialization. This issue happens only when cooking certain platform.

To avoid this issue, the build platform data are now indexed by a DDC key with default RequestedMeshLOD (=0) to avoid any mismatch betwenn caching and quering.

diff --git a/Engine/Plugins/Runtime/HairStrands/Source/HairStrandsCore/Private/GroomBindingAsset.cpp b/Engine/Plugins/Runtime/HairStrands/Source/HairStrandsCore/Private/GroomBindingAsset.cpp
--- a/Engine/Plugins/Runtime/HairStrands/Source/HairStrandsCore/Private/GroomBindingAsset.cpp
+++ b/Engine/Plugins/Runtime/HairStrands/Source/HairStrandsCore/Private/GroomBindingAsset.cpp
@@ -1501,11 +1501,11 @@ static TArray<FString> GetGroupDerivedDataKeys(const UGroomBindingAsset* In, con
 	return Out;
 }
 
-static UGroomBindingAsset::FCachedCookedPlatformData* FindCachedCookedPlatformData(const TArray<FString>& InGroupKeys, TArray<UGroomBindingAsset::FCachedCookedPlatformData*>& InCachedCookedData)
+static UGroomBindingAsset::FCachedCookedPlatformData* FindCachedCookedPlatformData(const TArray<FString>& InGroupKeysWithoutMeshLOD, TArray<UGroomBindingAsset::FCachedCookedPlatformData*>& InCachedCookedData)
 {
 	for (UGroomBindingAsset::FCachedCookedPlatformData* CookedPlatformData : InCachedCookedData)
 	{
-		if (CookedPlatformData->GroupDerivedDataKeys == InGroupKeys)
+		if (CookedPlatformData->GroupDerivedDataKeys_Query == InGroupKeysWithoutMeshLOD)
 		{
 			return CookedPlatformData;
 		}
@@ -1524,22 +1524,23 @@ void UGroomBindingAsset::BeginCacheForCookedPlatformData(const ITargetPlatform*
 	Super::BeginCacheForCookedPlatformData(TargetPlatform);
 
 	// 1. Build the key for each group
-	const int32 SourceMeshLOD = GetClampedSourceMeshLOD(TargetPlatform);
-	const int32 TargetMeshMinLOD = GetClampedTargetMeshMinLOD(TargetPlatform);
-	const TArray<FString> GroupDerivedDataKeys = GetGroupDerivedDataKeys(this, TargetPlatform, SourceMeshLOD, TargetMeshMinLOD);
+	const TArray<FString> GroupDerivedDataKeys_Query = GetGroupDerivedDataKeys(this, TargetPlatform, 0/*SourceMeshLOD*/, 0/*TargetMeshMinLOD*/);
 
 	// 2. Find existing cached cooked data
-	UGroomBindingAsset::FCachedCookedPlatformData* TargetPlatformData = FindCachedCookedPlatformData(GroupDerivedDataKeys, CachedCookedPlatformDatas);
+	UGroomBindingAsset::FCachedCookedPlatformData* TargetPlatformData = FindCachedCookedPlatformData(GroupDerivedDataKeys_Query, CachedCookedPlatformDatas);
 
 	// 3. If the target cooked data does not already exist, we build it
 	if (TargetPlatformData == nullptr && GetGroom() != nullptr)
 	{
-		auto BuildTargetPlatformData = [TargetPlatform](const TArray<FString>& InGroupDerivedDataKeys, UGroomBindingAsset* BindingAsset, int32 InSourceMeshLOD, int32 InTargetMeshMinLOD)
+		auto BuildTargetPlatformData = [TargetPlatform](const TArray<FString>& InGroupDerivedDataKeys_Query, const TArray<FString>& InGroupDerivedDataKeys_Cache, UGroomBindingAsset* BindingAsset, int32 InSourceMeshLOD, int32 InTargetMeshMinLOD)
 		{
 			// 3.a Build cooked derived data
-			const uint32 GroupCount = InGroupDerivedDataKeys.Num();
+			const uint32 GroupCount = InGroupDerivedDataKeys_Cache.Num();
 			UGroomBindingAsset::FCachedCookedPlatformData* NewTargetPlatformData = new FCachedCookedPlatformData();
-			NewTargetPlatformData->GroupDerivedDataKeys = InGroupDerivedDataKeys;
+			NewTargetPlatformData->GroupDerivedDataKeys_Query = InGroupDerivedDataKeys_Query;
+			NewTargetPlatformData->GroupDerivedDataKeys = InGroupDerivedDataKeys_Cache;
+			NewTargetPlatformData->SourceMeshLOD = InSourceMeshLOD;
+			NewTargetPlatformData->TargetMeshMinLOD = InTargetMeshMinLOD;
 			NewTargetPlatformData->GroupPlatformDatas.SetNum(GroupCount);
 			for (uint32 GroupIndex = 0; GroupIndex < GroupCount; ++GroupIndex)
 			{
@@ -1578,17 +1579,23 @@ void UGroomBindingAsset::BeginCacheForCookedPlatformData(const ITargetPlatform*
 		};
 
 		// 3.1 Build binding data for the target platform
-		BuildTargetPlatformData(GroupDerivedDataKeys, this, SourceMeshLOD, TargetMeshMinLOD);
+		{
+			const int32 SourceMeshLOD = GetClampedSourceMeshLOD(TargetPlatform);
+			const int32 TargetMeshMinLOD = GetClampedTargetMeshMinLOD(TargetPlatform);
+			const TArray<FString> GroupDerivedDataKeys_Cache = GetGroupDerivedDataKeys(this, TargetPlatform, SourceMeshLOD, TargetMeshMinLOD);
+			BuildTargetPlatformData(GroupDerivedDataKeys_Query, GroupDerivedDataKeys_Cache, this, SourceMeshLOD, TargetMeshMinLOD);
+		}
 
 		// 3.2 Check if the derived data keys are unchanged. If they have change, recompute the target platform binding data again.
 		// This can happen when a skel. mesh asset hasn't be resaved correct. While being cached, the skel. mesh data will update/transform 
 		// its DDC key, causing the binding data to be cached as the incorrect DDC key.
-		const int32 UpdatedSourceMeshLOD = GetClampedSourceMeshLOD(TargetPlatform);
-		const int32 UpdatedTargetMeshMinLOD = GetClampedTargetMeshMinLOD(TargetPlatform);
-		const TArray<FString> UpdatedGroupDerivedDataKeys = GetGroupDerivedDataKeys(this, TargetPlatform, UpdatedSourceMeshLOD, UpdatedTargetMeshMinLOD);
-		if (UpdatedGroupDerivedDataKeys != GroupDerivedDataKeys)
+		const TArray<FString> UpdatedGroupDerivedDataKeys_Query = GetGroupDerivedDataKeys(this, TargetPlatform, 0/*SourceMeshLOD*/, 0/*TargetMeshMinLOD*/);
+		if (UpdatedGroupDerivedDataKeys_Query != GroupDerivedDataKeys_Query)
 		{
-			BuildTargetPlatformData(UpdatedGroupDerivedDataKeys, this, UpdatedSourceMeshLOD, UpdatedTargetMeshMinLOD);
+			const int32 UpdatedSourceMeshLOD = GetClampedSourceMeshLOD(TargetPlatform);
+			const int32 UpdatedTargetMeshMinLOD = GetClampedTargetMeshMinLOD(TargetPlatform);
+			const TArray<FString> UpdatedGroupDerivedDataKeys_Cache = GetGroupDerivedDataKeys(this, TargetPlatform, UpdatedSourceMeshLOD, UpdatedTargetMeshMinLOD);
+			BuildTargetPlatformData(UpdatedGroupDerivedDataKeys_Query, UpdatedGroupDerivedDataKeys_Cache, this, UpdatedSourceMeshLOD, UpdatedTargetMeshMinLOD);
 		}
 	}
 }
@@ -1596,12 +1603,10 @@ void UGroomBindingAsset::BeginCacheForCookedPlatformData(const ITargetPlatform*
 TArray<UGroomBindingAsset::FHairGroupPlatformData>* UGroomBindingAsset::GetCachedCookedPlatformData(const ITargetPlatform* TargetPlatform)
 {
 	// 1. Build the key for each group
-	const int32 SourceMeshLOD = GetClampedSourceMeshLOD(TargetPlatform);
-	const int32 TargetMeshMinLOD = GetClampedTargetMeshMinLOD(TargetPlatform);
-	const TArray<FString> GroupDerivedDataKeys = GetGroupDerivedDataKeys(this, TargetPlatform, SourceMeshLOD, TargetMeshMinLOD);
+	const TArray<FString> GroupDerivedDataKeys_Query = GetGroupDerivedDataKeys(this, TargetPlatform, 0/*SourceMeshLOD*/, 0/*TargetMeshMinLOD*/);
 
 	// 2. Find existing cached cooked data
-	if (UGroomBindingAsset::FCachedCookedPlatformData* CachedCookedPlatformData = FindCachedCookedPlatformData(GroupDerivedDataKeys, CachedCookedPlatformDatas))
+	if (UGroomBindingAsset::FCachedCookedPlatformData* CachedCookedPlatformData = FindCachedCookedPlatformData(GroupDerivedDataKeys_Query, CachedCookedPlatformDatas))
 	{
 		return &CachedCookedPlatformData->GroupPlatformDatas;
 	}
diff --git a/Engine/Plugins/Runtime/HairStrands/Source/HairStrandsCore/Public/GroomBindingAsset.h b/Engine/Plugins/Runtime/HairStrands/Source/HairStrandsCore/Public/GroomBindingAsset.h
--- a/Engine/Plugins/Runtime/HairStrands/Source/HairStrandsCore/Public/GroomBindingAsset.h
+++ b/Engine/Plugins/Runtime/HairStrands/Source/HairStrandsCore/Public/GroomBindingAsset.h
@@ -524,8 +524,9 @@ private:
 
 public:
 #if WITH_EDITORONLY_DATA
-	virtual void BeginCacheForCookedPlatformData(const ITargetPlatform* TargetPlatform);
-	virtual void ClearAllCachedCookedPlatformData();
+	virtual void BeginCacheForCookedPlatformData(const ITargetPlatform* TargetPlatform) override;
+	virtual void ClearAllCachedCookedPlatformData() override;
+
 	TArray<FHairGroupPlatformData>* GetCachedCookedPlatformData(const ITargetPlatform* TargetPlatform);
 
 	void InvalidateBinding();
@@ -533,8 +534,14 @@ public:
 
 	struct FCachedCookedPlatformData
 	{
+		// Actual DDC key and platform data
 		TArray<FString> GroupDerivedDataKeys;
 		TArray<FHairGroupPlatformData> GroupPlatformDatas;
+
+		// DDC key without source/target MeshLOD and requested MeshLOD value
+		TArray<FString> GroupDerivedDataKeys_Query;
+		int32 SourceMeshLOD = INDEX_NONE;
+		int32 TargetMeshMinLOD = INDEX_NONE;
 	};
 
 	TArray<FCachedCookedPlatformData*> CachedCookedPlatformDatas;
