Fixed groom cache did not import because its validation check was not updated for a new groom asset feature.

diff --git a/Engine/Plugins/Runtime/HairStrands/Source/HairStrandsEditor/HairStrandsEditor.Build.cs b/Engine/Plugins/Runtime/HairStrands/Source/HairStrandsEditor/HairStrandsEditor.Build.cs
--- a/Engine/Plugins/Runtime/HairStrands/Source/HairStrandsEditor/HairStrandsEditor.Build.cs
+++ b/Engine/Plugins/Runtime/HairStrands/Source/HairStrandsEditor/HairStrandsEditor.Build.cs
@@ -38,6 +38,7 @@ namespace UnrealBuildTool.Rules
 					"Sequencer",
 					"HairCardGeneratorFramework",
 					"CommonMenuExtensions",
+					"RenderCore"
 				});
 
 			PrivateIncludePathModuleNames.AddRange(
diff --git a/Engine/Plugins/Runtime/HairStrands/Source/HairStrandsEditor/Private/GroomCacheImporter.cpp b/Engine/Plugins/Runtime/HairStrands/Source/HairStrandsEditor/Private/GroomCacheImporter.cpp
--- a/Engine/Plugins/Runtime/HairStrands/Source/HairStrandsEditor/Private/GroomCacheImporter.cpp
+++ b/Engine/Plugins/Runtime/HairStrands/Source/HairStrandsEditor/Private/GroomCacheImporter.cpp
@@ -12,6 +12,7 @@
 #include "Misc/ScopedSlowTask.h"
 #include "ObjectTools.h"
 #include "PackageTools.h"
+#include "RenderUtils.h"
 
 #define LOCTEXT_NAMESPACE "GroomCacheImporter"
 
@@ -174,8 +175,16 @@ TArray<UGroomCache*> FGroomCacheImporter::ImportGroomCache(const FString& Source
 				{
 					if (bImportStrandsCache)
 					{
-						if (HairGroupsData[GroupIndex].Strands.GetNumCurves() != GroupPlatformData[GroupIndex].Strands.BulkData.GetNumCurves() ||
-							HairGroupsData[GroupIndex].Strands.GetNumPoints() != GroupPlatformData[GroupIndex].Strands.BulkData.GetNumPoints())
+						// When UsesTriangleStrips is enabled, we add an extra control point at the end of each curve in the groom asset.
+						// Since the groom cache needs to be independent of this settings (since the data is serialized directly into the asset
+						// and not rebuilt from a description), it does not contains these extra control points.
+						// We account for them for validation here
+						const uint32 ExtraControlPointCount = GetHairStrandsUsesTriangleStrips() ? HairGroupsData[GroupIndex].Strands.GetNumCurves() : 0u;
+
+						const bool bCurveCountMatch = HairGroupsData[GroupIndex].Strands.GetNumCurves() == GroupPlatformData[GroupIndex].Strands.BulkData.GetNumCurves();
+						const bool bPointCountMatch = HairGroupsData[GroupIndex].Strands.GetNumPoints() + ExtraControlPointCount == GroupPlatformData[GroupIndex].Strands.BulkData.GetNumPoints();
+
+						if (!bCurveCountMatch || !bPointCountMatch)
 						{
 							bSuccess = false;
 							UE_LOG(LogGroomCacheImporter, Warning, TEXT("GroomCache frame %d does not have the same number of curves (%u) \
@@ -193,8 +202,12 @@ TArray<UGroomCache*> FGroomCacheImporter::ImportGroomCache(const FString& Source
 							bGuidesOnly = true;
 						}
 
-						if (HairGroupsData[GroupIndex].Guides.GetNumCurves() != GroupPlatformData[GroupIndex].Guides.BulkData.GetNumCurves() ||
-							HairGroupsData[GroupIndex].Guides.GetNumPoints() != GroupPlatformData[GroupIndex].Guides.BulkData.GetNumPoints())
+						const uint32 ExtraControlPointCount = GetHairStrandsUsesTriangleStrips() ? HairGroupsData[GroupIndex].Guides.GetNumCurves() : 0u;
+
+						const bool bCurveCountMatch = HairGroupsData[GroupIndex].Guides.GetNumCurves() == GroupPlatformData[GroupIndex].Guides.BulkData.GetNumCurves();
+						const bool bPointCountMatch = HairGroupsData[GroupIndex].Guides.GetNumPoints() + ExtraControlPointCount == GroupPlatformData[GroupIndex].Guides.BulkData.GetNumPoints();
+
+						if (!bCurveCountMatch || !bPointCountMatch)
 						{
 							bSuccess = false;
 							UE_LOG(LogGroomCacheImporter, Warning, TEXT("GroomCache frame %d does not have the same number of curves (%u) \
