Ray Tracing: Hair strands proxy was missing HasRayTracingRepresentation

This caused hair to be excluded from ray tracing

#rnx

diff --git a/Engine/Plugins/Runtime/HairStrands/Source/HairStrandsCore/Private/GroomComponent.cpp b/Engine/Plugins/Runtime/HairStrands/Source/HairStrandsCore/Private/GroomComponent.cpp
--- a/Engine/Plugins/Runtime/HairStrands/Source/HairStrandsCore/Private/GroomComponent.cpp
+++ b/Engine/Plugins/Runtime/HairStrands/Source/HairStrandsCore/Private/GroomComponent.cpp
@@ -641,8 +641,9 @@ public:
 	}
 
 #if RHI_RAYTRACING
-	virtual bool IsRayTracingRelevant() const { return true; }
-	virtual bool IsRayTracingStaticRelevant() const { return false; }
+	virtual bool IsRayTracingRelevant() const override { return true; }
+	virtual bool IsRayTracingStaticRelevant() const override { return false; }
+	virtual bool HasRayTracingRepresentation() const override { return true; }
 
 	virtual void GetDynamicRayTracingInstances(FRayTracingInstanceCollector& Collector) override
 	{
