Fix crash with groom bone transform not normalized

diff --git a/Engine/Plugins/Runtime/HairStrands/Source/HairStrandsCore/Private/Niagara/NiagaraDataInterfaceHairStrands.cpp b/Engine/Plugins/Runtime/HairStrands/Source/HairStrandsCore/Private/Niagara/NiagaraDataInterfaceHairStrands.cpp
--- a/Engine/Plugins/Runtime/HairStrands/Source/HairStrandsCore/Private/Niagara/NiagaraDataInterfaceHairStrands.cpp
+++ b/Engine/Plugins/Runtime/HairStrands/Source/HairStrandsCore/Private/Niagara/NiagaraDataInterfaceHairStrands.cpp
@@ -553,6 +553,7 @@ void FNDIHairStrandsData::Update(
 				const FHairSimulationSettings& SimulationSettings = Interface->SourceComponent->SimulationSettings;
 				LocalSimulation = SimulationSettings.SimulationSetup.bLocalSimulation;
 				Interface->SourceComponent->BuildSimulationTransform(BoneTransform);
+				BoneTransform.NormalizeRotation();
 
 				// Convert to double for LWC
 				FMatrix44d BoneTransformDouble = BoneTransform.ToMatrixWithScale();
@@ -592,7 +593,8 @@ void FNDIHairStrandsData::Update(
 				BoneTransformDouble = BoneTransformDouble * WorldTransformDouble.Inverse();
 				const FMatrix44d WorldTransformFloat = BoneTransformDouble;
 				BoneTransform = FTransform(WorldTransformFloat);
-
+				BoneTransform.NormalizeRotation();
+				
 				if (SimulationSettings.bOverrideSettings)
 				{
 					GravityVector = SimulationSettings.ExternalForces.GravityVector;
