Fix overflow allocation when creating binding asset for very large groom.

diff --git a/Engine/Plugins/Runtime/HairStrands/Source/HairStrandsCore/Private/GroomBindingBuilder.cpp b/Engine/Plugins/Runtime/HairStrands/Source/HairStrandsCore/Private/GroomBindingBuilder.cpp
--- a/Engine/Plugins/Runtime/HairStrands/Source/HairStrandsCore/Private/GroomBindingBuilder.cpp
+++ b/Engine/Plugins/Runtime/HairStrands/Source/HairStrandsCore/Private/GroomBindingBuilder.cpp
@@ -1006,18 +1006,29 @@ namespace GroomBinding_RootProjection
 		};
 		typedef TArray<const FCell*> FCells;
 
-		FTriangleGrid(const FVector3f& InMinBound, const FVector3f& InMaxBound, float InVoxelWorldSize)
+		FTriangleGrid(const FVector3f& InMinBound, const FVector3f& InMaxBound)
 		{
+			static const int32 CellsLimitCount = 1 << 27; // 134M
+			float InVoxelWorldSize = 2; //cm
+
 			MinBound = InMinBound;
 			MaxBound = InMaxBound;
 
 			// Compute the voxel volume resolution, and snap the max bound to the voxel grid
+			// Increase voxel size until the total allocation fit into CellsLimitCount budget
 			GridResolution = FIntVector::ZeroValue;
-			FVector3f VoxelResolutionF = (MaxBound - MinBound) / InVoxelWorldSize;
-			GridResolution = FIntVector(FMath::CeilToInt(VoxelResolutionF.X), FMath::CeilToInt(VoxelResolutionF.Y), FMath::CeilToInt(VoxelResolutionF.Z));
+			InVoxelWorldSize *= 0.5;
+			do
+			{
+				InVoxelWorldSize *= 2;
+				FVector3f VoxelResolutionF = (MaxBound - MinBound) / InVoxelWorldSize;
+				GridResolution = FIntVector(FMath::CeilToInt(VoxelResolutionF.X), FMath::CeilToInt(VoxelResolutionF.Y), FMath::CeilToInt(VoxelResolutionF.Z));
+			} while (static_cast<int64>(GridResolution.X) * GridResolution.Y * GridResolution.Z > CellsLimitCount);
+
 			MaxBound = MinBound + FVector3f(GridResolution) * InVoxelWorldSize;
 
-			Cells.SetNum(GridResolution.X * GridResolution.Y * GridResolution.Z);
+			const uint32 CellCount = GridResolution.X * GridResolution.Y * GridResolution.Z;
+			Cells.SetNum(CellCount);
 		}
 
 		FORCEINLINE bool IsValid(const FIntVector& P) const
@@ -1282,7 +1293,6 @@ namespace GroomBinding_RootProjection
 		// 2. Project root for each mesh LOD
 		const uint32 CurveCount = InRootPositions.GetNumCurves();
 		const uint32 ChannelIndex = 0;
-		const float VoxelWorldSize = 2; //cm
 		const int32 MeshLODCount = InMeshData->GetNumLODs();
 		check(MeshLODCount == OutRootData.Num());
 
@@ -1416,7 +1426,7 @@ namespace GroomBinding_RootProjection
 				}
 			}
 
-			FTriangleGrid Grid(GridMin, GridMax, VoxelWorldSize);
+			FTriangleGrid Grid(GridMin, GridMax);
 			bool bIsGridPopulated = false;
 			for (uint32 SectionIt = 0; SectionIt < SectionCount; ++SectionIt)
 			{
