Change hair strands data type to improve data flow clarity and avoiding incorrect API call during build time.

diff --git a/Engine/Plugins/Importers/USDImporter/Source/USDSchemas/Private/USDGroomTranslator.cpp b/Engine/Plugins/Importers/USDImporter/Source/USDSchemas/Private/USDGroomTranslator.cpp
--- a/Engine/Plugins/Importers/USDImporter/Source/USDSchemas/Private/USDGroomTranslator.cpp
+++ b/Engine/Plugins/Importers/USDImporter/Source/USDSchemas/Private/USDGroomTranslator.cpp
@@ -379,27 +379,26 @@ protected:
 					const uint32 GroupCount = HairDescriptionGroups.HairGroups.Num();
 
 					TArray<FHairGroupInfoWithVisibility> HairGroupsInfo = GroomAsset->GetHairGroupsInfo();
-					TArray<FHairDescriptionGroup> HairGroupsData;
-					HairGroupsData.SetNum(GroupCount);
+					TArray<FGroomCacheInputData> CacheInputData;
+					CacheInputData.SetNum(GroupCount);
 					for (uint32 GroupIndex = 0; GroupIndex < GroupCount; ++GroupIndex)
 					{
 						const FHairDescriptionGroup& HairGroup = HairDescriptionGroups.HairGroups[GroupIndex];
-						FHairDescriptionGroup& HairGroupData = HairGroupsData[GroupIndex];
 						FGroomBuilder::BuildData(
 							HairGroup,
 							GroomAsset->GetHairGroupsInterpolation()[GroupIndex],
 							HairGroupsInfo[GroupIndex],
-							HairGroupData.Strands,
-							HairGroupData.Guides
+							CacheInputData[GroupIndex].Strands,
+							CacheInputData[GroupIndex].Guides
 						);
 					}
 
 					// Validate that the GroomCache has the same topology as the static groom
-					if (HairGroupsData.Num() == GroomHairGroupsData.Num())
+					if (CacheInputData.Num() == GroomHairGroupsData.Num())
 					{
 						for (uint32 GroupIndex = 0; GroupIndex < GroupCount; ++GroupIndex)
 						{
-							if (HairGroupsData[GroupIndex].Strands.GetNumPoints() != GroomHairGroupsData[GroupIndex].Strands.BulkData.GetNumPoints())
+							if (CacheInputData[GroupIndex].Strands.GetNumPoints() != GroomHairGroupsData[GroupIndex].Strands.BulkData.GetNumPoints())
 							{
 								bSuccess = false;
 								UE_LOG(
@@ -408,7 +407,7 @@ protected:
 									TEXT("GroomCache frame %d does not have the same number of vertices as the static groom (%u instead of %u). "
 										 "Aborting GroomCache import."),
 									FrameIndex,
-									HairGroupsData[GroupIndex].Strands.GetNumPoints(),
+									CacheInputData[GroupIndex].Strands.GetNumPoints(),
 									GroomHairGroupsData[GroupIndex].Strands.BulkData.GetNumPoints()
 								);
 								break;
@@ -423,7 +422,7 @@ protected:
 							Warning,
 							TEXT("GroomCache does not have the same number of groups as the static groom (%d instead of %d). Aborting GroomCache "
 								 "import."),
-							HairGroupsData.Num(),
+							CacheInputData.Num(),
 							GroomHairGroupsData.Num()
 						);
 					}
@@ -434,7 +433,7 @@ protected:
 					}
 
 					// The HairGroupData is converted into animated groom data by the GroomCacheProcessor
-					GroomCacheProcessor.AddGroomSample(MoveTemp(HairGroupsData));
+					GroomCacheProcessor.AddGroomSample(MoveTemp(CacheInputData));
 				}
 
 				if (bSuccess)
diff --git a/Engine/Plugins/Runtime/HairStrands/Source/HairStrandsCore/Private/GroomAsset.cpp b/Engine/Plugins/Runtime/HairStrands/Source/HairStrandsCore/Private/GroomAsset.cpp
--- a/Engine/Plugins/Runtime/HairStrands/Source/HairStrandsCore/Private/GroomAsset.cpp
+++ b/Engine/Plugins/Runtime/HairStrands/Source/HairStrandsCore/Private/GroomAsset.cpp
@@ -2842,8 +2842,8 @@ static bool InternalImportCardGeometry(
 	// * Voxelize hair strands data (all group), to transfer group index from strands to cards
 	//FHairGroupInfo DummyInfo;
 
-	FHairStrandsDatas TempHairGuidesData = InHairDescriptionGroups.HairGroups[InDesc->GroupIndex].Guides;
-	FHairStrandsDatas TempHairStrandsData = InHairDescriptionGroups.HairGroups[InDesc->GroupIndex].Strands;
+	FHairStrandsDatas TempHairGuidesData;
+	FHairStrandsDatas TempHairStrandsData;
 	FGroomBuilder::BuildData(
 		InHairDescriptionGroups.HairGroups[InDesc->GroupIndex], 
 		InHairInterpolationGroups[InDesc->GroupIndex],
@@ -2986,8 +2986,7 @@ bool UGroomAsset::BuildHairGroup_Cards(uint32 GroupIndex)
 
 				FHairGroupInfo DummyGroupInfo;
 
-				// 1. Build data & bulk data for the cards guides
-				FGroomBuilder::BuildData(LODGuidesData);
+				// 1. Build bulk data for the cards guides
 				FGroomBuilder::BuildBulkData(DummyGroupInfo, LODGuidesData, LOD.GuideBulkData, false /*bAllowCompression*/);
 
 				// 2. (Re)Build data for the sim guides (since there are transient/no-cached)
diff --git a/Engine/Plugins/Runtime/HairStrands/Source/HairStrandsCore/Private/GroomBuilder.cpp b/Engine/Plugins/Runtime/HairStrands/Source/HairStrandsCore/Private/GroomBuilder.cpp
--- a/Engine/Plugins/Runtime/HairStrands/Source/HairStrandsCore/Private/GroomBuilder.cpp
+++ b/Engine/Plugins/Runtime/HairStrands/Source/HairStrandsCore/Private/GroomBuilder.cpp
@@ -51,7 +51,7 @@ bool DoesHairStrandsSupportCompressedPosition();
 
 FString FGroomBuilder::GetVersion()
 {
-	return TEXT("v16e");
+	return TEXT("v17_debug0");
 }
 
 namespace GroomBuilder_Voxelization
@@ -97,32 +97,32 @@ namespace FHairStrandsDecimation
 
 namespace HairStrandsBuilder
 {
-	/** Build the internal points and curves data */
-	void BuildInternalData(FHairStrandsDatas& HairStrands)
+	/** Build FHairStrandsDatas from a FHairStrandsRawDatas */
+	FHairStrandsDatas BuildInternalData(const FHairStrandsCommonDatas& In)
 	{
 		TRACE_CPUPROFILER_EVENT_SCOPE(HairStrandsBuilder::BuildInternalData);
+		FHairStrandsDatas Out;
 
-		FHairStrandsCurves& Curves = HairStrands.StrandsCurves;
-		FHairStrandsPoints& Points = HairStrands.StrandsPoints;
+		Out.StrandsCurves = In.StrandsCurves;
+		Out.StrandsPoints = In.StrandsPoints;
+		Out.BoundingBox.Min = {  FLT_MAX,  FLT_MAX ,  FLT_MAX };
+		Out.BoundingBox.Max = { -FLT_MAX, -FLT_MAX , -FLT_MAX };
+		Out.BoundingBox.IsValid = 0;
 
-		HairStrands.BoundingBox.Min = {  FLT_MAX,  FLT_MAX ,  FLT_MAX };
-		HairStrands.BoundingBox.Max = { -FLT_MAX, -FLT_MAX , -FLT_MAX };
-		HairStrands.BoundingBox.IsValid = 0;
-
-		if (HairStrands.GetNumCurves() > 0 && HairStrands.GetNumPoints() > 0)
+		if (Out.GetNumCurves() > 0 && Out.GetNumPoints() > 0)
 		{
-			TArray<FVector3f>::TIterator PositionIterator = Points.PointsPosition.CreateIterator();
-			TArray<float>::TIterator RadiusIterator = Points.PointsRadius.CreateIterator();
-			TArray<float>::TIterator CoordUIterator = Points.PointsCoordU.CreateIterator();
+			TArray<FVector3f>::TIterator PositionIterator = Out.StrandsPoints.PointsPosition.CreateIterator();
+			TArray<float>::TIterator RadiusIterator = Out.StrandsPoints.PointsRadius.CreateIterator();
+			TArray<float>::TIterator CoordUIterator = Out.StrandsPoints.PointsCoordU.CreateIterator();
 
-			TArray<uint16>::TIterator CountIterator = Curves.CurvesCount.CreateIterator();
-			TArray<uint32>::TIterator OffsetIterator = Curves.CurvesOffset.CreateIterator();
-			TArray<float>::TIterator LengthIterator = Curves.CurvesLength.CreateIterator();
+			TArray<uint16>::TIterator CountIterator = Out.StrandsCurves.CurvesCount.CreateIterator();
+			TArray<uint32>::TIterator OffsetIterator = Out.StrandsCurves.CurvesOffset.CreateIterator();
+			TArray<float>::TIterator LengthIterator = Out.StrandsCurves.CurvesLength.CreateIterator();
 
 			uint32 StrandOffset = 0;
 			*OffsetIterator = StrandOffset; ++OffsetIterator;
 
-			for (uint32 CurveIndex = 0; CurveIndex < HairStrands.GetNumCurves(); ++CurveIndex, ++OffsetIterator, ++LengthIterator, ++CountIterator)
+			for (uint32 CurveIndex = 0; CurveIndex < Out.GetNumCurves(); ++CurveIndex, ++OffsetIterator, ++LengthIterator, ++CountIterator)
 			{
 				const uint16& StrandCount = *CountIterator;
 
@@ -133,8 +133,8 @@ namespace HairStrandsBuilder
 				FVector3f PreviousPosition(0.0, 0.0, 0.0);
 				for (uint32 PointIndex = 0; PointIndex < StrandCount; ++PointIndex, ++PositionIterator, ++RadiusIterator, ++CoordUIterator)
 				{
-					HairStrands.BoundingBox += *PositionIterator;
-					HairStrands.BoundingBox.IsValid = 1;
+					Out.BoundingBox += *PositionIterator;
+					Out.BoundingBox.IsValid = 1;
 					if (PointIndex > 0)
 					{
 						StrandLength += (*PositionIterator - PreviousPosition).Size();
@@ -150,7 +150,7 @@ namespace HairStrandsBuilder
 			RadiusIterator.Reset();
 			CoordUIterator.Reset();
 
-			for (uint32 CurveIndex = 0; CurveIndex < HairStrands.GetNumCurves(); ++CurveIndex, ++LengthIterator, ++CountIterator)
+			for (uint32 CurveIndex = 0; CurveIndex < Out.GetNumCurves(); ++CurveIndex, ++LengthIterator, ++CountIterator)
 			{
 				const uint16& StrandCount = *CountIterator;
 
@@ -163,6 +163,7 @@ namespace HairStrandsBuilder
 				}
 			}
 		}
+		return Out;
 	}
 
 	inline void CopyVectorToPosition(const FVector3f& InVector, FHairStrandsPositionFormat::Type& OutPosition)
@@ -2082,7 +2083,7 @@ bool FGroomBuilder::BuildHairDescriptionGroups(const FHairDescription& HairDescr
 			}
 		}
 
-		FHairStrandsDatas* CurrentHairStrandsDatas = nullptr;
+		FHairStrandsRawDatas* CurrentHairStrandsDatas = nullptr;
 		FHairDescriptionGroup& Group = FindOrAdd(StrandID);
 		bool bNumCurveValid = false;
 		if (bHasInvalidPosition)
@@ -2347,9 +2348,9 @@ bool FGroomBuilder::BuildHairDescriptionGroups(const FHairDescription& HairDescr
 	return true;
 }
 
-void FGroomBuilder::BuildData(FHairStrandsDatas& OutStrands)
+void FGroomBuilder::BuildData(const FHairStrandsRawDatas& In, FHairStrandsDatas& Out)
 {
-	HairStrandsBuilder::BuildInternalData(OutStrands);
+	Out = HairStrandsBuilder::BuildInternalData(In);
 }
 
 void FGroomBuilder::BuildData(
@@ -2391,11 +2392,10 @@ void FGroomBuilder::BuildData(
 		OutGroupInfo.GroupName = InHairDescriptionGroup.Info.GroupName;
 
 		// Rendering data
+		OutRen.Reset();
 		if (InHairDescriptionGroup.Strands.IsValid())
 		{
-			OutRen = InHairDescriptionGroup.Strands;
-
-			HairStrandsBuilder::BuildInternalData(OutRen);
+			OutRen = HairStrandsBuilder::BuildInternalData(InHairDescriptionGroup.Strands);
 
 			// Decimate
 			if (CurveDecimation < 1 || VertexDecimation < 1 || bCurveReordering)
@@ -2411,37 +2411,36 @@ void FGroomBuilder::BuildData(
 			// Sanity check
 			check(OutGroupInfo.NumCurves		<= InHairDescriptionGroup.Info.NumCurves);
 			check(OutGroupInfo.NumCurveVertices <= InHairDescriptionGroup.Info.NumCurveVertices);
-
 		}
 
 		// Simulation data
+		OutSim.Reset();
 		{
-			OutSim = InHairDescriptionGroup.Guides;
 			const bool bHasImportedGuides = InHairDescriptionGroup.Info.NumGuides > 0;
 			if (bHasImportedGuides && InSettings.InterpolationSettings.GuideType == EGroomGuideType::Imported)
 			{
-				HairStrandsBuilder::BuildInternalData(OutSim);
+				OutSim = HairStrandsBuilder::BuildInternalData(InHairDescriptionGroup.Guides);
 			}
 			else
 			{
-				OutSim.Reset();
 				if(InSettings.InterpolationSettings.GuideType == EGroomGuideType::Rigged)
 				{
 					if (InHairDescriptionGroup.Info.NumGuides > 0)
 					{
 						// Pick the new guides among the imported ones
-						FHairStrandsDatas TempSim = InHairDescriptionGroup.Guides;
-						HairStrandsBuilder::BuildInternalData(TempSim);
+						FHairStrandsDatas TempSim = HairStrandsBuilder::BuildInternalData(InHairDescriptionGroup.Guides);
 						FHairStrandsDecimation::Decimate(TempSim, InSettings.InterpolationSettings.RiggedGuideNumCurves, InSettings.InterpolationSettings.RiggedGuideNumPoints, OutSim);
 					}
 					else
 					{
 						// Otherwise let s pick the guides among the rendered strands
+						check(OutRen.GetNumPoints() > 0);
 						FHairStrandsDecimation::Decimate(OutRen, InSettings.InterpolationSettings.RiggedGuideNumCurves, InSettings.InterpolationSettings.RiggedGuideNumPoints, OutSim);
 					}
 				}
 				else if (!bHasImportedGuides || InSettings.InterpolationSettings.GuideType == EGroomGuideType::Generated)
 				{
+					check(OutRen.GetNumPoints() > 0);
 					FHairStrandsDecimation::Decimate(OutRen, HairToGuideDensity, 1, false, OutSim);
 				}
 				else
@@ -2485,7 +2484,7 @@ void FGroomBuilder::BuildBulkData(
 	TArray<GroomBuilder_Voxelization::FCoverageScale> CoverageScales;
 	if (bIsNotGuide)
 	{
-		GroomBuilder_Voxelization::ComputeHairCoverageScale(InData, CoverageScales);
+		GroomBuilder_Voxelization::ComputeHairCoverageScale(InData, CoverageScales); // TODO This should be done during the build step, not the bulk data build step
 	}
 	else
 	{
@@ -2780,7 +2779,7 @@ void Decimate(
 		OutPointOffset += OutPointCount;
 	}
 
-	HairStrandsBuilder::BuildInternalData(OutData);
+	OutData = HairStrandsBuilder::BuildInternalData(OutData);
 }
 } // namespace FHairStrandsDecimation
 
@@ -2885,13 +2884,8 @@ static void AllocateVoxels(const FBox3f& InBoundingBox, FHairGrid& Out)
 	Out.Voxels.SetNum(FMath::Max(Out.Resolution.X * Out.Resolution.Y * Out.Resolution.Z, 1));
 }
 
-static void VoxelizeCurves(const FHairStrandsDatas& InData, uint32 InGroupIndex, FHairGrid& Out)
+static void VoxelizeCurves(const FHairStrandsDatas& In, uint32 InGroupIndex, FHairGrid& Out)
 {
-	// 3. Voxelization curves
-	// Local copy of the hair strands data, and build the derived data so that curve offset are available
-	FHairStrandsDatas In = InData;
-	HairStrandsBuilder::BuildInternalData(In);
-
 	const uint32 Attributes = In.GetAttributes();
 
 	// Fill in voxel (TODO: make it parallel)
@@ -2948,7 +2942,9 @@ static void Voxelize(const FHairDescriptionGroups& InGroups, FHairGrid& Out)
 	uint32 GroupIndex = 0;
 	for (const FHairDescriptionGroup& Group : InGroups.HairGroups)
 	{
-		VoxelizeCurves(Group.Strands, GroupIndex, Out);
+		// Build the derived data so that curve offset are available
+		FHairStrandsDatas BuiltData = HairStrandsBuilder::BuildInternalData(Group.Strands);
+		VoxelizeCurves(BuiltData, GroupIndex, Out);
 		++GroupIndex;
 	}
 }
diff --git a/Engine/Plugins/Runtime/HairStrands/Source/HairStrandsCore/Private/GroomCache.cpp b/Engine/Plugins/Runtime/HairStrands/Source/HairStrandsCore/Private/GroomCache.cpp
--- a/Engine/Plugins/Runtime/HairStrands/Source/HairStrandsCore/Private/GroomCache.cpp
+++ b/Engine/Plugins/Runtime/HairStrands/Source/HairStrandsCore/Private/GroomCache.cpp
@@ -327,13 +327,13 @@ FGroomCacheProcessor::FGroomCacheProcessor(EGroomCacheType InType, EGroomCacheAt
 {
 }
 
-void FGroomCacheProcessor::AddGroomSample(TArray<FHairDescriptionGroup>&& GroupData)
+void FGroomCacheProcessor::AddGroomSample(TArray<FGroomCacheInputData>&& InputDatas)
 {
 	TArray<uint8> TempBytes;
 	FMemoryWriter Ar(TempBytes, true);
 
 	// The HairGroupData is converted into GroomCacheAnimationData and serialized to a buffer
-	FGroomCacheAnimationData AnimData(MoveTemp(GroupData), FGroomCacheInfo::GetCurrentVersion(), Type, Attributes);
+	FGroomCacheAnimationData AnimData(MoveTemp(InputDatas), FGroomCacheInfo::GetCurrentVersion(), Type, Attributes);
 	AnimData.Serialize(Ar);
 
 	FGroomCacheChunk& Chunk = Chunks.AddDefaulted_GetRef();
diff --git a/Engine/Plugins/Runtime/HairStrands/Source/HairStrandsCore/Private/GroomCacheData.cpp b/Engine/Plugins/Runtime/HairStrands/Source/HairStrandsCore/Private/GroomCacheData.cpp
--- a/Engine/Plugins/Runtime/HairStrands/Source/HairStrandsCore/Private/GroomCacheData.cpp
+++ b/Engine/Plugins/Runtime/HairStrands/Source/HairStrandsCore/Private/GroomCacheData.cpp
@@ -87,14 +87,14 @@ void FGroomCacheGroupData::Serialize(FArchive& Ar, int32 Version, EGroomCacheAtt
 	StrandData.Serialize(Ar, Version, Attributes);
 }
 
-FGroomCacheAnimationData::FGroomCacheAnimationData(TArray<FHairDescriptionGroup>&& HairGroupData, int32 InVersion, EGroomCacheType Type, EGroomCacheAttributes InAttributes)
+FGroomCacheAnimationData::FGroomCacheAnimationData(TArray<FGroomCacheInputData>&& InputDatas, int32 InVersion, EGroomCacheType Type, EGroomCacheAttributes InAttributes)
 : Version(InVersion)
 {
 	Attributes = Type == EGroomCacheType::Strands ? InAttributes : InAttributes & EGroomCacheAttributes::Position;
 
-	for (FHairDescriptionGroup& GroupData : HairGroupData)
+	for (FGroomCacheInputData& InputData : InputDatas)
 	{
-		GroupsData.Add(FGroomCacheGroupData(MoveTemp(Type == EGroomCacheType::Strands ? GroupData.Strands : GroupData.Guides)));
+		GroupsData.Add(FGroomCacheGroupData(MoveTemp(Type == EGroomCacheType::Strands ? InputData.Strands : InputData.Guides)));
 	}
 }
 
diff --git a/Engine/Plugins/Runtime/HairStrands/Source/HairStrandsCore/Private/GroomRBFDeformer.cpp b/Engine/Plugins/Runtime/HairStrands/Source/HairStrandsCore/Private/GroomRBFDeformer.cpp
--- a/Engine/Plugins/Runtime/HairStrands/Source/HairStrandsCore/Private/GroomRBFDeformer.cpp
+++ b/Engine/Plugins/Runtime/HairStrands/Source/HairStrandsCore/Private/GroomRBFDeformer.cpp
@@ -695,7 +695,7 @@ void FGroomRBFDeformer::GetRBFDeformedGroomAsset(const UGroomAsset* InGroomAsset
 		for (const FHairDescriptionGroup& Group : HairDescriptionGroups.HairGroups)
 		{
 			const uint32 GroupIndex = Group.Info.GroupID;
-			const FHairStrandsDatas& OriginalGuides = Group.Guides;
+			const FHairStrandsRawDatas& OriginalGuides = Group.Guides;
 			const FHairGroupPlatformData& HairGroupData = InGroomAsset->GetHairGroupsPlatformData()[GroupIndex];
 
 			FHairStrandsDatas StrandsData;
diff --git a/Engine/Plugins/Runtime/HairStrands/Source/HairStrandsCore/Private/HairAttributes.cpp b/Engine/Plugins/Runtime/HairStrands/Source/HairStrandsCore/Private/HairAttributes.cpp
--- a/Engine/Plugins/Runtime/HairStrands/Source/HairStrandsCore/Private/HairAttributes.cpp
+++ b/Engine/Plugins/Runtime/HairStrands/Source/HairStrandsCore/Private/HairAttributes.cpp
@@ -99,7 +99,7 @@ bool FHairStrandsCurves::HasPrecomputedWeights() const
 	return HasAttribute(EHairAttribute::PrecomputedGuideWeights); 
 }
 
-uint32 FHairStrandsDatas::GetAttributes() const
+uint32 FHairStrandsCommonDatas::GetAttributes() const
 {
 	uint32 Out = 0;
 	for (uint32 It=0;It<uint32(EHairAttribute::Count); ++It)
@@ -113,7 +113,7 @@ uint32 FHairStrandsDatas::GetAttributes() const
 	return Out;
 }
 
-uint32 FHairStrandsDatas::GetAttributeFlags() const
+uint32 FHairStrandsCommonDatas::GetAttributeFlags() const
 {
 	return StrandsCurves.AttributeFlags;
 }
diff --git a/Engine/Plugins/Runtime/HairStrands/Source/HairStrandsCore/Private/HairCardsBuilder.cpp b/Engine/Plugins/Runtime/HairStrands/Source/HairStrandsCore/Private/HairCardsBuilder.cpp
--- a/Engine/Plugins/Runtime/HairStrands/Source/HairStrandsCore/Private/HairCardsBuilder.cpp
+++ b/Engine/Plugins/Runtime/HairStrands/Source/HairStrandsCore/Private/HairCardsBuilder.cpp
@@ -69,9 +69,9 @@ FString GetVersion()
 	return TEXT("11");
 }
 
-bool InternalCreateCardsGuides(
+static bool InternalCreateCardsGuides(
 	FHairCardsGeometry& InCards,
-	FHairStrandsDatas& OutGuides,
+	FHairStrandsRawDatas& OutGuides,
 	TArray<float>& OutCardLengths)
 {
 	// Build the guides from the triangles that form the card
@@ -85,7 +85,6 @@ bool InternalCreateCardsGuides(
 	OutGuides.StrandsPoints.PointsRadius.Reserve(MaxGuidePoints);
 	OutGuides.StrandsPoints.PointsCoordU.Reserve(MaxGuidePoints);
 	OutGuides.StrandsCurves.SetNum(NumCards, 0u);
-	OutGuides.BoundingBox.Init();
 
 	InCards.CoordU.SetNum(InCards.Positions.Num());
 	InCards.LocalUVs.SetNum(InCards.Positions.Num());
@@ -329,8 +328,6 @@ bool InternalCreateCardsGuides(
 		{
 			const FVector3f P0 = CenterPoints[PointIt];
 
-			OutGuides.BoundingBox += P0;
-
 			OutGuides.StrandsPoints.PointsPosition.Add(P0);
 			OutGuides.StrandsPoints.PointsCoordU.Add(FMath::Clamp(CurrentLength / TotalLength, 0.f, 1.f));
 			OutGuides.StrandsPoints.PointsRadius.Add(1);
@@ -351,7 +348,7 @@ bool InternalCreateCardsGuides(
 // Build interpolation data between the cards geometry and the guides
 void InternalCreateCardsInterpolation(
 	const FHairCardsGeometry& InCards,
-	const FHairStrandsDatas& InGuides,
+	const FHairStrandsRawDatas& InGuides,
 	FHairCardsInterpolationDatas& Out)
 {		
 	const uint32 TotalVertexCount = InCards.Positions.Num();
@@ -605,13 +602,14 @@ static bool InternalImportGeometry_WithGeneratedGuides(
 
 	const uint32 CardsCount = Out.Cards.IndexOffsets.Num();
 
+	FHairStrandsRawDatas OutRawGuides;
 	TArray<float> CardLengths;
 	CardLengths.Reserve(CardsCount);
-	bool bSuccess = InternalCreateCardsGuides(Out.Cards, OutGuides, CardLengths);
+	bool bSuccess = InternalCreateCardsGuides(Out.Cards, OutRawGuides, CardLengths);
 	if (bSuccess)
 	{
 		FHairCardsInterpolationDatas InterpolationData;
-		InternalCreateCardsInterpolation(Out.Cards, OutGuides, InterpolationData);
+		InternalCreateCardsInterpolation(Out.Cards, OutRawGuides, InterpolationData);
 
 		// Fill out the interpolation data
 		static_assert(sizeof(FHairCardsInterpolationVertex) == HAIR_INTERPOLATION_CARDS_GUIDE_STRIDE);
@@ -865,6 +863,10 @@ static bool InternalImportGeometry_WithGeneratedGuides(
 		}
 	}
 
+	// Build the FHairStrandsData from the FHairStrandsRawData
+	PRAGMA_DISABLE_DEPRECATION_WARNINGS
+	FGroomBuilder::BuildData(OutRawGuides, OutGuides);
+	PRAGMA_ENABLE_DEPRECATION_WARNINGS
 	return bSuccess;
 }
 
diff --git a/Engine/Plugins/Runtime/HairStrands/Source/HairStrandsCore/Private/HairStrandsDatas.cpp b/Engine/Plugins/Runtime/HairStrands/Source/HairStrandsCore/Private/HairStrandsDatas.cpp
--- a/Engine/Plugins/Runtime/HairStrands/Source/HairStrandsCore/Private/HairStrandsDatas.cpp
+++ b/Engine/Plugins/Runtime/HairStrands/Source/HairStrandsCore/Private/HairStrandsDatas.cpp
@@ -94,6 +94,12 @@ void FHairStrandsPoints::SetNum(const uint32 NumPoints, uint32 InAttributes)
 	if (HasHairAttribute(InAttributes, EHairAttribute::AO))			{ PointsAO.SetNum(NumPoints); }
 }
 
+void FHairStrandsRawDatas::Reset()
+{
+	StrandsCurves.Reset();
+	StrandsPoints.Reset();
+}
+
 void FHairStrandsDatas::Reset()
 {
 	StrandsCurves.Reset();
diff --git a/Engine/Plugins/Runtime/HairStrands/Source/HairStrandsCore/Public/GroomAsset.h b/Engine/Plugins/Runtime/HairStrands/Source/HairStrandsCore/Public/GroomAsset.h
--- a/Engine/Plugins/Runtime/HairStrands/Source/HairStrandsCore/Public/GroomAsset.h
+++ b/Engine/Plugins/Runtime/HairStrands/Source/HairStrandsCore/Public/GroomAsset.h
@@ -394,9 +394,9 @@ struct FHairGroupPlatformData
 
 struct FHairDescriptionGroup
 {
-	FHairGroupInfo		Info;
-	FHairStrandsDatas	Strands;
-	FHairStrandsDatas	Guides;
+	FHairGroupInfo		 Info;
+	FHairStrandsRawDatas Strands;
+	FHairStrandsRawDatas Guides;
 	uint32 GetHairAttributes() const { return Strands.GetAttributes() | Guides.GetAttributes(); }
 	uint32 GetHairAttributeFlags() const { return Strands.GetAttributeFlags() | Guides.GetAttributeFlags(); }
 };
diff --git a/Engine/Plugins/Runtime/HairStrands/Source/HairStrandsCore/Public/GroomBuilder.h b/Engine/Plugins/Runtime/HairStrands/Source/HairStrandsCore/Public/GroomBuilder.h
--- a/Engine/Plugins/Runtime/HairStrands/Source/HairStrandsCore/Public/GroomBuilder.h
+++ b/Engine/Plugins/Runtime/HairStrands/Source/HairStrandsCore/Public/GroomBuilder.h
@@ -9,6 +9,7 @@
 #include "Math/IntVector.h"
 
 struct FHairStrandsDatas;
+struct FHairStrandsRawDatas;
 struct FHairStrandsClusterData;
 struct FHairStrandsClusterBulkData;
 struct FHairGroupsLOD;
@@ -47,22 +48,22 @@ struct FHairStrandsVoxelData
 
 // Data flow overview
 // ==================
-// HairDescription -> HairDescriptionGroups -> HairStrandsData -> HairStrandsBulkData*
-//															   -> HairStrandsInterpolationData ->HairStrandsInterpolationBulkData*
-//															   -> HairStrandsClusterData*
+// HairDescription -> HairDescriptionGroups (HairStrandsRawData) -> HairStrandsData -> HairStrandsBulkData*
+//																 -> HairStrandsInterpolationData ->HairStrandsInterpolationBulkData*
+//																 -> HairStrandsClusterData*
 //
 // * Data used at runtime. Other type of data are intermediate data used only during building/within the editor
 struct HAIRSTRANDSCORE_API FGroomBuilder
 {
 	static FString GetVersion();
 
-	// 1. Build hair group based on the hair description
+	// 1. Build hair *group* description based on the hair description. This builds HairStrandsRawData.
 	static bool BuildHairDescriptionGroups(
 		const FHairDescription& HairDescription, 
 		FHairDescriptionGroups& Out,
 		bool bAllowAddEndControlPoint=true);
 
-	// 2.a Build FHairStrandsDatas for Strands & Guides, based on HairDescriptionGroups and DecimationSettings
+	// 2.a Build FHairStrandsDatas from HairDescriptionGroups(HairStrandsRawData) and DecimationSettings (Strands / Guides)
 	static void BuildData(
 		const FHairDescriptionGroup& InHairDescriptionGroup,
 		const FHairGroupsInterpolation& InSettings,
@@ -78,13 +79,11 @@ struct HAIRSTRANDSCORE_API FGroomBuilder
 		FHairStrandsDatas& OutGuides,
 		bool bAllowCurveReordering=true);
 
-	// 2.b Build FHairStrandsDatas for Strands or Guides. 
-	// This version:
-	// * Suppose OutStrands already contains curves & points data. 
-	// * Compute only offset data/bounding box/max length/max radius/...
-	static void BuildData(FHairStrandsDatas& OutStrands);
+	// 2.b Build FHairStrandsDatas from a FHairStrandsRawDatas (Strands / Guides)
+	UE_DEPRECATED(5.6, "Please, no longer use this build function, as it does not contains certain build steps (decimation/shuffling)")
+	static void BuildData(const FHairStrandsRawDatas& In, FHairStrandsDatas& Out);
 
-	// 3. Build bulk data for Strands / Guides
+	// 3. Build bulk data from a FHairStrandsDatas (Strands / Guides)
 	static void BuildBulkData(
 		const FHairGroupInfo& InInfo,
 		const FHairStrandsDatas& InData,
diff --git a/Engine/Plugins/Runtime/HairStrands/Source/HairStrandsCore/Public/GroomCache.h b/Engine/Plugins/Runtime/HairStrands/Source/HairStrandsCore/Public/GroomCache.h
--- a/Engine/Plugins/Runtime/HairStrands/Source/HairStrandsCore/Public/GroomCache.h
+++ b/Engine/Plugins/Runtime/HairStrands/Source/HairStrandsCore/Public/GroomCache.h
@@ -106,7 +106,7 @@ class HAIRSTRANDSCORE_API FGroomCacheProcessor
 public:
 	FGroomCacheProcessor(EGroomCacheType InType, EGroomCacheAttributes InAttributes);
 
-	void AddGroomSample(TArray<FHairDescriptionGroup>&& GroomData);
+	void AddGroomSample(TArray<FGroomCacheInputData>&& GroomData);
 	void TransferChunks(UGroomCache* GroomCache);
 	EGroomCacheType GetType() const { return Type; }
 
diff --git a/Engine/Plugins/Runtime/HairStrands/Source/HairStrandsCore/Public/GroomCacheData.h b/Engine/Plugins/Runtime/HairStrands/Source/HairStrandsCore/Public/GroomCacheData.h
--- a/Engine/Plugins/Runtime/HairStrands/Source/HairStrandsCore/Public/GroomCacheData.h
+++ b/Engine/Plugins/Runtime/HairStrands/Source/HairStrandsCore/Public/GroomCacheData.h
@@ -3,6 +3,7 @@
 #pragma once
 
 #include "UObject/Object.h"
+#include "HairStrandsDatas.h"
 #include "GroomCacheData.generated.h"
 
 class FRDGPooledBuffer;
@@ -156,11 +157,18 @@ struct FGroomCacheGroupData
 	void Serialize(FArchive& Ar, int32 Version, EGroomCacheAttributes InAttributes);
 };
 
+/** Built groom data input used for a groom cache animation frame */
+struct FGroomCacheInputData
+{
+	FHairStrandsDatas Strands;
+	FHairStrandsDatas Guides;
+};
+
 /** Groom animation data for a frame*/
 struct FGroomCacheAnimationData
 {
 	FGroomCacheAnimationData() = default;
-	FGroomCacheAnimationData(TArray<struct FHairDescriptionGroup>&& HairGroupData, int32 Version, EGroomCacheType Type, EGroomCacheAttributes Attributes);
+	FGroomCacheAnimationData(TArray<struct FGroomCacheInputData>&& HairGroupData, int32 Version, EGroomCacheType Type, EGroomCacheAttributes Attributes);
 
 	TArray<FGroomCacheGroupData> GroupsData;
 	EGroomCacheAttributes Attributes;
diff --git a/Engine/Plugins/Runtime/HairStrands/Source/HairStrandsCore/Public/HairStrandsDatas.h b/Engine/Plugins/Runtime/HairStrands/Source/HairStrandsCore/Public/HairStrandsDatas.h
--- a/Engine/Plugins/Runtime/HairStrands/Source/HairStrandsCore/Public/HairStrandsDatas.h
+++ b/Engine/Plugins/Runtime/HairStrands/Source/HairStrandsCore/Public/HairStrandsDatas.h
@@ -573,7 +573,7 @@ struct HAIRSTRANDSCORE_API FHairStrandsCurves
 };
 
 /** Hair strands datas that are stored on CPU */
-struct HAIRSTRANDSCORE_API FHairStrandsDatas
+struct HAIRSTRANDSCORE_API FHairStrandsCommonDatas
 {
 	/* Get the total number of points */
 	uint32 GetNumPoints() const { return StrandsPoints.Num(); }
@@ -588,16 +588,29 @@ struct HAIRSTRANDSCORE_API FHairStrandsDatas
 
 	bool IsValid() const { return StrandsCurves.Num() > 0 && StrandsPoints.Num() > 0; }
 
-	/* Copy a point or a curve from In to Out data */
-	static void CopyCurve(const FHairStrandsDatas& In, FHairStrandsDatas& Out, uint32 InAttributes, uint32 InIndex, uint32 OutIndex);
-	static void CopyPoint(const FHairStrandsDatas& In, FHairStrandsDatas& Out, uint32 InAttributes, uint32 InIndex, uint32 OutIndex);
-	static void CopyPointLerp(const FHairStrandsDatas& In, FHairStrandsDatas& Out, uint32 InAttributes, uint32 InIndex0, uint32 InIndex1, float InAlpha, uint32 OutIndex);
-
 	/** List of all the strands points */
 	FHairStrandsPoints StrandsPoints;
 
 	/** List of all the strands curves */
 	FHairStrandsCurves StrandsCurves;
+};
+
+/** Hair strands raw datas that are stored on CPU */
+/** Unlike FHairStrandsDatas, this is raw datas, i.e. not processed by any built step */
+struct HAIRSTRANDSCORE_API FHairStrandsRawDatas : public FHairStrandsCommonDatas
+{
+	void Reset();
+};
+
+/** Built hair strands data (decimated/shuffled/...) */
+struct HAIRSTRANDSCORE_API FHairStrandsDatas : public FHairStrandsCommonDatas
+{
+	void Reset();
+
+	/* Copy a point or a curve from In to Out data */
+	static void CopyCurve(const FHairStrandsDatas& In, FHairStrandsDatas& Out, uint32 InAttributes, uint32 InIndex, uint32 OutIndex);
+	static void CopyPoint(const FHairStrandsDatas& In, FHairStrandsDatas& Out, uint32 InAttributes, uint32 InIndex, uint32 OutIndex);
+	static void CopyPointLerp(const FHairStrandsDatas& In, FHairStrandsDatas& Out, uint32 InAttributes, uint32 InIndex0, uint32 InIndex1, float InAlpha, uint32 OutIndex);
 
 	/** The Standard Hair Density */
 	float HairDensity = 1;
diff --git a/Engine/Plugins/Runtime/HairStrands/Source/HairStrandsEditor/Private/GroomCacheImporter.cpp b/Engine/Plugins/Runtime/HairStrands/Source/HairStrandsEditor/Private/GroomCacheImporter.cpp
--- a/Engine/Plugins/Runtime/HairStrands/Source/HairStrandsEditor/Private/GroomCacheImporter.cpp
+++ b/Engine/Plugins/Runtime/HairStrands/Source/HairStrandsEditor/Private/GroomCacheImporter.cpp
@@ -161,13 +161,16 @@ TArray<UGroomCache*> FGroomCacheImporter::ImportGroomCache(const FString& Source
 				}
 
 				TArray<FHairGroupInfoWithVisibility> HairGroupsInfo = GroomAssetForCache->GetHairGroupsInfo();
-				TArray<FHairDescriptionGroup> HairGroupsData;
-				HairGroupsData.SetNum(GroupCount);
+				TArray<FGroomCacheInputData> CacheInputDatas;
+				CacheInputDatas.SetNum(GroupCount);
 				for (uint32 GroupIndex = 0; GroupIndex < GroupCount; ++GroupIndex)
 				{
-					const FHairDescriptionGroup& HairGroup = HairDescriptionGroups.HairGroups[GroupIndex];
-					FHairDescriptionGroup& HairGroupData = HairGroupsData[GroupIndex];
-					FGroomBuilder::BuildData(HairGroup, GroomAssetForCache->GetHairGroupsInterpolation()[GroupIndex], HairGroupsInfo[GroupIndex], HairGroupData.Strands, HairGroupData.Guides);
+					FGroomBuilder::BuildData(
+						HairDescriptionGroups.HairGroups[GroupIndex], 
+						GroomAssetForCache->GetHairGroupsInterpolation()[GroupIndex], 
+						HairGroupsInfo[GroupIndex], 
+						CacheInputDatas[GroupIndex].Strands, 
+						CacheInputDatas[GroupIndex].Guides);
 				}
 
 				// Validate that the GroomCache has the same topology as the static groom
@@ -179,17 +182,17 @@ TArray<UGroomCache*> FGroomCacheImporter::ImportGroomCache(const FString& Source
 						// Since the groom cache needs to be independent of this settings (since the data is serialized directly into the asset
 						// and not rebuilt from a description), it does not contains these extra control points.
 						// We account for them for validation here
-						const uint32 ExtraControlPointCount = GetHairStrandsUsesTriangleStrips() ? HairGroupsData[GroupIndex].Strands.GetNumCurves() : 0u;
+						const uint32 ExtraControlPointCount = GetHairStrandsUsesTriangleStrips() ? CacheInputDatas[GroupIndex].Strands.GetNumCurves() : 0u;
 
-						const bool bCurveCountMatch = HairGroupsData[GroupIndex].Strands.GetNumCurves() == GroupPlatformData[GroupIndex].Strands.BulkData.GetNumCurves();
-						const bool bPointCountMatch = HairGroupsData[GroupIndex].Strands.GetNumPoints() + ExtraControlPointCount == GroupPlatformData[GroupIndex].Strands.BulkData.GetNumPoints();
+						const bool bCurveCountMatch = CacheInputDatas[GroupIndex].Strands.GetNumCurves() == GroupPlatformData[GroupIndex].Strands.BulkData.GetNumCurves();
+						const bool bPointCountMatch = CacheInputDatas[GroupIndex].Strands.GetNumPoints() + ExtraControlPointCount == GroupPlatformData[GroupIndex].Strands.BulkData.GetNumPoints();
 
 						if (!bCurveCountMatch || !bPointCountMatch)
 						{
 							bSuccess = false;
 							UE_LOG(LogGroomCacheImporter, Warning, TEXT("GroomCache frame %d does not have the same number of curves (%u) \
 								or vertices (%u) for the strands as the static groom (%u and %u respectively). Aborting GroomCache import."),
-								FrameIndex, HairGroupsData[GroupIndex].Strands.GetNumCurves(), HairGroupsData[GroupIndex].Strands.GetNumPoints(),
+								FrameIndex, CacheInputDatas[GroupIndex].Strands.GetNumCurves(), CacheInputDatas[GroupIndex].Strands.GetNumPoints(),
 								GroupPlatformData[GroupIndex].Strands.BulkData.GetNumCurves(), GroupPlatformData[GroupIndex].Strands.BulkData.GetNumPoints());
 							break;
 						}
@@ -197,22 +200,22 @@ TArray<UGroomCache*> FGroomCacheImporter::ImportGroomCache(const FString& Source
 
 					if (bImportGuidesCache)
 					{
-						if (HairGroupsData[GroupIndex].Strands.GetNumPoints() == 0)
+						if (CacheInputDatas[GroupIndex].Strands.GetNumPoints() == 0)
 						{
 							bGuidesOnly = true;
 						}
 
-						const uint32 ExtraControlPointCount = GetHairStrandsUsesTriangleStrips() ? HairGroupsData[GroupIndex].Guides.GetNumCurves() : 0u;
+						const uint32 ExtraControlPointCount = GetHairStrandsUsesTriangleStrips() ? CacheInputDatas[GroupIndex].Guides.GetNumCurves() : 0u;
 
-						const bool bCurveCountMatch = HairGroupsData[GroupIndex].Guides.GetNumCurves() == GroupPlatformData[GroupIndex].Guides.BulkData.GetNumCurves();
-						const bool bPointCountMatch = HairGroupsData[GroupIndex].Guides.GetNumPoints() + ExtraControlPointCount == GroupPlatformData[GroupIndex].Guides.BulkData.GetNumPoints();
+						const bool bCurveCountMatch = CacheInputDatas[GroupIndex].Guides.GetNumCurves() == GroupPlatformData[GroupIndex].Guides.BulkData.GetNumCurves();
+						const bool bPointCountMatch = CacheInputDatas[GroupIndex].Guides.GetNumPoints() + ExtraControlPointCount == GroupPlatformData[GroupIndex].Guides.BulkData.GetNumPoints();
 
 						if (!bCurveCountMatch || !bPointCountMatch)
 						{
 							bSuccess = false;
 							UE_LOG(LogGroomCacheImporter, Warning, TEXT("GroomCache frame %d does not have the same number of curves (%u) \
 								or vertices (%u) for the guides as the static groom (%u and %u respectively). Aborting GroomCache import."),
-								FrameIndex, HairGroupsData[GroupIndex].Guides.GetNumCurves(), HairGroupsData[GroupIndex].Guides.GetNumPoints(),
+								FrameIndex, CacheInputDatas[GroupIndex].Guides.GetNumCurves(), CacheInputDatas[GroupIndex].Guides.GetNumPoints(),
 								GroupPlatformData[GroupIndex].Guides.BulkData.GetNumCurves(), GroupPlatformData[GroupIndex].Guides.BulkData.GetNumPoints());
 							break;
 						}
@@ -227,12 +230,12 @@ TArray<UGroomCache*> FGroomCacheImporter::ImportGroomCache(const FString& Source
 				// The HairGroupData is converted into animated groom data by the GroomCacheProcessor
 				if (bImportStrandsCache && !bGuidesOnly)
 				{
-					StrandsProcessor.AddGroomSample(MoveTemp(HairGroupsData));
+					StrandsProcessor.AddGroomSample(MoveTemp(CacheInputDatas));
 				}
 
 				if (bImportGuidesCache)
 				{
-					GuidesProcessor.AddGroomSample(MoveTemp(HairGroupsData));
+					GuidesProcessor.AddGroomSample(MoveTemp(CacheInputDatas));
 				}
 			}
 		}
