Add GroupIndex to FHairGroupInfo/ in order to preserve GroupID value.

diff --git a/Engine/Plugins/Runtime/HairStrands/Source/HairStrandsCore/Private/GroomAsset.cpp b/Engine/Plugins/Runtime/HairStrands/Source/HairStrandsCore/Private/GroomAsset.cpp
--- a/Engine/Plugins/Runtime/HairStrands/Source/HairStrandsCore/Private/GroomAsset.cpp
+++ b/Engine/Plugins/Runtime/HairStrands/Source/HairStrandsCore/Private/GroomAsset.cpp
@@ -537,7 +537,7 @@ static bool BuildHairGroup_Strands(
 	{
 		const FHairDescriptionGroup& HairGroup = HairDescriptionGroups.HairGroups[GroupIndex];
 		check(GroupIndex <= HairDescriptionGroups.HairGroups.Num());
-		check(GroupIndex == HairGroup.Info.GroupID);
+		check(GroupIndex == HairGroup.Info.GroupIndex);
 
 		FHairStrandsDatas StrandsData;
 		FHairStrandsDatas GuidesData;
@@ -1162,7 +1162,7 @@ void UGroomAsset::UpdateHairGroupsInfo()
 	for (FHairGroupPlatformData& Data: GetHairGroupsPlatformData())
 	{
 		FHairGroupInfoWithVisibility& Info = GetHairGroupsInfo()[GroupIndex];
-		Info.GroupID = GroupIndex;
+		Info.GroupIndex = GroupIndex;
 		Info.NumCurves = Data.Strands.BulkData.GetNumCurves();
 		Info.NumGuides = Data.Guides.BulkData.GetNumCurves();
 		Info.NumCurveVertices = Data.Strands.BulkData.GetNumPoints();
@@ -3716,7 +3716,7 @@ bool UGroomAsset::GetHairStrandsDatas(
 	if (bIsValid)
 	{
 		const FHairDescriptionGroup& HairGroup = DescriptionGroups.HairGroups[GroupIndex];
-		check(GroupIndex == HairGroup.Info.GroupID);
+		check(GroupIndex == HairGroup.Info.GroupIndex);
 
 		FHairGroupInfo DummyInfo;
 		FGroomBuilder::BuildData(HairGroup, GetHairGroupsInterpolation()[GroupIndex], OutStrandsData, OutGuidesData);
diff --git a/Engine/Plugins/Runtime/HairStrands/Source/HairStrandsCore/Private/GroomBuilder.cpp b/Engine/Plugins/Runtime/HairStrands/Source/HairStrandsCore/Private/GroomBuilder.cpp
--- a/Engine/Plugins/Runtime/HairStrands/Source/HairStrandsCore/Private/GroomBuilder.cpp
+++ b/Engine/Plugins/Runtime/HairStrands/Source/HairStrandsCore/Private/GroomBuilder.cpp
@@ -2294,13 +2294,12 @@ bool FGroomBuilder::BuildHairDescriptionGroups(const FHairDescription& HairDescr
 		}
 	}
 
-	// Sparse->Dense Groups
-	// Change GroupID to be contiguous
+	// Initialize group index
 	{
 		int32 GroupIndex = 0;
 		for (FHairDescriptionGroup& Group : Out.HairGroups)
 		{
-			Group.Info.GroupID = GroupIndex++;
+			Group.Info.GroupIndex = GroupIndex++;
 		}
 	}
 
@@ -2388,6 +2387,7 @@ void FGroomBuilder::BuildData(
 	const bool bCurveReordering 	= bAllowCurveReordering;
 
 	{
+		OutGroupInfo.GroupIndex = InHairDescriptionGroup.Info.GroupIndex;
 		OutGroupInfo.GroupID = InHairDescriptionGroup.Info.GroupID;
 		OutGroupInfo.GroupName = InHairDescriptionGroup.Info.GroupName;
 
diff --git a/Engine/Plugins/Runtime/HairStrands/Source/HairStrandsCore/Private/GroomEdit.cpp b/Engine/Plugins/Runtime/HairStrands/Source/HairStrandsCore/Private/GroomEdit.cpp
--- a/Engine/Plugins/Runtime/HairStrands/Source/HairStrandsCore/Private/GroomEdit.cpp
+++ b/Engine/Plugins/Runtime/HairStrands/Source/HairStrandsCore/Private/GroomEdit.cpp
@@ -168,6 +168,7 @@ void ConvertFromGroomAsset(UGroomAsset* InGroom, FEditableGroom* OutGroom)
 		FHairStrandsDatas GuidesData;
 		FGroomBuilder::BuildData(HairDescriptionGroups.HairGroups[GroupIndex], In.GetHairGroupsInterpolation()[GroupIndex], GroupInfo, StrandsData, GuidesData);
 
+		Out.Groups[GroupIndex].GroupIndex = GroupIndex;
 		Out.Groups[GroupIndex].GroupID = GroupInfo.GroupID;
 		Out.Groups[GroupIndex].GroupName = GroupInfo.GroupName;
 		HairStrandsDataToEditableHairGuides(GuidesData, Out.Groups[GroupIndex].Guides);
@@ -321,6 +322,7 @@ static void InternalEditTest4(UGroomAsset* InGroom)
 		const FEditableGroomGroup& SrcG = EditGroom.Groups[0];
 		FEditableGroomGroup& DstG = EditGroom.Groups[GroupIt];
 
+		DstG.GroupIndex = GroupIt;
 		DstG.GroupID = GroupIt;
 		DstG.GroupName = FName(TEXT("MyNewGroup"));
 
diff --git a/Engine/Plugins/Runtime/HairStrands/Source/HairStrandsCore/Private/GroomRBFDeformer.cpp b/Engine/Plugins/Runtime/HairStrands/Source/HairStrandsCore/Private/GroomRBFDeformer.cpp
--- a/Engine/Plugins/Runtime/HairStrands/Source/HairStrandsCore/Private/GroomRBFDeformer.cpp
+++ b/Engine/Plugins/Runtime/HairStrands/Source/HairStrandsCore/Private/GroomRBFDeformer.cpp
@@ -467,7 +467,7 @@ static void ApplyDeformationToGroom(const TArray<FRBFDeformedPositions>& Deforme
 			{
 				const FHairDescriptionGroup& HairGroup = HairDescriptionGroups.HairGroups[GroupIndex];
 				check(GroupIndex <= HairDescriptionGroups.HairGroups.Num());
-				check(GroupIndex == HairGroup.Info.GroupID);
+				check(GroupIndex == HairGroup.Info.GroupIndex);
 				const FHairInterpolationSettings& InterpolationSettings = GroomAsset->GetHairGroupsInterpolation()[GroupIndex].InterpolationSettings;
 				const FHairGroupsLOD& HairGroupLOD = GroomAsset->GetHairGroupsLOD()[GroupIndex];
 
@@ -694,7 +694,7 @@ void FGroomRBFDeformer::GetRBFDeformedGroomAsset(const UGroomAsset* InGroomAsset
 		// but the ordering of the groups does represent the GroupIndex
 		for (const FHairDescriptionGroup& Group : HairDescriptionGroups.HairGroups)
 		{
-			const uint32 GroupIndex = Group.Info.GroupID;
+			const uint32 GroupIndex = Group.Info.GroupIndex;
 			const FHairStrandsRawDatas& OriginalGuides = Group.Guides;
 			const FHairGroupPlatformData& HairGroupData = InGroomAsset->GetHairGroupsPlatformData()[GroupIndex];
 
diff --git a/Engine/Plugins/Runtime/HairStrands/Source/HairStrandsCore/Public/GroomAsset.h b/Engine/Plugins/Runtime/HairStrands/Source/HairStrandsCore/Public/GroomAsset.h
--- a/Engine/Plugins/Runtime/HairStrands/Source/HairStrandsCore/Public/GroomAsset.h
+++ b/Engine/Plugins/Runtime/HairStrands/Source/HairStrandsCore/Public/GroomAsset.h
@@ -58,6 +58,9 @@ struct HAIRSTRANDSCORE_API FHairGroupInfo
 {
 	GENERATED_BODY()
 
+	UPROPERTY(VisibleAnywhere, Category = "Info")
+	int32 GroupIndex = 0;
+
 	UPROPERTY(VisibleAnywhere, Category = "Info")
 	int32 GroupID = 0;
 
diff --git a/Engine/Plugins/Runtime/HairStrands/Source/HairStrandsCore/Public/GroomEdit.h b/Engine/Plugins/Runtime/HairStrands/Source/HairStrandsCore/Public/GroomEdit.h
--- a/Engine/Plugins/Runtime/HairStrands/Source/HairStrandsCore/Public/GroomEdit.h
+++ b/Engine/Plugins/Runtime/HairStrands/Source/HairStrandsCore/Public/GroomEdit.h
@@ -53,6 +53,7 @@ struct FEditableHairGuide
 // Group & Groom
 struct FEditableGroomGroup
 {
+	uint32 GroupIndex;
 	uint32 GroupID;
 	FName GroupName;
 	TArray<FEditableHairStrand> Strands;
diff --git a/Engine/Plugins/Runtime/HairStrands/Source/HairStrandsCore/Public/GroomImportOptions.h b/Engine/Plugins/Runtime/HairStrands/Source/HairStrandsCore/Public/GroomImportOptions.h
--- a/Engine/Plugins/Runtime/HairStrands/Source/HairStrandsCore/Public/GroomImportOptions.h
+++ b/Engine/Plugins/Runtime/HairStrands/Source/HairStrandsCore/Public/GroomImportOptions.h
@@ -27,6 +27,9 @@ struct HAIRSTRANDSCORE_API FGroomHairGroupPreview
 {
 	GENERATED_USTRUCT_BODY()
 
+	UPROPERTY(VisibleAnywhere, BlueprintReadWrite, Category = Preview)
+	int32 GroupIndex = 0;
+
 	UPROPERTY(VisibleAnywhere, BlueprintReadWrite, Category = Preview)
 	FName GroupName;
 
diff --git a/Engine/Plugins/Runtime/HairStrands/Source/HairStrandsEditor/Private/GroomActions.cpp b/Engine/Plugins/Runtime/HairStrands/Source/HairStrandsEditor/Private/GroomActions.cpp
--- a/Engine/Plugins/Runtime/HairStrands/Source/HairStrandsEditor/Private/GroomActions.cpp
+++ b/Engine/Plugins/Runtime/HairStrands/Source/HairStrandsEditor/Private/GroomActions.cpp
@@ -115,7 +115,8 @@ void ExecuteRebuild(const FToolMenuContext& InContext)
 					for (uint32 GroupIndex = 0; GroupIndex < GroupCount; GroupIndex++)
 					{
 						FGroomHairGroupPreview& OutGroup = GroupsPreview->Groups.AddDefaulted_GetRef();
-						OutGroup.GroupID    	= GroupIndex;
+						OutGroup.GroupIndex 	= GroupIndex;
+						OutGroup.GroupID		= GroomAsset->GetHairGroupsInfo()[GroupIndex].GroupID;
 						OutGroup.GroupName		= GroomAsset->GetHairGroupsInfo()[GroupIndex].GroupName;
 						OutGroup.CurveCount 	= GroomAsset->GetHairGroupsPlatformData()[GroupIndex].Strands.BulkData.GetNumCurves();
 						OutGroup.GuideCount 	= GroomAsset->GetHairGroupsPlatformData()[GroupIndex].Guides.BulkData.GetNumCurves();
@@ -150,7 +151,7 @@ void ExecuteRebuild(const FToolMenuContext& InContext)
 					CurrentOptions->Rename(nullptr, GroomAssetImportData);
 					for (const FGroomHairGroupPreview& GroupPreview : GroupsPreview->Groups)
 					{
-						CurrentOptions->InterpolationSettings[GroupPreview.GroupID] = GroupPreview.InterpolationSettings;
+						CurrentOptions->InterpolationSettings[GroupPreview.GroupIndex] = GroupPreview.InterpolationSettings;
 					}
 					GroomAssetImportData->ImportOptions = CurrentOptions;
 					GroomAsset->MarkPackageDirty();
diff --git a/Engine/Plugins/Runtime/HairStrands/Source/HairStrandsEditor/Private/HairStrandsFactory.cpp b/Engine/Plugins/Runtime/HairStrands/Source/HairStrandsEditor/Private/HairStrandsFactory.cpp
--- a/Engine/Plugins/Runtime/HairStrands/Source/HairStrandsEditor/Private/HairStrandsFactory.cpp
+++ b/Engine/Plugins/Runtime/HairStrands/Source/HairStrandsEditor/Private/HairStrandsFactory.cpp
@@ -140,6 +140,7 @@ UObject* UHairStrandsFactory::FactoryCreateFile(UClass* InClass, UObject* InPare
 				FGroomHairGroupPreview& OutGroup = GroupsPreview->Groups.AddDefaulted_GetRef();
 				OutGroup.GroupName  		= Group.Info.GroupName;
 				OutGroup.GroupID			= Group.Info.GroupID;
+				OutGroup.GroupIndex			= Group.Info.GroupIndex;
 				OutGroup.CurveCount 		= Group.Info.NumCurves;
 				OutGroup.GuideCount 		= Group.Info.NumGuides;
 				OutGroup.Attributes 		= Group.GetHairAttributes();
@@ -147,9 +148,9 @@ UObject* UHairStrandsFactory::FactoryCreateFile(UClass* InClass, UObject* InPare
 				OutGroup.Flags      		= Group.Info.Flags;
 				bGuidesOnly |= (OutGroup.CurveCount == 0 && OutGroup.GuideCount > 0);
 
-				if (OutGroup.GroupID < OutDescription.HairGroups.Num())
+				if (OutGroup.GroupIndex < OutDescription.HairGroups.Num())
 				{				
-					OutGroup.InterpolationSettings = ImportOptions->InterpolationSettings[OutGroup.GroupID];
+					OutGroup.InterpolationSettings = ImportOptions->InterpolationSettings[OutGroup.GroupIndex];
 				}
 			}
 		}
@@ -179,9 +180,9 @@ UObject* UHairStrandsFactory::FactoryCreateFile(UClass* InClass, UObject* InPare
 			// Save the options as the new default
 			for (const FGroomHairGroupPreview& GroupPreview : GroupsPreview->Groups)
 			{
-				if (GroupPreview.GroupID < OutDescription.HairGroups.Num())
+				if (GroupPreview.GroupIndex < OutDescription.HairGroups.Num())
 				{
-					ImportOptions->InterpolationSettings[GroupPreview.GroupID] = GroupPreview.InterpolationSettings;
+					ImportOptions->InterpolationSettings[GroupPreview.GroupIndex] = GroupPreview.InterpolationSettings;
 				}
 			}
 			ImportOptions->SaveConfig();
diff --git a/Engine/Plugins/Runtime/HairStrands/Source/HairStrandsEditor/Private/ReimportHairStrandsFactory.cpp b/Engine/Plugins/Runtime/HairStrands/Source/HairStrandsEditor/Private/ReimportHairStrandsFactory.cpp
--- a/Engine/Plugins/Runtime/HairStrands/Source/HairStrandsEditor/Private/ReimportHairStrandsFactory.cpp
+++ b/Engine/Plugins/Runtime/HairStrands/Source/HairStrandsEditor/Private/ReimportHairStrandsFactory.cpp
@@ -382,14 +382,15 @@ EReimportResult::Type UReimportHairStrandsFactory::Reimport(UObject* Obj)
 				FGroomHairGroupPreview& OutGroup = GroupsPreview->Groups.AddDefaulted_GetRef();
 				OutGroup.GroupName  	= Group.Info.GroupName;
 				OutGroup.GroupID		= Group.Info.GroupID;
+				OutGroup.GroupIndex		= Group.Info.GroupIndex;
 				OutGroup.CurveCount 	= Group.Info.NumCurves;
 				OutGroup.GuideCount 	= Group.Info.NumGuides;
 				OutGroup.Attributes 	= Group.GetHairAttributes();
 				OutGroup.AttributeFlags = Group.GetHairAttributeFlags();
 				OutGroup.Flags			= Group.Info.Flags;
-				if (OutGroup.GroupID < OutDescription.HairGroups.Num())
+				if (OutGroup.GroupIndex < OutDescription.HairGroups.Num())
 				{
-					OutGroup.InterpolationSettings = GroomReimportOptions->InterpolationSettings[OutGroup.GroupID];
+					OutGroup.InterpolationSettings = GroomReimportOptions->InterpolationSettings[OutGroup.GroupIndex];
 				}
 			}
 		}
@@ -404,9 +405,9 @@ EReimportResult::Type UReimportHairStrandsFactory::Reimport(UObject* Obj)
 		// Save the options as the new default
 		for (const FGroomHairGroupPreview& GroupPreview : GroupsPreview->Groups)
 		{
-			if (GroupPreview.GroupID < OutDescription.HairGroups.Num())
+			if (GroupPreview.GroupIndex < OutDescription.HairGroups.Num())
 			{
-				GroomReimportOptions->InterpolationSettings[GroupPreview.GroupID] = GroupPreview.InterpolationSettings;
+				GroomReimportOptions->InterpolationSettings[GroupPreview.GroupIndex] = GroupPreview.InterpolationSettings;
 			}
 		}
 	}
