Fix groom reimporting when the groups layout differ from the original asset.

This CL tries to preserve the groom settings of each group when the group layout differs.

diff --git a/Engine/Plugins/Runtime/HairStrands/Source/HairStrandsCore/Private/GroomAsset.cpp b/Engine/Plugins/Runtime/HairStrands/Source/HairStrandsCore/Private/GroomAsset.cpp
--- a/Engine/Plugins/Runtime/HairStrands/Source/HairStrandsCore/Private/GroomAsset.cpp
+++ b/Engine/Plugins/Runtime/HairStrands/Source/HairStrandsCore/Private/GroomAsset.cpp
@@ -2053,9 +2053,7 @@ void UGroomAsset::SetNumGroup(uint32 InGroupCount, bool bResetGroupData, bool bR
 		{
 			if (GroupLOD.LODs.IsEmpty())
 			{
-				FHairLODSettings& S = GroupLOD.LODs.AddDefaulted_GetRef();
-				S.ScreenSize = 1;
-				S.CurveDecimation = 1;
+				GroupLOD = FHairGroupsLOD::GetDefault();
 			}
 		}
 	}
@@ -2086,6 +2084,31 @@ void UGroomAsset::SetNumGroup(uint32 InGroupCount, bool bResetGroupData, bool bR
 #endif
 }
 
+void UGroomAsset::ClearNumGroup(uint32 InGroupCount)
+{
+	ReleaseResource();
+
+	const FHairGroupsLOD Default = FHairGroupsLOD::GetDefault();
+
+	GetHairGroupsPlatformData().Init(FHairGroupPlatformData(), InGroupCount);
+	GetHairGroupsResources().Init(FHairGroupResources(), InGroupCount);
+	GetHairGroupsPhysics().Init(FHairGroupsPhysics(), InGroupCount);
+	GetHairGroupsRendering().Init(FHairGroupsRendering(), InGroupCount);
+	GetHairGroupsInterpolation().Init(FHairGroupsInterpolation(), InGroupCount);	
+	GetHairGroupsLOD().Init(Default, InGroupCount);
+	GetHairGroupsInfo().Init(FHairGroupInfoWithVisibility(), InGroupCount);
+	GetEffectiveLODBias().Init(0, InGroupCount);
+	// The following groups remain unchanged:
+	// * HairGroupsCards
+	// * GetHairGroupsMeshes()
+	// * GetHairGroupsMaterials()
+	#if WITH_EDITORONLY_DATA
+	StrandsDerivedDataKey.Init(FString(), InGroupCount);
+	CardsDerivedDataKey.Init(FString(), InGroupCount);
+	MeshesDerivedDataKey.Init(FString(), InGroupCount);
+	#endif
+}
+
 void UGroomAsset::SetStableRasterization(bool bEnable)
 {
 	for (FHairGroupsRendering& Group : GetHairGroupsRendering())
diff --git a/Engine/Plugins/Runtime/HairStrands/Source/HairStrandsCore/Private/GroomAssetInterpolation.cpp b/Engine/Plugins/Runtime/HairStrands/Source/HairStrandsCore/Private/GroomAssetInterpolation.cpp
--- a/Engine/Plugins/Runtime/HairStrands/Source/HairStrandsCore/Private/GroomAssetInterpolation.cpp
+++ b/Engine/Plugins/Runtime/HairStrands/Source/HairStrandsCore/Private/GroomAssetInterpolation.cpp
@@ -123,3 +123,12 @@ void FHairGroupsLOD::BuildDDCKey(FArchive& Ar)
 		Ar << LOD.GlobalInterpolation;
 	}
 }
+
+FHairGroupsLOD FHairGroupsLOD::GetDefault()
+{
+	FHairGroupsLOD Out;
+	FHairLODSettings& S = Out.LODs.AddDefaulted_GetRef();
+	S.ScreenSize = 1;
+	S.CurveDecimation = 1;
+	return Out;
+}
\ No newline at end of file
diff --git a/Engine/Plugins/Runtime/HairStrands/Source/HairStrandsCore/Public/GroomAsset.h b/Engine/Plugins/Runtime/HairStrands/Source/HairStrandsCore/Public/GroomAsset.h
--- a/Engine/Plugins/Runtime/HairStrands/Source/HairStrandsCore/Public/GroomAsset.h
+++ b/Engine/Plugins/Runtime/HairStrands/Source/HairStrandsCore/Public/GroomAsset.h
@@ -692,6 +692,8 @@ public:
 	void ReleaseMeshesResource(uint32 GroupIndex);
 
 	void SetNumGroup(uint32 InGroupCount, bool bResetGroupData=true, bool bResetOtherData=true);
+	void ClearNumGroup(uint32 InGroupCount);
+
 	bool AreGroupsValid() const;
 	int32 GetNumHairGroups() const;
 
diff --git a/Engine/Plugins/Runtime/HairStrands/Source/HairStrandsCore/Public/GroomAssetInterpolation.h b/Engine/Plugins/Runtime/HairStrands/Source/HairStrandsCore/Public/GroomAssetInterpolation.h
--- a/Engine/Plugins/Runtime/HairStrands/Source/HairStrandsCore/Public/GroomAssetInterpolation.h
+++ b/Engine/Plugins/Runtime/HairStrands/Source/HairStrandsCore/Public/GroomAssetInterpolation.h
@@ -227,4 +227,5 @@ struct HAIRSTRANDSCORE_API FHairGroupsLOD
 	bool operator==(const FHairGroupsLOD& A) const;
 
 	void BuildDDCKey(FArchive& Ar);
+	static FHairGroupsLOD GetDefault();
 };
diff --git a/Engine/Plugins/Runtime/HairStrands/Source/HairStrandsEditor/Private/HairStrandsImporter.cpp b/Engine/Plugins/Runtime/HairStrands/Source/HairStrandsEditor/Private/HairStrandsImporter.cpp
--- a/Engine/Plugins/Runtime/HairStrands/Source/HairStrandsEditor/Private/HairStrandsImporter.cpp
+++ b/Engine/Plugins/Runtime/HairStrands/Source/HairStrandsEditor/Private/HairStrandsImporter.cpp
@@ -9,6 +9,159 @@
 
 DEFINE_LOG_CATEGORY_STATIC(LogHairImporter, Log, All);
 
+static int32 GHairStrandsReimportGroupRemapping = 1;
+static FAutoConsoleVariableRef CVarHairStrandsLoadAsset(TEXT("r.HairStrands.ReimportGroupRemapping"), GHairStrandsReimportGroupRemapping, TEXT("Remap hair group settings/parameters when reimporting a groom asset (experimental)"));
+
+///////////////////////////////////////////////////////////////////////////////////////////////////
+// Remapping of settings when reimporting groom asset
+
+struct FHairGroupSettings
+{
+	//uint32 GroupIndex = 0;
+	FHairGroupInfoWithVisibility Info;
+	FHairGroupsRendering Rendering;
+	FHairGroupsPhysics Physics;
+	FHairGroupsInterpolation Interpolation;
+	FHairGroupsLOD LOD;
+	float EffectiveLODBias = 0;
+};
+
+static TArray<FHairGroupSettings> GetHairGroupSettings(const UGroomAsset* In)
+{
+	TArray<FHairGroupSettings> Out;
+	for (uint32 GroupIndex = 0, GroupCount = In->GetNumHairGroups(); GroupIndex < GroupCount; ++GroupIndex)
+	{
+		FHairGroupSettings& OutGroup = Out.AddDefaulted_GetRef();
+		OutGroup.Info				= In->GetHairGroupsInfo()[GroupIndex];
+		OutGroup.Rendering			= In->GetHairGroupsRendering()[GroupIndex];
+		OutGroup.Physics			= In->GetHairGroupsPhysics()[GroupIndex];
+		OutGroup.Interpolation		= In->GetHairGroupsInterpolation()[GroupIndex];
+		OutGroup.LOD				= In->GetHairGroupsLOD()[GroupIndex];
+		OutGroup.EffectiveLODBias	= In->GetEffectiveLODBias()[GroupIndex];
+		check(GroupIndex == OutGroup.Info.GroupIndex);
+	}
+	return Out;
+}
+
+static void SetHairGroupSettings(UGroomAsset* OutAsset, const TArray<FHairGroupSettings>& InSettings, const TArray<int32>& OldToNewGroupIndexMapping)
+{
+	for (uint32 GroupIndex = 0, GroupCount = OutAsset->GetNumHairGroups(); GroupIndex < GroupCount; ++GroupIndex)
+	{
+		const FHairGroupSettings& Settings = InSettings[GroupIndex];
+		OutAsset->GetHairGroupsInfo()[GroupIndex]			= Settings.Info;
+		OutAsset->GetHairGroupsRendering()[GroupIndex]		= Settings.Rendering;
+		OutAsset->GetHairGroupsPhysics()[GroupIndex]		= Settings.Physics;
+		OutAsset->GetHairGroupsInterpolation()[GroupIndex]	= Settings.Interpolation;
+		OutAsset->GetHairGroupsLOD()[GroupIndex]			= Settings.LOD;
+		OutAsset->GetEffectiveLODBias()[GroupIndex]			= Settings.EffectiveLODBias;
+	}
+
+	// Update cards description with new group index
+	for (FHairGroupsCardsSourceDescription& Desc : OutAsset->GetHairGroupsCards())
+	{
+		Desc.GroupIndex = OldToNewGroupIndexMapping[Desc.GroupIndex];
+	}
+
+	// Update mesh description with new group index
+	for (FHairGroupsMeshesSourceDescription& Desc : OutAsset->GetHairGroupsMeshes())
+	{
+		Desc.GroupIndex = OldToNewGroupIndexMapping[Desc.GroupIndex];
+	}
+}
+
+static bool FindHairGroupSettings(uint32 InGroupID, const FName& InGroupName, const TArray<FHairGroupSettings>& InSettings, FHairGroupSettings& Out)
+{
+	for (const FHairGroupSettings& Group : InSettings)
+	{
+		// For now, only use GroupName, as if no group name is provided, the group name is built from the GroupID
+		if (Group.Info.GroupName == InGroupName)
+		{
+			Out = Group;
+			return true;
+		}
+	}
+	return false;
+}
+
+static void RemapHairGroupSettings(
+	const UGroomAsset* InGroomAsset, 
+	const FHairDescriptionGroups& InHairDescriptionGroups, 
+	const TArray<FHairGroupsInterpolation> ImportedInterpolationSettings,
+	TArray<FHairGroupSettings>& OutNewSettings, 
+	TArray<int32>& OutRemapping)
+{
+	const FHairGroupsLOD DefaultLOD = FHairGroupsLOD::GetDefault();
+
+	const bool bUseGroupRemapping = GHairStrandsReimportGroupRemapping > 0;
+	if (bUseGroupRemapping)
+	{
+		// 1. Extract existing settings from the old groom asset
+		const TArray<FHairGroupSettings> OldSettings = GetHairGroupSettings(InGroomAsset);
+		OutRemapping.Init(-1, OldSettings.Num());
+
+		// 2. Remap old settings to the new asset (use GroupName/GroomID to do the remapping)	
+		const uint32 GroupCount = InHairDescriptionGroups.HairGroups.Num();
+		OutNewSettings.SetNum(GroupCount);
+		for (const FHairDescriptionGroup& GroupDesc : InHairDescriptionGroups.HairGroups)
+		{
+			FHairGroupSettings OldGroup;
+			FHairGroupSettings& NewGroup = OutNewSettings[GroupDesc.Info.GroupIndex];
+			if (FindHairGroupSettings(GroupDesc.Info.GroupID, GroupDesc.Info.GroupName, OldSettings, OldGroup))
+			{
+				NewGroup = OldGroup;
+				OutRemapping[OldGroup.Info.GroupIndex] = GroupDesc.Info.GroupIndex;
+			}
+			else
+			{
+				NewGroup.LOD = DefaultLOD;
+			}
+
+			// Preserve the new group GroupIndex/GroupID/GroupName
+			NewGroup.Info.GroupIndex = GroupDesc.Info.GroupIndex;
+			NewGroup.Info.GroupID    = GroupDesc.Info.GroupID;
+			NewGroup.Info.GroupName  = GroupDesc.Info.GroupName;
+
+			// Use import interpolation settings
+			NewGroup.Interpolation   = ImportedInterpolationSettings[GroupDesc.Info.GroupIndex];
+		}
+	}
+}
+
+void RemapHairGroupInteprolationSettings(
+	const UGroomAsset* InGroomAsset, 
+	const FHairDescriptionGroups& InHairDescriptionGroups,
+	TArray<FHairGroupsInterpolation>& OutInterpolationSettings)
+{
+	// Remap existing interpolation settings based on GroupName/GroupID if possible, otherwise initialize them to default
+	// Populate the interpolation settings based on the group count from the description
+	const uint32 GroupCount = InHairDescriptionGroups.HairGroups.Num();
+	OutInterpolationSettings.Init(FHairGroupsInterpolation(), GroupCount);
+
+	const bool bUseGroupRemapping = GHairStrandsReimportGroupRemapping > 0;
+	if (bUseGroupRemapping)
+	{
+		const TArray<FHairGroupSettings> OldSettings = GetHairGroupSettings(InGroomAsset);
+		for (const FHairDescriptionGroup& Group : InHairDescriptionGroups.HairGroups)
+		{			
+			FHairGroupSettings GroupSettings;
+			if (FindHairGroupSettings(Group.Info.GroupID, Group.Info.GroupName, OldSettings, GroupSettings))
+			{
+				OutInterpolationSettings[Group.Info.GroupIndex] = GroupSettings.Interpolation;
+			}
+		}
+	}
+	else
+	{
+		// Old path
+		for (uint32 GroupIndex = 0; GroupIndex < GroupCount; ++GroupIndex)
+		{
+			OutInterpolationSettings[GroupIndex] = InGroomAsset->GetHairGroupsInterpolation()[GroupIndex];
+		}
+	}
+}
+
+///////////////////////////////////////////////////////////////////////////////////////////////////
+
 FHairImportContext::FHairImportContext(UGroomImportOptions* InImportOptions, UObject* InParent, UClass* InClass, FName InName, EObjectFlags InFlags)
 	: ImportOptions(InImportOptions)
 	, Parent(InParent)
@@ -35,7 +188,39 @@ UGroomAsset* FHairStrandsImporter::ImportHair(const FHairImportContext& ImportCo
 			return nullptr;
 		}
 	}
-	OutHairAsset->SetNumGroup(GroupCount);
+
+	const bool bUseGroupRemapping = GHairStrandsReimportGroupRemapping > 0;
+	if (ExistingHair && bUseGroupRemapping)
+	{
+		FHairDescriptionGroups HairDescriptionGroups;
+		FGroomBuilder::BuildHairDescriptionGroups(HairDescription, HairDescriptionGroups);
+
+		// Remap group settings when reimporting the asset and group number/order mismatch
+		TArray<FHairGroupSettings> NewSettings;
+		TArray<int32> OldToNewGroupIndexMapping;
+		RemapHairGroupSettings(
+			ExistingHair, 
+			HairDescriptionGroups, 
+			ImportContext.ImportOptions->InterpolationSettings,
+			NewSettings, 
+			OldToNewGroupIndexMapping);
+
+		// 3. Apply new settings to the new asset
+		OutHairAsset->ClearNumGroup(GroupCount);
+		SetHairGroupSettings(OutHairAsset, NewSettings, OldToNewGroupIndexMapping);
+	}
+	else
+	{
+		OutHairAsset->SetNumGroup(GroupCount);
+
+		// Populate the interpolation settings with the new settings from the importer	
+		for (uint32 GroupIndex = 0; GroupIndex < GroupCount; ++GroupIndex)
+		{
+			// Insure the interpolation settings matches between the importer and the actual asset
+			const FHairGroupsInterpolation& InterpolationSettings = ImportContext.ImportOptions->InterpolationSettings[GroupIndex];
+			OutHairAsset->GetHairGroupsInterpolation()[GroupIndex] = InterpolationSettings;
+		}
+	}
 
 	// Sanity check
 	check(OutHairAsset->AreGroupsValid());
@@ -43,14 +228,6 @@ UGroomAsset* FHairStrandsImporter::ImportHair(const FHairImportContext& ImportCo
 
 	OutHairAsset->CommitHairDescription(MoveTemp(HairDescription), UGroomAsset::EHairDescriptionType::Source);
 
-	// Populate the interpolation settings with the new settings from the importer	
-	for (uint32 GroupIndex = 0; GroupIndex < GroupCount; ++GroupIndex)
-	{
-		// Insure the interpolation settings matches between the importer and the actual asset
-		const FHairGroupsInterpolation& InterpolationSettings = ImportContext.ImportOptions->InterpolationSettings[GroupIndex];
-		OutHairAsset->GetHairGroupsInterpolation()[GroupIndex] = InterpolationSettings;
-	}
-
 	const bool bSucceeded = OutHairAsset->CacheDerivedDatas();
 	if (!bSucceeded)
 	{
diff --git a/Engine/Plugins/Runtime/HairStrands/Source/HairStrandsEditor/Private/ReimportHairStrandsFactory.cpp b/Engine/Plugins/Runtime/HairStrands/Source/HairStrandsEditor/Private/ReimportHairStrandsFactory.cpp
--- a/Engine/Plugins/Runtime/HairStrands/Source/HairStrandsEditor/Private/ReimportHairStrandsFactory.cpp
+++ b/Engine/Plugins/Runtime/HairStrands/Source/HairStrandsEditor/Private/ReimportHairStrandsFactory.cpp
@@ -258,6 +258,11 @@ namespace ReimportGroomCacheHelpers
 	}
 }
 
+void RemapHairGroupInteprolationSettings(
+	const UGroomAsset* InGroomAsset, 
+	const FHairDescriptionGroups& InHairDescriptionGroups,
+	TArray<FHairGroupsInterpolation>& OutInterpolationSettings);
+
 EReimportResult::Type UReimportHairStrandsFactory::Reimport(UObject* Obj)
 {
 	// The reimport can start either from a GroomAsset or from a GroomCache
@@ -353,21 +358,8 @@ EReimportResult::Type UReimportHairStrandsFactory::Reimport(UObject* Obj)
 
 		FGroomBuilder::BuildHairDescriptionGroups(HairDescription, OutDescription);
 
-		// Populate the interpolation settings based on the group count from the description
-		const uint32 GroupCount = OutDescription.HairGroups.Num();
-		if (uint32(HairAsset->GetNumHairGroups()) != GroupCount)
-		{
-			// When reimporting an asset, if the number of groups has changed, we try to preserve existing settings
-			HairAsset->SetNumGroup(GroupCount, true /*bResetGroupData*/, false /*bResetOtherData*/);
-		}
-
-		// Initialized the settings value based on the assets interpolation options 
-		// (if the group are not resize above, the settings are preserved)
-		GroomReimportOptions->InterpolationSettings.Init(FHairGroupsInterpolation(), GroupCount);
-		for (uint32 GroupIndex = 0; GroupIndex < GroupCount; ++GroupIndex)
-		{
-			GroomReimportOptions->InterpolationSettings[GroupIndex] = HairAsset->GetHairGroupsInterpolation()[GroupIndex];
-		}
+		// Remap existing interpolation settings based on GroupName if possible, otherwise initialize them to default
+		RemapHairGroupInteprolationSettings(HairAsset, OutDescription, GroomReimportOptions->InterpolationSettings);
 	}
 
 	FGroomCacheImporter::SetupImportSettings(GroomCacheReimportOptions->ImportSettings, AnimInfo);
