4-Groom dataflow support

#virtualized

diff --git a/Engine/Plugins/Experimental/ChaosCaching/Source/ChaosCaching/Public/Chaos/ChaosCache.h b/Engine/Plugins/Experimental/ChaosCaching/Source/ChaosCaching/Public/Chaos/ChaosCache.h
--- a/Engine/Plugins/Experimental/ChaosCaching/Source/ChaosCaching/Public/Chaos/ChaosCache.h
+++ b/Engine/Plugins/Experimental/ChaosCaching/Source/ChaosCaching/Public/Chaos/ChaosCache.h
@@ -6,6 +6,7 @@
 #include "CacheEvents.h"
 #include "Containers/Queue.h"
 #include "Curves/RichCurve.h"
+#include "UObject/Interface.h"
 
 #include "Chaos/ChaosCacheInterpolationMode.h"
 
@@ -295,6 +296,21 @@ private:
 	FCacheUserToken& operator=(FCacheUserToken&&) = delete;
 };
 
+/** Interface for the chaos cache data storage. @todo(chaos) : We should probably refactor the chaos cache to allow more custom storage */
+UINTERFACE(MinimalAPI, Blueprintable)
+class UChaosCacheData : public UInterface
+{
+	GENERATED_BODY()
+};
+class IChaosCacheData
+{
+	GENERATED_BODY()
+ 
+public:
+	/** Duration of the chaos cache */
+	virtual float GetDuration() const = 0;
+};
+
 UCLASS(Experimental, MinimalAPI)
 class UChaosCache : public UObject
 {
@@ -374,6 +390,12 @@ public:
 	 */
 	CHAOSCACHING_API const FCacheSpawnableTemplate& GetSpawnableTemplate() const;
 
+	/** Get the cache data asset if defined */
+	IChaosCacheData* GetCacheData() const {return CacheData.GetInterface();}
+
+	/** Set the cache data asset interface */
+	void SetCacheData(IChaosCacheData* InCacheData) {CacheData.SetObject(CastChecked<UObject>(InCacheData)); CacheData.SetInterface(InCacheData);}
+
 	/**
 	 * Evaluates a single particle from the tracks array
 	 * @param InIndex Particle track index (unchecked, ensure valid before call)
@@ -448,6 +470,11 @@ public:
 	}
 
 private:
+
+	/** Optional cache data to store on the chaos cache */
+	UPROPERTY()
+	TScriptInterface<IChaosCacheData> CacheData;
+	
 	CHAOSCACHING_API void FlushPendingFrames_ChannelOnlyReservePass(TQueue<FPendingFrameWrite, EQueueMode::Spsc>& LocalPendingWrites,
 		bool& bCanSimpleCopyChannelData);
 
diff --git a/Engine/Plugins/Experimental/Dataflow/Source/DataflowEditor/Private/Dataflow/DataflowSimulationScene.cpp b/Engine/Plugins/Experimental/Dataflow/Source/DataflowEditor/Private/Dataflow/DataflowSimulationScene.cpp
--- a/Engine/Plugins/Experimental/Dataflow/Source/DataflowEditor/Private/Dataflow/DataflowSimulationScene.cpp
+++ b/Engine/Plugins/Experimental/Dataflow/Source/DataflowEditor/Private/Dataflow/DataflowSimulationScene.cpp
@@ -308,7 +308,7 @@ FBox FDataflowSimulationScene::GetBoundingBox() const
 				}
 			}
 		}
-		else
+		else if(PreviewActor)
 		{
 			constexpr bool bIncludeNonColliding = true;
 			SceneBounds += PreviewActor->GetComponentsBoundingBox(bIncludeNonColliding);
diff --git a/Engine/Plugins/Runtime/HairStrands/HairStrands.uplugin b/Engine/Plugins/Runtime/HairStrands/HairStrands.uplugin
--- a/Engine/Plugins/Runtime/HairStrands/HairStrands.uplugin
+++ b/Engine/Plugins/Runtime/HairStrands/HairStrands.uplugin
@@ -57,6 +57,10 @@
 		{
 			"Name": "DeformerGraph",
 			"Enabled": true
+		},
+		{
+			"Name": "Dataflow",
+			"Enabled": true
 		}
 	]
 }
diff --git a/Engine/Plugins/Runtime/HairStrands/Source/HairStrandsCore/HairStrandsCore.Build.cs b/Engine/Plugins/Runtime/HairStrands/Source/HairStrandsCore/HairStrandsCore.Build.cs
--- a/Engine/Plugins/Runtime/HairStrands/Source/HairStrandsCore/HairStrandsCore.Build.cs
+++ b/Engine/Plugins/Runtime/HairStrands/Source/HairStrandsCore/HairStrandsCore.Build.cs
@@ -18,6 +18,8 @@ namespace UnrealBuildTool.Rules
 					"MovieScene",
 					"NiagaraCore",
 					"NiagaraShader",
+					"DataflowCore",
+					"DataflowEngine",
 					"RenderCore",
 					"Renderer",
 					"VectorVM",
@@ -25,12 +27,15 @@ namespace UnrealBuildTool.Rules
 					"TraceLog",
 					"StaticMeshDescription",
 					"Eigen",
+					"ChaosCore",
+					"Chaos"
 				});
 
 			PublicDependencyModuleNames.AddRange(
 				new string[]
 				{
-					"Niagara"
+					"Niagara",
+					"ChaosCaching"
 				});
 
 			PrivateIncludePathModuleNames.AddRange(
diff --git a/Engine/Plugins/Runtime/HairStrands/Source/HairStrandsCore/Private/GroomAsset.cpp b/Engine/Plugins/Runtime/HairStrands/Source/HairStrandsCore/Private/GroomAsset.cpp
--- a/Engine/Plugins/Runtime/HairStrands/Source/HairStrandsCore/Private/GroomAsset.cpp
+++ b/Engine/Plugins/Runtime/HairStrands/Source/HairStrandsCore/Private/GroomAsset.cpp
@@ -1611,7 +1611,7 @@ void UGroomAsset::PostEditChangeProperty(FPropertyChangedEvent& PropertyChangedE
 			CachedHairGroupsMeshes.Init(Dirty, GroupCount);
 		}
 	}
-	if(PropertyName == GET_MEMBER_NAME_CHECKED(FHairSolverSettings, NiagaraSolver))
+	else if(PropertyName == GET_MEMBER_NAME_CHECKED(FHairSolverSettings, NiagaraSolver))
 	{
 		UpdatePhysicsSystems();
 	}
@@ -3476,8 +3476,9 @@ bool UGroomAsset::IsVisible(int32 GroupIndex, int32 LODIndex) const
 bool UGroomAsset::IsDeformationEnable(int32 GroupIndex) const
 {
 	return 
-		GetHairGroupsInterpolation().IsValidIndex(GroupIndex) && 
-		GetHairGroupsInterpolation()[GroupIndex].InterpolationSettings.GuideType == EGroomGuideType::Rigged;
+		(GetHairGroupsInterpolation().IsValidIndex(GroupIndex) && 
+		GetHairGroupsInterpolation()[GroupIndex].InterpolationSettings.GuideType == EGroomGuideType::Rigged) ||
+		(GetHairGroupsPhysics().IsValidIndex(GroupIndex) &&	GetHairGroupsPhysics()[GroupIndex].SolverSettings.bEnableDeformation);
 }
 
 bool UGroomAsset::IsSimulationEnable(int32 GroupIndex, int32 LODIndex) const 
@@ -3899,6 +3900,8 @@ DEFINE_GROOM_MEMBER_ACCESSOR_ARRAY(TArray<float>, EffectiveLODBias);
 DEFINE_GROOM_MEMBER_ACCESSOR_ARRAY(TArray<FHairGroupPlatformData>, HairGroupsPlatformData);
 DEFINE_GROOM_MEMBER_ACCESSOR_ARRAY(TArray<FHairGroupInfoWithVisibility>, HairGroupsInfo);
 
+GROOMASSET_DEFINE_MEMBER_NAME(DataflowSettings)
+
 const TArray<FHairGroupResources>& UGroomAsset::GetHairGroupsResources() const
 {
 	return HairGroupsResources;
@@ -3950,3 +3953,41 @@ float UGroomAsset::GetAutoLODBias() const
 {
 	return AutoLODBias;
 }
+
+TObjectPtr<UDataflowBaseContent> UGroomAsset::CreateDataflowContent()
+{
+	TObjectPtr<UDataflowGroomContent> GroomContent = NewObject<UDataflowGroomContent>(this, UDataflowGroomContent::StaticClass());
+	GroomContent->SetIsSaved(false);
+
+	GroomContent->SetDataflowOwner(this);
+	GroomContent->SetTerminalAsset(this);
+
+	WriteDataflowContent(GroomContent);
+	
+	return GroomContent;
+}
+
+void UGroomAsset::WriteDataflowContent(const TObjectPtr<UDataflowBaseContent>& DataflowContent) const
+{
+	if(const TObjectPtr<UDataflowGroomContent> GroomContent = Cast<UDataflowGroomContent>(DataflowContent))
+	{
+		GroomContent->SetDataflowAsset(DataflowSettings.GetDataflowAsset());
+		GroomContent->SetDataflowTerminal(DataflowSettings.GetDataflowTerminal());
+
+#if WITH_EDITORONLY_DATA
+		GroomContent->SetAnimationAsset(DataflowSettings.GetPreviewAnimationAsset());
+		GroomContent->SetBindingAsset(DataflowSettings.GetPreviewBindingAsset());
+#endif
+	}
+}
+
+void UGroomAsset::ReadDataflowContent(const TObjectPtr<UDataflowBaseContent>& DataflowContent)
+{
+	if(const TObjectPtr<UDataflowGroomContent> GroomContent = Cast<UDataflowGroomContent>(DataflowContent))
+	{
+#if WITH_EDITORONLY_DATA
+		DataflowSettings.SetPreviewAnimationAsset( GroomContent->GetAnimationAsset());
+		DataflowSettings.SetPreviewBindingAsset(GroomContent->GetBindingAsset());
+#endif
+	}
+}
diff --git a/Engine/Plugins/Runtime/HairStrands/Source/HairStrandsCore/Private/GroomAssetDataflow.cpp b/Engine/Plugins/Runtime/HairStrands/Source/HairStrandsCore/Private/GroomAssetDataflow.cpp
new file mode 100644
--- /dev/null
+++ b/Engine/Plugins/Runtime/HairStrands/Source/HairStrandsCore/Private/GroomAssetDataflow.cpp
@@ -0,0 +1,115 @@
+﻿// Copyright Epic Games, Inc. All Rights Reserved.
+
+#include "GroomAssetDataflow.h"
+#include "GroomBindingAsset.h"
+#include "Animation/AnimationAsset.h"
+
+#include UE_INLINE_GENERATED_CPP_BY_NAME(GroomAssetDataflow)
+
+
+UDataflowGroomContent::UDataflowGroomContent() : Super()
+{
+	bHideSkeletalMesh = true;
+	bHideAnimationAsset = false;
+}
+
+void UDataflowGroomContent::SetBindingAsset(const TObjectPtr<UGroomBindingAsset>& InBindingAsset)
+{
+	BindingAsset = InBindingAsset;
+	if(BindingAsset)
+	{
+		if(SkeletalMesh != BindingAsset->GetTargetSkeletalMesh())
+		{
+			SetSkeletalMesh(BindingAsset->GetTargetSkeletalMesh(), true);
+		}
+	}
+	SetConstructionDirty(true);
+	SetSimulationDirty(true);
+}
+
+#if WITH_EDITOR
+
+void UDataflowGroomContent::PostEditChangeProperty(FPropertyChangedEvent& PropertyChangedEvent) 
+{
+	const FName PropertyName = PropertyChangedEvent.Property->GetFName();
+
+	if (PropertyName == GET_MEMBER_NAME_CHECKED(UDataflowGroomContent, BindingAsset))
+	{
+		SetBindingAsset(BindingAsset);
+	}
+	Super::PostEditChangeProperty(PropertyChangedEvent);
+}
+
+#endif //if WITH_EDITOR
+
+void UDataflowGroomContent::AddReferencedObjects(UObject* InThis, FReferenceCollector& Collector)
+{
+	Super::AddReferencedObjects(InThis, Collector);
+	UDataflowGroomContent* This = CastChecked<UDataflowGroomContent>(InThis);
+	Collector.AddReferencedObject(This->BindingAsset);
+	Super::AddReferencedObjects(InThis, Collector);
+}
+
+void UDataflowGroomContent::SetActorProperties(TObjectPtr<AActor>& PreviewActor) const
+{
+	Super::SetActorProperties(PreviewActor);
+	OverrideActorProperty(PreviewActor, BindingAsset, TEXT("BindingAsset"));
+}
+
+#if WITH_EDITORONLY_DATA
+
+void FGroomDataflowSettings::SetPreviewBindingAsset(UGroomBindingAsset* BindingAsset)
+{
+	PreviewBindingAsset = BindingAsset;
+}
+
+UGroomBindingAsset* FGroomDataflowSettings::GetPreviewBindingAsset() const
+{
+	// Load the binding asset if it's not already loaded
+	return PreviewBindingAsset.LoadSynchronous();
+}
+
+void FGroomDataflowSettings::SetPreviewAnimationAsset(UAnimationAsset* AnimationAsset)
+{
+	PreviewAnimationAsset = AnimationAsset;
+}
+
+UAnimationAsset* FGroomDataflowSettings::GetPreviewAnimationAsset() const
+{
+	// Load the animation asset if it's not already loaded
+	return PreviewAnimationAsset.LoadSynchronous();
+}
+
+#endif
+
+FGroomDataflowSettings::FGroomDataflowSettings()
+{
+	RestCollection = MakeShared<FManagedArrayCollection>();
+}
+
+FName FGroomDataflowSettings::GetDataflowAssetMemberName()
+{
+	return FName("DataflowAsset");
+}
+
+FName FGroomDataflowSettings::GetDataflowTerminalMemberName()
+{
+	return FName("DataflowTerminal");
+}
+
+bool FGroomDataflowSettings::Serialize(FArchive& Ar)
+{
+	// Serialize normal tagged property data
+	if (Ar.IsLoading() || Ar.IsSaving())
+	{
+		UScriptStruct* const Struct = FGroomDataflowSettings::StaticStruct();
+		Struct->SerializeTaggedProperties(Ar, reinterpret_cast<uint8*>(this), Struct, nullptr);
+	}
+	if(RestCollection.IsValid())
+	{
+		Chaos::FChaosArchive ChaosAr(Ar);
+		RestCollection->Serialize(ChaosAr);
+	}
+	
+	return true;
+}
diff --git a/Engine/Plugins/Runtime/HairStrands/Source/HairStrandsCore/Private/GroomAssetPhysics.cpp b/Engine/Plugins/Runtime/HairStrands/Source/HairStrandsCore/Private/GroomAssetPhysics.cpp
--- a/Engine/Plugins/Runtime/HairStrands/Source/HairStrandsCore/Private/GroomAssetPhysics.cpp
+++ b/Engine/Plugins/Runtime/HairStrands/Source/HairStrandsCore/Private/GroomAssetPhysics.cpp
@@ -8,6 +8,7 @@
 FHairSolverSettings::FHairSolverSettings()
 {
 	EnableSimulation = false;
+	bEnableDeformation = false;
 	NiagaraSolver = EGroomNiagaraSolvers::AngularSprings;
 	SubSteps = 5;
 	IterationCount = 5;
@@ -88,6 +89,7 @@ bool FHairSolverSettings::operator==(const FHairSolverSettings& A) const
 {
 	return 
 		EnableSimulation == A.EnableSimulation &&
+		bEnableDeformation == A.bEnableDeformation &&
 		NiagaraSolver == A.NiagaraSolver &&
 		SubSteps == A.SubSteps &&
 		IterationCount == A.IterationCount &&
diff --git a/Engine/Plugins/Runtime/HairStrands/Source/HairStrandsCore/Private/GroomBuilder.cpp b/Engine/Plugins/Runtime/HairStrands/Source/HairStrandsCore/Private/GroomBuilder.cpp
--- a/Engine/Plugins/Runtime/HairStrands/Source/HairStrandsCore/Private/GroomBuilder.cpp
+++ b/Engine/Plugins/Runtime/HairStrands/Source/HairStrandsCore/Private/GroomBuilder.cpp
@@ -12,6 +12,7 @@
 #include "GlobalShader.h"
 #include "RHIGPUReadback.h"
 #include "ShaderParameterStruct.h"
+#include "SceneView.h"
 
 #include "Async/ParallelFor.h"
 #include "HAL/IConsoleManager.h"
@@ -3883,43 +3884,57 @@ IMPLEMENT_GLOBAL_SHADER(FHairStrandCopyPositionCS, "/Engine/Private/HairStrands/
 static FRDGBufferRef ReadPositions(
 	FRDGBuilder& GraphBuilder,
 	FGlobalShaderMap* ShaderMap,
-	const FHairGroupInstance* Instance)
+	const FHairGroupInstance* Instance,
+	const bool bReadGuides, bool& bHasValidResource)
 {
 	check(Instance);
-
-	FRDGBufferRef OutPositions = nullptr;
-	const uint32 NumVertices = Instance->HairGroupPublicData->GetActiveStrandsPointCount();
+	const uint32 NumVertices = bReadGuides ? Instance->Guides.GetData().GetNumPoints() : Instance->HairGroupPublicData->GetActiveStrandsPointCount();
+	
+	auto ReadData = [](FRDGBuilder& GraphBuilder, FHairStrandsRestResource* InRest, FHairStrandsDeformedResource* InDeformedResource, FRDGBufferSRVRef& InPositions, FRDGBufferSRVRef& InPositionOffset)
+	{
+		if (InDeformedResource)
+		{
+			InPositions = Register(GraphBuilder, InDeformedResource->GetBuffer(FHairStrandsDeformedResource::EFrameType::Current), ERDGImportedBufferFlags::CreateSRV).SRV;
+			InPositionOffset = Register(GraphBuilder, InDeformedResource->GetPositionOffsetBuffer(FHairStrandsDeformedResource::EFrameType::Current), ERDGImportedBufferFlags::CreateSRV).SRV;
+		}
+		else
+		{
+			InPositions = Register(GraphBuilder, InRest->PositionBuffer, ERDGImportedBufferFlags::CreateSRV).SRV;
+			InPositionOffset = Register(GraphBuilder, InRest->PositionOffsetBuffer, ERDGImportedBufferFlags::CreateSRV).SRV;
+		}
+	};
 
 	FRDGBufferSRVRef InPositions = nullptr;
 	FRDGBufferSRVRef InPositionOffset = nullptr;
-	if (Instance->Strands.DeformedResource)
+	if (bReadGuides)
 	{
-		InPositions = Register(GraphBuilder, Instance->Strands.DeformedResource->GetBuffer(FHairStrandsDeformedResource::EFrameType::Current), ERDGImportedBufferFlags::CreateSRV).SRV;
-		InPositionOffset = Register(GraphBuilder, Instance->Strands.DeformedResource->GetPositionOffsetBuffer(FHairStrandsDeformedResource::EFrameType::Current), ERDGImportedBufferFlags::CreateSRV).SRV;
+		ReadData(GraphBuilder, Instance->Guides.RestResource,  Instance->Guides.DeformedResource, InPositions, InPositionOffset);
 	}
 	else
 	{
-		InPositions = Register(GraphBuilder, Instance->Strands.RestResource->PositionBuffer, ERDGImportedBufferFlags::CreateSRV).SRV;
-		InPositionOffset = Register(GraphBuilder, Instance->Strands.RestResource->PositionOffsetBuffer, ERDGImportedBufferFlags::CreateSRV).SRV;
+		ReadData(GraphBuilder, Instance->Strands.RestResource, Instance->Strands.DeformedResource, InPositions, InPositionOffset);
+	}
+	
+	FRDGBufferRef OutPositions = GraphBuilder.CreateBuffer(FRDGBufferDesc::CreateBufferDesc(sizeof(FVector4f), NumVertices), TEXT("StrandsPositionForReadback"));
+	bHasValidResource = InPositions && InPositionOffset;
+	if(bHasValidResource)
+	{
+		FHairStrandCopyPositionCS::FParameters* Parameters = GraphBuilder.AllocParameters<FHairStrandCopyPositionCS::FParameters>();
+		Parameters->MaxVertexCount = NumVertices;
+		Parameters->InPositionOffset = InPositionOffset;
+		Parameters->InPositions = InPositions;
+		Parameters->OutPositions = GraphBuilder.CreateUAV(OutPositions, PF_A32B32G32R32F);
+
+		FIntVector DispatchCount = FComputeShaderUtils::GetGroupCount(NumVertices, 128);
+
+		TShaderMapRef<FHairStrandCopyPositionCS> ComputeShader(ShaderMap);
+		FComputeShaderUtils::AddPass(
+			GraphBuilder,
+			RDG_EVENT_NAME("HairStrands::CopyPositions"),
+			ComputeShader,
+			Parameters,
+			DispatchCount);
 	}
-
-	OutPositions = GraphBuilder.CreateBuffer(FRDGBufferDesc::CreateBufferDesc(sizeof(FVector4f), NumVertices), TEXT("StrandsPositionForReadback"));
-
-	FHairStrandCopyPositionCS::FParameters* Parameters = GraphBuilder.AllocParameters<FHairStrandCopyPositionCS::FParameters>();
-	Parameters->MaxVertexCount = NumVertices;
-	Parameters->InPositionOffset = InPositionOffset;
-	Parameters->InPositions = InPositions;
-	Parameters->OutPositions = GraphBuilder.CreateUAV(OutPositions, PF_A32B32G32R32F);
-
-	FIntVector DispatchCount = FComputeShaderUtils::GetGroupCount(NumVertices, 128);
-
-	TShaderMapRef<FHairStrandCopyPositionCS> ComputeShader(ShaderMap);
-	FComputeShaderUtils::AddPass(
-		GraphBuilder,
-		RDG_EVENT_NAME("HairStrands::CopyPositions"),
-		ComputeShader,
-		Parameters,
-		DispatchCount);
 
 	return OutPositions;
 }
@@ -3933,7 +3948,8 @@ struct FStrandsPositionData
 	uint32 CurveCount = 0;
 	uint32 PointCount = 0;
 	const UGroomComponent* Component = nullptr;
-	FStrandsPositionOutput* Output;
+	TSharedPtr<FStrandsPositionOutput> Output;
+	bool bReadGuides = false;
 };
 
 TQueue<FStrandsPositionData*>		GStrandsPositionQueries;
@@ -3944,7 +3960,7 @@ bool HasHairStrandsPositionQueries()
 	return !GStrandsPositionQueries.IsEmpty() || !GStrandsPositionReadbacks.IsEmpty();
 }
 
-void RunHairStrandsPositionQueries(FRDGBuilder& GraphBuilder, FGlobalShaderMap* ShaderMap, const struct FShaderPrintData* DebugShaderData)
+void RunHairStrandsPositionQueries(FRDGBuilder& GraphBuilder, FGlobalShaderMap* ShaderMap, const struct FShaderPrintData* DebugShaderData, const FSceneView* SceneView)
 {
 	// Operations are ordered in reverse to ensure they are processed on independent frames to avoid TDR on heavy grooms.
 
@@ -3953,43 +3969,50 @@ void RunHairStrandsPositionQueries(FRDGBuilder& GraphBuilder, FGlobalShaderMap*
 		TArray<FStrandsPositionData*> QueryNotReady;
 		FStrandsPositionData* Q = nullptr;
 		while (GStrandsPositionReadbacks.Dequeue(Q))
-		{			
-			if (Q->GPUPositions->IsReady() || Q->ReadbackDelay == 0)
+		{
+			if(Q && Q->Component && Q->Component->GetWorld() &&
+				SceneView && SceneView->Family && (Q->Component->GetWorld()->Scene == SceneView->Family->Scene))
 			{
-				const FHairGroupInstance* Instance = Q->Component->GetGroupInstance(Q->GroupIt);
-
-				check(Q->GroupIt < uint32(Q->Output->Groups.Num()));
-				FStrandsPositionOutput::FGroup& Group = Q->Output->Groups[Q->GroupIt];
-				Group.Reserve(Q->CurveCount);
-
-				FStrandsPositionOutput::FStrand Strand;
-				const FVector4f* Positions = (const FVector4f*)Q->GPUPositions->Lock(sizeof(FVector4f) * Q->PointCount);
-				for (uint32 PointIt = 0; PointIt < Q->PointCount; ++PointIt)
+				if (Q->GPUPositions->IsReady() || Q->ReadbackDelay == 0)
 				{
-					const FVector4f P = Positions[PointIt];
-					check(P.W != HAIR_CONTROLPOINT_START || Strand.Num() == 0);
-					check(P.W != HAIR_CONTROLPOINT_INSIDE || Strand.Num() > 0);
+					const FHairGroupInstance* Instance = Q->Component->GetGroupInstance(Q->GroupIt);
 
-					Strand.Add(FVector3f(P.X, P.Y, P.Z));
-					if (P.W == HAIR_CONTROLPOINT_END)
+					check(Q->GroupIt < uint32(Q->Output->Groups.Num()));
+					FStrandsPositionOutput::FGroup& Group = Q->Output->Groups[Q->GroupIt];
+					Group.Reserve(Q->CurveCount);
+		
+					FStrandsPositionOutput::FStrand Strand;
+					const FVector4f* Positions = (const FVector4f*)Q->GPUPositions->Lock(sizeof(FVector4f) * Q->PointCount);
+					for (uint32 PointIt = 0; PointIt < Q->PointCount; ++PointIt)
 					{
-						Group.Add(Strand);
-						Strand.SetNum(0);
+						const FVector4f P = Positions[PointIt];
+						check(P.W != HAIR_CONTROLPOINT_START || Strand.Num() == 0);
+						check(P.W != HAIR_CONTROLPOINT_INSIDE || Strand.Num() > 0);
+
+						Strand.Add(FVector3f(P.X, P.Y, P.Z));
+						if (P.W == HAIR_CONTROLPOINT_END)
+						{
+							Group.Add(Strand);
+							Strand.SetNum(0);
+						}
 					}
-				}
-				check(Q->CurveCount == Group.Num());
-				Q->GPUPositions->Unlock();
+					check(Q->CurveCount == Group.Num());
+					Q->GPUPositions->Unlock();
 
-				--Q->Output->Status;
-				check(Q->Output->Status >= 0);
+					--Q->Output->Status;
+					check(Q->Output->Status >= 0);
+				}
+				else
+				{
+					if (Q->ReadbackDelay > 0) { --Q->ReadbackDelay; }
+					QueryNotReady.Add(Q);
+				}
 			}
 			else
 			{
-				if (Q->ReadbackDelay > 0) { --Q->ReadbackDelay; }
 				QueryNotReady.Add(Q);
 			}
 		}
-
 		for (FStrandsPositionData* N : QueryNotReady)
 		{
 			GStrandsPositionReadbacks.Enqueue(N);
@@ -3998,32 +4021,54 @@ void RunHairStrandsPositionQueries(FRDGBuilder& GraphBuilder, FGlobalShaderMap*
 
 	// 1. Copy positions
 	{
-		FStrandsPositionData* Q;
+		TArray<FStrandsPositionData*> QueryNotReady;
+		FStrandsPositionData* Q = nullptr;
 		while (GStrandsPositionQueries.Dequeue(Q))
 		{
-			check(Q->Component);
-			if (const FHairGroupInstance* Instance = Q->Component->GetGroupInstance(Q->GroupIt))
+			if(Q && Q->Component && Q->Component->GetWorld() &&
+				SceneView && SceneView->Family && (Q->Component->GetWorld()->Scene == SceneView->Family->Scene))
 			{
-				Q->CurveCount = Instance->HairGroupPublicData->GetActiveStrandsCurveCount();
-				Q->PointCount = Instance->HairGroupPublicData->GetActiveStrandsPointCount();
+				check(Q->Component);
+				if (const FHairGroupInstance* Instance = Q->Component->GetGroupInstance(Q->GroupIt))
+				{
+					Q->CurveCount = Q->bReadGuides ? Instance->Guides.GetData().GetNumCurves() : Instance->HairGroupPublicData->GetActiveStrandsCurveCount();
+					Q->PointCount = Q->bReadGuides ? Instance->Guides.GetData().GetNumPoints() : Instance->HairGroupPublicData->GetActiveStrandsPointCount();
 
-				FRDGBufferRef Positions = ReadPositions(GraphBuilder, ShaderMap, Instance);
-				AddEnqueueCopyPass(GraphBuilder, Q->GPUPositions.Get(), Positions, Positions->Desc.GetSize());
-				GStrandsPositionReadbacks.Enqueue(Q);
+					bool bHasValidResources = true;
+					FRDGBufferRef Positions = ReadPositions(GraphBuilder, ShaderMap, Instance, Q->bReadGuides, bHasValidResources);
+					if(bHasValidResources)
+					{
+						AddEnqueueCopyPass(GraphBuilder, Q->GPUPositions.Get(), Positions, Positions->Desc.GetSize());
+						GStrandsPositionReadbacks.Enqueue(Q);
+					}
+					else
+					{
+						--Q->Output->Status;
+						check(Q->Output->Status >= 0);
+					}
+				}
+				else
+				{
+					// If the instance is missing, Q->Output->Groups[GroupIt] remains an empty array
+					// in this case, skip the readback to prevent the query from remaining unfinished
+					--Q->Output->Status;
+					check(Q->Output->Status >= 0);
+				}
 			}
 			else
 			{
-				// If the instance is missing, Q->Output->Groups[GroupIt] remains an empty array
-				// in this case, skip the readback to prevent the query from remaining unfinished
-				--Q->Output->Status;
-				check(Q->Output->Status >= 0);
+				QueryNotReady.Add(Q);
 			}
 		}
+		for (FStrandsPositionData* N : QueryNotReady)
+		{
+			GStrandsPositionQueries.Enqueue(N);
+		}
 	}
 }
 
 #if WITH_EDITOR
-bool RequestStrandsPosition(const UGroomComponent* Component, FStrandsPositionOutput* Output)
+bool RequestStrandsPosition(const UGroomComponent* Component, TSharedPtr<FStrandsPositionOutput> Output, const bool bReadGuides)
 {
 	if (Component == nullptr || Output == nullptr)
 	{
@@ -4042,9 +4087,10 @@ bool RequestStrandsPosition(const UGroomComponent* Component, FStrandsPositionOu
 		Data->PointCount = 0;
 		Data->GPUPositions = MakeUnique<FRHIGPUBufferReadback>(TEXT("Readback.Positions"));
 		Data->Output = Output;
+		Data->bReadGuides = bReadGuides;
 		GStrandsPositionQueries.Enqueue(Data);
 	}
-
+	
 	*Output = FStrandsPositionOutput();
 	Output->Component = Component;
 	Output->Status = GroupCount;
diff --git a/Engine/Plugins/Runtime/HairStrands/Source/HairStrandsCore/Private/GroomCache.cpp b/Engine/Plugins/Runtime/HairStrands/Source/HairStrandsCore/Private/GroomCache.cpp
--- a/Engine/Plugins/Runtime/HairStrands/Source/HairStrandsCore/Private/GroomCache.cpp
+++ b/Engine/Plugins/Runtime/HairStrands/Source/HairStrandsCore/Private/GroomCache.cpp
@@ -5,6 +5,7 @@
 #include "EditorFramework/AssetImportData.h"
 #include "Engine/AssetUserData.h"
 #include "GroomAsset.h"
+#include "GroomBuilder.h"
 #include "Serialization/MemoryReader.h"
 #include "Serialization/MemoryWriter.h"
 #include "UObject/AssetRegistryTagsContext.h"
@@ -245,6 +246,13 @@ bool UGroomCache::GetGroomDataAtFrameIndex(int32 FrameIndex, FGroomCacheAnimatio
 			Ar.SetUEVer(ArchiveVersion.GetValue());
 		}
 		AnimData.Serialize(Ar);
+
+		// Reset the pointers in order to trigger a read next time
+		for( auto& GroupData : AnimData.GroupsData)
+		{
+			GroupData.VertexData.PositionBuffer = nullptr;
+			GroupData.VertexData.RadiusBuffer = nullptr;
+		}
 	}
 
 	return true;
@@ -353,3 +361,43 @@ void FGroomCacheProcessor::TransferChunks(UGroomCache* GroomCache)
 	GroomCache->Chunks = MoveTemp(Chunks);
 }
 
+namespace UE::Groom
+{
+	void BuildGroomCache(FGroomCacheProcessor& Processor, const FGroomAnimationInfo& AnimInfo, UGroomCache* GroomCache)
+	{
+		if (GroomCache)
+		{
+			Processor.TransferChunks(GroomCache);
+			GroomCache->SetGroomAnimationInfo(AnimInfo);
+
+			GroomCache->MarkPackageDirty();
+#if WITH_EDITOR
+			GroomCache->PostEditChange();
+#endif
+		}
+	}
+
+	bool BuildGroupsData(const FHairDescription& HairDescription, const TArray<FHairGroupPlatformData>& PlatformData,
+		TArray<FHairGroupInfoWithVisibility>& VisibilityData, const TArray<FHairGroupsInterpolation>& InterpolationData, TArray<FGroomCacheInputData>& GroupsData)
+	{
+		bool bSuccess = true;
+		FHairDescriptionGroups DescriptionGroups;
+		// Do not add extra control points at the end of curve when hair strip geometry is enabled. This is because groom cache data are 
+		// serialize within the uasset (i.e. do not used intermediate cached/build data), and need the asset to be compatible with or without hair strip geometry.
+		if (FGroomBuilder::BuildHairDescriptionGroups(HairDescription, DescriptionGroups, false /*bAllowAddEndControlPoint*/))
+		{
+			const uint32 GroupCount = DescriptionGroups.HairGroups.Num();
+			if(GroupCount == PlatformData.Num() && GroupCount == VisibilityData.Num() && GroupCount == InterpolationData.Num())
+			{
+				GroupsData.SetNum(GroupCount);
+				for (uint32 GroupIndex = 0; GroupIndex < GroupCount; ++GroupIndex)
+				{
+					FGroomBuilder::BuildData(DescriptionGroups.HairGroups[GroupIndex], InterpolationData[GroupIndex], VisibilityData[GroupIndex], GroupsData[GroupIndex].Strands, GroupsData[GroupIndex].Guides);
+				}
+				return true;
+			}
+		}
+		return false;
+	}
+}
+
diff --git a/Engine/Plugins/Runtime/HairStrands/Source/HairStrandsCore/Private/GroomCacheStreamingData.cpp b/Engine/Plugins/Runtime/HairStrands/Source/HairStrandsCore/Private/GroomCacheStreamingData.cpp
--- a/Engine/Plugins/Runtime/HairStrands/Source/HairStrandsCore/Private/GroomCacheStreamingData.cpp
+++ b/Engine/Plugins/Runtime/HairStrands/Source/HairStrandsCore/Private/GroomCacheStreamingData.cpp
@@ -169,7 +169,7 @@ void FGroomCacheStreamingData::UpdateStreamingStatus(bool bAsyncDeletionAllowed)
 			}
 
 			// If not requested yet, then request a load
-			if (!ChunksRequested.Contains(NeededIndex))
+			if (!ChunksRequested.Contains(NeededIndex) && GroomCache->GetChunks().IsValidIndex(NeededIndex))
 			{
 				const FGroomCacheChunk& Chunk = GroomCache->GetChunks()[NeededIndex];
 
@@ -452,9 +452,10 @@ void FGroomCacheStreamingData::UnmapAnimationData(uint32 ChunkIndex)
 	FScopeLock Lock(&CriticalSection);
 
 	FResidentGroomCacheChunk* ResidentChunk = Chunks.Find(ChunkIndex);
-	if (ResidentChunk != nullptr )
+	// If the unmap is done from RT and GT at the same time (init resource...) even if the critical section is preventing a race condition
+	// we could end up with already a refcount = 0 in that function
+	if (ResidentChunk != nullptr && ResidentChunk->Refcount > 0)
 	{
-		checkf(ResidentChunk->Refcount > 0, TEXT("Map/Unmap out of balance. Make sure you unmap once fore every map."));
 		checkf(ChunksAvailable.Contains(ChunkIndex) || ChunksEvicted.Contains(ChunkIndex), TEXT("Tried to unmap a chunk in an invalid state."));
 		--ResidentChunk->Refcount;
 	}
diff --git a/Engine/Plugins/Runtime/HairStrands/Source/HairStrandsCore/Private/GroomComponent.cpp b/Engine/Plugins/Runtime/HairStrands/Source/HairStrandsCore/Private/GroomComponent.cpp
--- a/Engine/Plugins/Runtime/HairStrands/Source/HairStrandsCore/Private/GroomComponent.cpp
+++ b/Engine/Plugins/Runtime/HairStrands/Source/HairStrandsCore/Private/GroomComponent.cpp
@@ -1343,7 +1343,7 @@ public:
 					IGroomCacheStreamingManager::Get().UnmapAnimationData(GroomCache, CurrentFrameIndex);
 					CurrentFramePtr = NextFramePtr;
 					CurrentFrameIndex = NextFrameIndex;
-				}
+				} 
 			}
 			else
 			{
@@ -1458,6 +1458,12 @@ bool UGroomComponent::IsSimulationEnable(int32 GroupIndex, int32 LODIndex) const
 	return bIsSimulationEnable;
 }
 
+bool UGroomComponent::IsDeformationEnable(int32 GroupIndex) const
+{
+	// To enable deformation groom asset should allow it and we must have a mesh deformer locally or being register to a solver component
+	return GroomAsset->IsDeformationEnable(GroupIndex) && (GroomSolver || MeshDeformer);
+}
+
 void UGroomComponent::ReleaseHairSimulation()
 {
 	for (int32 CompIndex = 0; CompIndex < NiagaraComponents.Num(); ++CompIndex)
@@ -2676,7 +2682,7 @@ void UGroomComponent::InitResources(bool bIsBindingReloading)
 			bHasAnyNeedSimulation				= bHasAnyNeedSimulation			|| bHasNeedSimulation[GroupIt];
 			bHasAnyNeedGlobalDeformation		= bHasAnyNeedGlobalDeformation	|| bHasNeedGlobalDeformation[GroupIt];
 		}
-		bHasNeedDeformation[GroupIt] = GroomAsset->IsDeformationEnable(GroupIt);
+		bHasNeedDeformation[GroupIt] = IsDeformationEnable(GroupIt);
 		bHasAnyNeedDeformation = bHasAnyNeedDeformation || bHasNeedDeformation[GroupIt];
 	}
 	const bool bHasNeedBindingData = bHasNeedSkinningBinding || bHasAnyNeedGlobalDeformation;
@@ -2768,6 +2774,7 @@ void UGroomComponent::InitResources(bool bIsBindingReloading)
 	if (GroomCache)
 	{
 		GroomCacheBuffers = GUseGroomCacheStreaming ? MakeShared<FGroomCacheStreamedBuffers, ESPMode::ThreadSafe>(GroomCache) : MakeShared<FGroomCacheBuffers, ESPMode::ThreadSafe>(GroomCache);
+		
 		UpdateGroomCache(ElapsedTime);
 	}
 	else
@@ -2807,6 +2814,7 @@ void UGroomComponent::InitResources(bool bIsBindingReloading)
 		HairGroupInstance->Debug.SkinningCurrentLocalToWorld	= RegisteredMeshComponent ? RegisteredMeshComponent->GetComponentTransform() : FTransform::Identity;
 		HairGroupInstance->Debug.SkinningPreviousLocalToWorld	= HairGroupInstance->Debug.SkinningCurrentLocalToWorld;
 		HairGroupInstance->Debug.RigidCurrentLocalToWorld		= GetComponentTransform();
+		HairGroupInstance->Debug.RigidRestLocalToWorld			= HairGroupInstance->Debug.RigidCurrentLocalToWorld;
 		HairGroupInstance->Debug.RigidPreviousLocalToWorld		= HairGroupInstance->Debug.RigidCurrentLocalToWorld;
 		HairGroupInstance->LocalToWorld							= HairGroupInstance->GetCurrentLocalToWorld();
 
@@ -2857,7 +2865,7 @@ void UGroomComponent::InitResources(bool bIsBindingReloading)
 				HairGroupInstance->HairGroupPublicData->LODGlobalInterpolations.Add(LODGlobalInterpolation);
 			}
 			HairGroupInstance->HairGroupPublicData->bIsSimulationCacheEnable = HairGroupInstance->Debug.GroomCacheType == EGroomCacheType::Guides;
-			HairGroupInstance->HairGroupPublicData->bIsDeformationEnable = GroomAsset->IsDeformationEnable(GroupIt);
+			HairGroupInstance->HairGroupPublicData->bIsDeformationEnable = IsDeformationEnable(GroupIt);
 			HairGroupInstance->HairGroupPublicData->SetLODScreenSizes(CPULODScreenSize);
 			HairGroupInstance->HairGroupPublicData->SetLODVisibilities(LODVisibility);
 			HairGroupInstance->HairGroupPublicData->SetLODGeometryTypes(LODGeometryTypes);
@@ -2913,7 +2921,7 @@ void UGroomComponent::InitResources(bool bIsBindingReloading)
 			const int32 LODIndex = -1;
 			HairGroupInstance->Guides.bIsSimulationEnable = IsSimulationEnable(GroupIt,LODIndex);
 			HairGroupInstance->Guides.bHasGlobalInterpolation = LocalBindingAsset && GroomAsset->IsGlobalInterpolationEnable(GroupIt,LODIndex);
-			HairGroupInstance->Guides.bIsDeformationEnable = GroomAsset->IsDeformationEnable(GroupIt);
+			HairGroupInstance->Guides.bIsDeformationEnable = IsDeformationEnable(GroupIt);
 			HairGroupInstance->Guides.bIsSimulationCacheEnable = HairGroupInstance->HairGroupPublicData->bIsSimulationCacheEnable;
 		}
 
@@ -2982,7 +2990,7 @@ void UGroomComponent::InitResources(bool bIsBindingReloading)
 				}
 
 				// Mesh deformer requires to dynamic resources
-				if (MeshDeformer)
+				if (MeshDeformer || GroomSolver)
 				{
 					bNeedDynamicResources = true;
 
@@ -3459,6 +3467,7 @@ void UGroomComponent::UpdateGroomCache(float Time)
 			{
 				if (UGroomComponent* Component = Cast<UGroomComponent>(WeakComponent.Get()))
 				{
+					
 					Component->MarkRenderTransformDirty();
 				}
 			});
@@ -3616,7 +3625,7 @@ void UGroomComponent::TickComponent(float DeltaTime, enum ELevelTick TickType, F
 				SetRelativeLocation(InvBoneTransform.GetLocation());
 				SetRelativeRotation(InvBoneTransform.GetRotation());
 			}
-		}
+		} 
 	}
 
 	bResetSimulation = bInitSimulation;
@@ -3774,6 +3783,7 @@ void UGroomComponent::SendRenderDynamicData_Concurrent()
 		UMeshDeformerInstance::FEnqueueWorkDesc Desc;
 		Desc.Scene = GetScene();
 		Desc.OwnerName = GroomAsset->GetFName();
+		Desc.ExecutionGroup = UMeshDeformerInstance::ExecutionGroup_BeginInitViews;
 		// TODO: Provide a FEnqueueWorkDesc::FallbackDelegate so that groom appears rest pose if the Enqueue fails.
 		MeshDeformerInstance->EnqueueWork(Desc);
 	}
diff --git a/Engine/Plugins/Runtime/HairStrands/Source/HairStrandsCore/Private/GroomEdit.cpp b/Engine/Plugins/Runtime/HairStrands/Source/HairStrandsCore/Private/GroomEdit.cpp
--- a/Engine/Plugins/Runtime/HairStrands/Source/HairStrandsCore/Private/GroomEdit.cpp
+++ b/Engine/Plugins/Runtime/HairStrands/Source/HairStrandsCore/Private/GroomEdit.cpp
@@ -59,6 +59,10 @@ static void HairStrandsDataToEditableHairStrands(const FHairStrandsDatas& In, TA
 		const uint32 PointCount  = In.StrandsCurves.CurvesCount[StrandsIt];
 		const uint32 PointOffset = In.StrandsCurves.CurvesOffset[StrandsIt];
 
+		const bool bHasValidColor = (In.StrandsPoints.PointsBaseColor.Num() == PointCount);
+		const bool bHasValidRoughness = (In.StrandsPoints.PointsRoughness.Num() == PointCount);
+		const bool bHasValidAO = (In.StrandsPoints.PointsAO.Num() == PointCount);
+
 		Out[StrandsIt].ControlPoints.SetNum(PointCount);
 		for (uint32 PointIt = 0; PointIt < PointCount; ++PointIt)
 		{				
@@ -66,9 +70,19 @@ static void HairStrandsDataToEditableHairStrands(const FHairStrandsDatas& In, TA
 			Out[StrandsIt].ControlPoints[PointIt].Position	= In.StrandsPoints.PointsPosition[EffectiveIndex];
 			Out[StrandsIt].ControlPoints[PointIt].U			= In.StrandsPoints.PointsCoordU[EffectiveIndex];
 			Out[StrandsIt].ControlPoints[PointIt].Radius	= In.StrandsPoints.PointsRadius[EffectiveIndex];
-			Out[StrandsIt].ControlPoints[PointIt].BaseColor	= In.StrandsPoints.PointsBaseColor[EffectiveIndex];
-			Out[StrandsIt].ControlPoints[PointIt].Roughness	= In.StrandsPoints.PointsRoughness[EffectiveIndex];
-			Out[StrandsIt].ControlPoints[PointIt].AO		= In.StrandsPoints.PointsAO[EffectiveIndex];
+			
+			if(bHasValidColor)
+			{
+				Out[StrandsIt].ControlPoints[PointIt].BaseColor	= In.StrandsPoints.PointsBaseColor[EffectiveIndex];
+			}
+			if(bHasValidRoughness)
+			{
+				Out[StrandsIt].ControlPoints[PointIt].Roughness	= In.StrandsPoints.PointsRoughness[EffectiveIndex];
+			}
+			if(bHasValidAO)
+			{
+				Out[StrandsIt].ControlPoints[PointIt].AO		= In.StrandsPoints.PointsAO[EffectiveIndex];
+			}
 		}
 	}
 }
diff --git a/Engine/Plugins/Runtime/HairStrands/Source/HairStrandsCore/Private/GroomManager.cpp b/Engine/Plugins/Runtime/HairStrands/Source/HairStrandsCore/Private/GroomManager.cpp
--- a/Engine/Plugins/Runtime/HairStrands/Source/HairStrandsCore/Private/GroomManager.cpp
+++ b/Engine/Plugins/Runtime/HairStrands/Source/HairStrandsCore/Private/GroomManager.cpp
@@ -783,50 +783,36 @@ static void RunHairBindingSurfaceUpdate(
 }
 
 ///////////////////////////////////////////////////////////////////////////////////////////////////
-// Guide interpolation
+// Guide Instance Data
 
-static void RunHairStrandsInterpolation_Guide(
-	FRDGBuilder& GraphBuilder,
-	FSceneInterface* Scene,
-	const FSceneView* View,
-	const uint32 ViewUniqueID,
-	const FHairStrandsInstances& Instances,
-	const FShaderPrintData* ShaderPrintData,
-	FHairTransientResources& TransientResources,
-	FGlobalShaderMap* ShaderMap)
+struct FGuideInstanceData
 {
-	check(IsInRenderingThread());
+	int32 HairLODIndex = -1;
+	int32 MeshLODIndex = -1;
+	EGroomCacheType ActiveGroomCacheType = EGroomCacheType::None;
+	EHairBindingType BindingType = EHairBindingType::NoneBinding;
+	bool bSimulationEnable = false;
+	bool bDeformationEnable = false;
+	bool bGlobalDeformationEnable = false;
 
-	RDG_EVENT_SCOPE_STAT(GraphBuilder, HairGuideInterpolation, "HairGuideInterpolation");
-	RDG_GPU_STAT_SCOPE(GraphBuilder, HairGuideInterpolation);
+	FCachedGeometry MeshLODData;
+	FHairGroupInstance* Instance = nullptr;
 
-	// Update dynamic mesh triangles
-	struct FInstanceData
+	bool NeedsMeshUpdate(bool bCheckMeshLODIndex=true) const
 	{
-		int32 HairLODIndex = -1;
-		int32 MeshLODIndex = -1;
-		EGroomCacheType ActiveGroomCacheType = EGroomCacheType::None;
-		EHairBindingType BindingType = EHairBindingType::NoneBinding;
-		bool bSimulationEnable = false;
-		bool bDeformationEnable = false;
-		bool bGlobalDeformationEnable = false;
-
-		FCachedGeometry MeshLODData;
-		FHairGroupInstance* Instance = nullptr;
-
-		bool NeedsMeshUpdate(bool bCheckMeshLODIndex=true) const
-		{
-			return 
-				(!bCheckMeshLODIndex || (bCheckMeshLODIndex && MeshLODIndex >= 0)) &&
-				(BindingType == EHairBindingType::Skinning) && 
-				(bGlobalDeformationEnable || bSimulationEnable || bDeformationEnable);
-		}
-	};
+		return 
+			(!bCheckMeshLODIndex || (bCheckMeshLODIndex && MeshLODIndex >= 0)) &&
+			(BindingType == EHairBindingType::Skinning) && 
+			(bGlobalDeformationEnable || bSimulationEnable || bDeformationEnable);
+	}
+};
 
+/** Gather all the instance Data */
+FORCEINLINE void GatherGuideInstanceData(const FHairStrandsInstances& Instances, FHairTransientResources& TransientResources, TArray<FGuideInstanceData>& GuideInstanceDatas, bool& bHasAnySimCacheInstances)
+{
 	// Gather all instances which require guides or RBF deformations
-	bool bHasAnySimCacheInstances = false;
-	TArray<FInstanceData> InstanceDatas;
-	InstanceDatas.Reserve(Instances.Num());
+	bHasAnySimCacheInstances = false;
+	GuideInstanceDatas.Reserve(Instances.Num());
 	for (FHairStrandsInstance* AbstractInstance : Instances)
 	{
 		FHairGroupInstance* Instance = static_cast<FHairGroupInstance*>(AbstractInstance);
@@ -836,7 +822,7 @@ static void RunHairStrandsInterpolation_Guide(
 		}	
 		check(Instance->HairGroupPublicData);
 
-		FInstanceData& InstanceData = InstanceDatas.AddDefaulted_GetRef();
+		FGuideInstanceData& InstanceData = GuideInstanceDatas.AddDefaulted_GetRef();
 		InstanceData.Instance 					= Instance;
 		InstanceData.ActiveGroomCacheType 		= GetHairInstanceCacheType(Instance);
 		InstanceData.BindingType 				= Instance->BindingType;
@@ -849,10 +835,34 @@ static void RunHairStrandsInterpolation_Guide(
 
 		if (InstanceData.ActiveGroomCacheType == EGroomCacheType::Guides)  { bHasAnySimCacheInstances = true; }
 	}	
+}
+
+///////////////////////////////////////////////////////////////////////////////////////////////////
+// Guide interpolation
+
+static void RunHairStrandsInterpolation_Guide(
+	FRDGBuilder& GraphBuilder,
+	FSceneInterface* Scene,
+	const FSceneView* View,
+	const uint32 ViewUniqueID,
+	const FHairStrandsInstances& Instances,
+	const FShaderPrintData* ShaderPrintData,
+	FHairTransientResources& TransientResources,
+	FGlobalShaderMap* ShaderMap)
+{
+	check(IsInRenderingThread());
+
+	RDG_EVENT_SCOPE_STAT(GraphBuilder, HairGuideInterpolation, "HairGuideInterpolation");
+	RDG_GPU_STAT_SCOPE(GraphBuilder, HairGuideInterpolation);
+
+	// Gather all instances which require guides or RBF deformations
+	bool bHasAnySimCacheInstances = false;
+	TArray<FGuideInstanceData> InstanceDatas;
+	GatherGuideInstanceData(Instances, TransientResources, InstanceDatas, bHasAnySimCacheInstances);
 
 	// Update dynamic mesh triangles
 	// Guide update need to run only if simulation is enabled, or if RBF is enabled (since RFB are transfer through guides)
-	for (FInstanceData& InstanceData : InstanceDatas)
+	for (FGuideInstanceData& InstanceData : InstanceDatas)
 	{
 		if (InstanceData.NeedsMeshUpdate())
 		{		
@@ -872,7 +882,7 @@ static void RunHairStrandsInterpolation_Guide(
 
 	// RBF sample position update
 	// Guide update need to run only if simulation is enabled, or if RBF is enabled (since RFB are transfer through guides)
-	for (FInstanceData& InstanceData : InstanceDatas)
+	for (FGuideInstanceData& InstanceData : InstanceDatas)
 	{
 		if (InstanceData.NeedsMeshUpdate())
 		{		
@@ -891,7 +901,7 @@ static void RunHairStrandsInterpolation_Guide(
 
 	// RBF weights update
 	// Guide update need to run only if simulation is enabled, or if RBF is enabled (since RFB are transfer through guides)
-	for (FInstanceData& InstanceData : InstanceDatas)
+	for (FGuideInstanceData& InstanceData : InstanceDatas)
 	{
 		if (InstanceData.NeedsMeshUpdate())
 		{
@@ -909,9 +919,10 @@ static void RunHairStrandsInterpolation_Guide(
 	}
 
 	// Update position offset for each instance (GPU)
-	for (FInstanceData& InstanceData : InstanceDatas)
+	for (FGuideInstanceData& InstanceData : InstanceDatas)
 	{
-		if (InstanceData.NeedsMeshUpdate())
+		const bool bRequiresCaching = bHasAnySimCacheInstances && (InstanceData.ActiveGroomCacheType == EGroomCacheType::Guides);
+		if (InstanceData.NeedsMeshUpdate() && !bRequiresCaching)
 		{
 			check(InstanceData.MeshLODIndex >= 0);
 
@@ -940,9 +951,10 @@ static void RunHairStrandsInterpolation_Guide(
 	}
 
 	// Update position offset for each instance (CPU)
-	for (FInstanceData& InstanceData : InstanceDatas)
+	for (FGuideInstanceData& InstanceData : InstanceDatas)
 	{
-		if (InstanceData.MeshLODIndex >= 0 && (InstanceData.BindingType != EHairBindingType::Skinning) && (InstanceData.bSimulationEnable || InstanceData.bDeformationEnable))
+		const bool bRequiresCaching = bHasAnySimCacheInstances && (InstanceData.ActiveGroomCacheType == EGroomCacheType::Guides);
+		if (InstanceData.MeshLODIndex >= 0 && (InstanceData.BindingType != EHairBindingType::Skinning) && !bRequiresCaching && (InstanceData.bSimulationEnable || InstanceData.bDeformationEnable))
 		{		
 			AddHairStrandUpdatePositionOffsetPass(
 				GraphBuilder,
@@ -958,9 +970,45 @@ static void RunHairStrandsInterpolation_Guide(
 			GraphBuilder.UseExternalAccessMode(Register(GraphBuilder, InstanceData.Instance->Guides.DeformedResource->GetPositionOffsetBuffer(FHairStrandsDeformedResource::EFrameType::Current), ERDGImportedBufferFlags::CreateSRV).Buffer, ERHIAccess::SRVMask);
 		}
 	}
+}
+
+///////////////////////////////////////////////////////////////////////////////////////////////////
+// Guide Deformation
+
+static void RunHairStrandsDeformation_Guide(
+	FRDGBuilder& GraphBuilder,
+	FSceneInterface* Scene,
+	const FSceneView* View,
+	const uint32 ViewUniqueID,
+	const FHairStrandsInstances& Instances,
+	const FShaderPrintData* ShaderPrintData,
+	FHairTransientResources& TransientResources,
+	FGlobalShaderMap* ShaderMap)
+{
+	check(IsInRenderingThread());
+
+	RDG_EVENT_SCOPE_STAT(GraphBuilder, HairGuideInterpolation, "HairGuideInterpolation");
+	RDG_GPU_STAT_SCOPE(GraphBuilder, HairGuideInterpolation);
+
+	// Gather all instances which require guides or RBF deformations
+	bool bHasAnySimCacheInstances = false;
+	TArray<FGuideInstanceData> InstanceDatas;
+	GatherGuideInstanceData(Instances, TransientResources, InstanceDatas, bHasAnySimCacheInstances);
+
+	// Apply deformation from graph deformer
+	for (FGuideInstanceData& InstanceData : InstanceDatas)
+	{
+		if(InstanceData.Instance->Guides.bIsDeformationEnable && InstanceData.Instance->Guides.DeformedResource && InstanceData.Instance->Guides.DeformedResource->DeformerBuffer.Buffer)
+		{
+			FRDGBufferRef DstBuffer = Register(GraphBuilder, InstanceData.Instance->Guides.DeformedResource->GetBuffer(FHairStrandsDeformedResource::Current), ERDGImportedBufferFlags::None).Buffer;
+			FRDGBufferRef SrcBuffer = Register(GraphBuilder, InstanceData.Instance->Guides.DeformedResource->DeformerBuffer, ERDGImportedBufferFlags::None).Buffer;
+			
+			AddCopyBufferPass(GraphBuilder, DstBuffer, SrcBuffer);
+		}
+	}
 
 	// Apply deformation from RBF
-	for (FInstanceData& InstanceData : InstanceDatas)
+	for (FGuideInstanceData& InstanceData : InstanceDatas)
 	{
 		if (
 			((InstanceData.Instance->Guides.bIsSimulationEnable || InstanceData.Instance->Guides.bIsDeformationEnable || InstanceData.Instance->Guides.bIsSimulationCacheEnable)) ||
@@ -988,7 +1036,7 @@ static void RunHairStrandsInterpolation_Guide(
 	}
 
 	// Apply deformation from skeletal mesh bones
-	for (FInstanceData& InstanceData : InstanceDatas)
+	for (FGuideInstanceData& InstanceData : InstanceDatas)
 	{
 		if (InstanceData.ActiveGroomCacheType == EGroomCacheType::None && InstanceData.Instance->Guides.bIsDeformationEnable && InstanceData.Instance->Guides.DeformedResource && InstanceData.Instance->DeformedComponent && (InstanceData.Instance->DeformedSection != INDEX_NONE))
 		{
@@ -1029,7 +1077,7 @@ static void RunHairStrandsInterpolation_Guide(
 	// Apply guide cache update
 	if (bHasAnySimCacheInstances)
 	{
-		for (FInstanceData& InstanceData : InstanceDatas)
+		for (FGuideInstanceData& InstanceData : InstanceDatas)
 		{
 			if (InstanceData.ActiveGroomCacheType == EGroomCacheType::Guides)
 			{
@@ -1038,7 +1086,7 @@ static void RunHairStrandsInterpolation_Guide(
 				FGroomCacheGroupData* GroomCacheData1 = const_cast<FGroomCacheGroupData*>(&InstanceData.Instance->Debug.GroomCacheBuffers->GetNextFrameBuffer().GroupsData[InstanceData.Instance->Debug.GroupIndex]);
 	
 				const float InterpolationFactor = InstanceData.Instance->Debug.GroomCacheBuffers->GetInterpolationFactor();
-	
+				
 				// Update CPU position offset to GPU
 				InstanceData.Instance->Guides.DeformedResource->SetPositionOffset(FHairStrandsDeformedResource::EFrameType::Current, FMath::Lerp(GroomCacheData0->BoundingBox.GetCenter(), GroomCacheData1->BoundingBox.GetCenter(), InterpolationFactor));
 				AddHairStrandUpdatePositionOffsetPass(
@@ -1294,7 +1342,7 @@ static void RunHairStrandsInterpolation_Strands(
 			AddTransitionPass(GraphBuilder, ShaderMap, View->GetShaderPlatform(), Transitions);
 		}
 	}
-
+	
 	// Update dynamic mesh triangles
 	{
 		TArray<FRDGBufferSRVRef> Transitions;
@@ -3074,10 +3122,10 @@ bool HasHairStrandsTexturesQueries();
 void RunHairStrandsTexturesQueries(FRDGBuilder& GraphBuilder, FGlobalShaderMap* ShaderMap, const struct FShaderPrintData* ShaderPrintData);
 
 bool HasHairStrandsPositionQueries();
-void RunHairStrandsPositionQueries(FRDGBuilder& GraphBuilder, FGlobalShaderMap* ShaderMap, const struct FShaderPrintData* DebugShaderData);
+void RunHairStrandsPositionQueries(FRDGBuilder& GraphBuilder, FGlobalShaderMap* ShaderMap, const struct FShaderPrintData* DebugShaderData, const FSceneView* SceneView);
 #endif
 
-static void RunHairStrandsProcess(FRDGBuilder& GraphBuilder, FGlobalShaderMap* ShaderMap, const struct FShaderPrintData* ShaderPrintData)
+static void RunHairStrandsProcess(FRDGBuilder& GraphBuilder, FGlobalShaderMap* ShaderMap, const struct FShaderPrintData* ShaderPrintData, const FSceneView* SceneView)
 {
 #if WITH_EDITOR
 	if (HasHairStrandsTexturesQueries())
@@ -3087,7 +3135,7 @@ static void RunHairStrandsProcess(FRDGBuilder& GraphBuilder, FGlobalShaderMap* S
 
 	if (HasHairStrandsPositionQueries())
 	{
-		RunHairStrandsPositionQueries(GraphBuilder, ShaderMap, ShaderPrintData);
+		RunHairStrandsPositionQueries(GraphBuilder, ShaderMap, ShaderPrintData, SceneView);
 	}
 #endif
 
@@ -3169,7 +3217,7 @@ void ProcessHairStrandsBookmark(
 		if (bHasHairStardsnProcess)
 		{
 			check(GraphBuilder);
-			RunHairStrandsProcess(*GraphBuilder, Parameters.ShaderMap, Parameters.ShaderPrintData);
+			RunHairStrandsProcess(*GraphBuilder, Parameters.ShaderMap, Parameters.ShaderPrintData, Parameters.View);
 		}
 	}
 	else if (Bookmark == EHairStrandsBookmark::ProcessLODSelection)
@@ -3215,7 +3263,6 @@ void ProcessHairStrandsBookmark(
 	}
 	else if (Bookmark == EHairStrandsBookmark::ProcessGuideInterpolation)
 	{
-
 		RunHairStrandsInterpolation_Guide(
 			*GraphBuilder,
 			Parameters.Scene,
@@ -3226,6 +3273,18 @@ void ProcessHairStrandsBookmark(
 			*Parameters.TransientResources,
 			Parameters.ShaderMap);
 	}
+	else if (Bookmark == EHairStrandsBookmark::ProcessGuideDeformation)
+	{
+		RunHairStrandsDeformation_Guide(
+			*GraphBuilder,
+			Parameters.Scene,
+			Parameters.View,
+			Parameters.ViewUniqueID,
+			*Instances,
+			Parameters.ShaderPrintData,
+			*Parameters.TransientResources,
+			Parameters.ShaderMap);
+	}
 	else if (Bookmark == EHairStrandsBookmark::ProcessCardsAndMeshesInterpolation_PrimaryView || Bookmark == EHairStrandsBookmark::ProcessCardsAndMeshesInterpolation_ShadowView)
 	{
 		check(GraphBuilder);
diff --git a/Engine/Plugins/Runtime/HairStrands/Source/HairStrandsCore/Public/GroomAsset.h b/Engine/Plugins/Runtime/HairStrands/Source/HairStrandsCore/Public/GroomAsset.h
--- a/Engine/Plugins/Runtime/HairStrands/Source/HairStrandsCore/Public/GroomAsset.h
+++ b/Engine/Plugins/Runtime/HairStrands/Source/HairStrandsCore/Public/GroomAsset.h
@@ -16,6 +16,7 @@
 #include "GroomAssetInterpolation.h"
 #include "GroomAssetPhysics.h"
 #include "GroomAssetRendering.h"
+#include "GroomAssetDataflow.h"
 #include "Curves/CurveFloat.h"
 #include "HairStrandsInterface.h"
 #include "Engine/SkeletalMesh.h"
@@ -27,6 +28,7 @@
 #include "GroomAsset.generated.h"
 
 
+class UGroomBindingAsset;
 class UAssetUserData;
 class UMaterialInterface;
 class UNiagaraSystem;
@@ -443,8 +445,8 @@ struct FHairVertexFactoryTypesPerMaterialData
 /**
  * Implements an asset that can be used to store hair strands
  */
-UCLASS(BlueprintType, AutoExpandCategories = ("HairRendering", "HairPhysics", "HairInterpolation"), hidecategories = (Object, "Hidden"))
-class HAIRSTRANDSCORE_API UGroomAsset : public UObject, public IInterface_AssetUserData
+UCLASS(BlueprintType, AutoExpandCategories = ("HairRendering", "HairPhysics", "HairInterpolation", "HairDataflow"), hidecategories = (Object, "Hidden"))
+class HAIRSTRANDSCORE_API UGroomAsset : public UObject, public IInterface_AssetUserData, public IDataflowContentOwner
 {
 	GENERATED_UCLASS_BODY()
 
@@ -462,6 +464,10 @@ private:
 	UPROPERTY(EditAnywhere, Category = "HairLOD", meta = (DisplayName = "Auto LOD Bias", ToolTip = "When LOD mode is set to Auto, decrease the screen size at which curves reduction will occur.", ClampMin = "-1", ClampMax = "1", UIMin = "-1.0", UIMax = "1.0"))
 	float AutoLODBias = 0;
 
+	/** Dataflow settings used for any dataflow related operations */
+	UPROPERTY(EditAnywhere, Category = "HairDataflow")
+	FGroomDataflowSettings DataflowSettings;
+
 public:
 	UE_DEPRECATED(5.3, "Please do not access this member directly; use UGroomAsset accessor.")
 	UPROPERTY(EditAnywhere, EditFixedSize, EditFixedSize, Category = "HairInfo", meta = (DisplayName = "Group"))
@@ -525,12 +531,12 @@ public:
 
 	/** Minimum LOD to cook */
 	UE_DEPRECATED(5.3, "Please do not access this member directly; use UGroomAsset accessor.")
-	UPROPERTY(EditAnywhere, Category = "LOD", meta = (DisplayName = "Minimum LOD"))
+	UPROPERTY(EditAnywhere, Category = "HairLOD", meta = (DisplayName = "Minimum LOD"))
 	FPerPlatformInt MinLOD;
 
 	/** When true all LODs below MinLod will still be cooked */
 	UE_DEPRECATED(5.3, "Please do not access this member directly; use UGroomAsset accessor.")
-	UPROPERTY(EditAnywhere, AdvancedDisplay, Category = "LOD")
+	UPROPERTY(EditAnywhere, AdvancedDisplay, Category = "HairLOD")
 	FPerPlatformBool DisableBelowMinLodStripping;
 
 	/** The LOD bias to use after LOD stripping, regardless of MinLOD. Computed at cook time */
@@ -539,7 +545,7 @@ public:
 	TArray<float> EffectiveLODBias;
 
 	/** Store strands/cards/meshes data */
-	UE_DEPRECATED(5.3, "Please do not access this member directly; use UGroomAsset accessor.")
+	UE_DEPRECATED(5.3, "Please do not access this member directly; use UGroomAsset accessor.") 
 	TArray<FHairGroupPlatformData> HairGroupsPlatformData;
 
 private:
@@ -547,6 +553,19 @@ private:
 	TArray<FHairGroupResources> HairGroupsResources;
 
 public:
+
+	/** Static member to retrieve the member name */
+	static FName GetDataflowSettingsMemberName();
+	
+	/** Return the Dataflow settings associated to this groom asset if any, const version. */
+	const FGroomDataflowSettings& GetDataflowSettings() const { return DataflowSettings; }
+
+	/** Return the Dataflow settings associated to this groom asset if any. */
+	FGroomDataflowSettings& GetDataflowSettings() { return DataflowSettings; }
+	
+	/** Set the Dataflow settings associated to this groom asset if any */
+    void SetDataflowSettings(const FGroomDataflowSettings& InDataflowSettings) { DataflowSettings = InDataflowSettings; }
+
 	static FName GetHairGroupsRenderingMemberName();
 	UFUNCTION(BlueprintGetter) TArray<FHairGroupsRendering>& GetHairGroupsRendering();
 	UFUNCTION(BlueprintSetter) void SetHairGroupsRendering(const TArray<FHairGroupsRendering>& In);
@@ -642,6 +661,12 @@ public:
 	virtual void BeginDestroy() override;
 	virtual void Serialize(FArchive& Ar) override;
 
+	//~ Begin IDataflowContentOwner interface 
+	virtual TObjectPtr<UDataflowBaseContent> CreateDataflowContent() override;
+	virtual void WriteDataflowContent(const TObjectPtr<UDataflowBaseContent>& DataflowContent) const override;
+	virtual void ReadDataflowContent(const TObjectPtr<UDataflowBaseContent>& DataflowContent) override;
+	//~ End IDataflowContentOwner interface
+
 #if WITH_EDITOR
 	FOnGroomAssetChanged& GetOnGroomAssetChanged() { return OnGroomAssetChanged;  }
 	FOnGroomAssetResourcesChanged& GetOnGroomAssetResourcesChanged() { return OnGroomAssetResourcesChanged; }
@@ -904,4 +929,4 @@ struct FGroomAssetMemoryStats
 	void Accumulate(const FGroomAssetMemoryStats& In);
 	uint32 GetTotalCPUSize() const;
 	uint32 GetTotalGPUSize() const;
-};
\ No newline at end of file
+};
diff --git a/Engine/Plugins/Runtime/HairStrands/Source/HairStrandsCore/Public/GroomAssetDataflow.h b/Engine/Plugins/Runtime/HairStrands/Source/HairStrandsCore/Public/GroomAssetDataflow.h
new file mode 100644
--- /dev/null
+++ b/Engine/Plugins/Runtime/HairStrands/Source/HairStrandsCore/Public/GroomAssetDataflow.h
@@ -0,0 +1,138 @@
+﻿// Copyright Epic Games, Inc. All Rights Reserved.
+
+#pragma once
+
+#include "Dataflow/DataflowContent.h"
+#include "UObject/Object.h"
+#include "UObject/ObjectMacros.h"
+#include "GroomAssetDataflow.generated.h"
+
+class UGroomBindingAsset;
+class UAnimationAsset;
+
+/** 
+ * Dataflow content owning dataflow and binding assets that will be used to evaluate the graph
+ */
+UCLASS()
+class HAIRSTRANDSCORE_API  UDataflowGroomContent : public UDataflowSkeletalContent
+{
+	GENERATED_BODY()
+
+public:
+	UDataflowGroomContent();
+	virtual ~UDataflowGroomContent() override {}
+	
+#if WITH_EDITOR
+	virtual void PostEditChangeProperty(FPropertyChangedEvent& PropertyChangedEvent) override;
+#endif //if WITH_EDITOR
+
+	/** Data flow binding asset accessors */
+	void SetBindingAsset(const TObjectPtr<UGroomBindingAsset>& InBindingAsset);
+	const TObjectPtr<UGroomBindingAsset>& GetBindingAsset() const { return BindingAsset; }
+
+	//~ UObject interface
+	static void AddReferencedObjects(UObject* InThis, FReferenceCollector& Collector);
+
+	/** Set all the preview actor exposed properties */
+	virtual void SetActorProperties(TObjectPtr<AActor>& PreviewActor) const override;
+
+protected:
+
+	/** Data flow skeletal mesh*/
+	UPROPERTY(EditAnywhere, Category = "Preview", Transient, SkipSerialization)
+	TObjectPtr<UGroomBindingAsset> BindingAsset = nullptr;
+};
+
+USTRUCT()
+struct HAIRSTRANDSCORE_API FGroomDataflowSettings
+{
+	GENERATED_BODY()
+
+	FGroomDataflowSettings();
+
+	/** Return the Dataflow asset name. */
+	static FName GetDataflowAssetMemberName();
+
+	/** Return the Dataflow terminal name. */
+	static FName GetDataflowTerminalMemberName();
+
+	/** Return the Dataflow asset associated to this groom asset if any, const version. */
+	UDataflow* GetDataflowAsset() const { return DataflowAsset; }
+	
+	/** Return the Dataflow terminal associated to this groom asset if any. */
+	const FString& GetDataflowTerminal() const { return DataflowTerminal; }
+
+	/** Return the skeletal mesh associated to this groom asset if any. */
+	USkeletalMesh* GetSkeletalMesh() const { return SkeletalMesh; }
+	
+	/** Return the groom asset dataflow settings rest collection . */
+	const FManagedArrayCollection* GetRestCollection() const { return RestCollection.Get(); }
+	
+	/** Set the Dataflow asset associated to this groom asset if any */
+	void SetDataflowAsset(UDataflow* InDataflowAsset) { DataflowAsset = InDataflowAsset; }
+	
+	/** Set the Dataflow terminal associated to this groom asset if any. */
+	void SetDataflowTerminal(const FString& InDataflowTerminal) { DataflowTerminal = InDataflowTerminal; }
+
+	/** Set the skeletal mesh associated to this groom asset if any. */
+	void SetSkeletalMesh(USkeletalMesh* InSkeletalMesh) { SkeletalMesh = InSkeletalMesh; }
+	
+	/** Set the rest collection onto the groom asset dataflow settings */
+	void SetRestCollection(FManagedArrayCollection* InRestCollection) { RestCollection = TSharedPtr<FManagedArrayCollection, ESPMode::ThreadSafe>(InRestCollection); }
+
+	/** Custom serialization function for the ManagedArrayCollection */
+	bool Serialize(FArchive& Ar);
+	
+#if WITH_EDITORONLY_DATA
+
+	/** Preview scene binding asset setter */
+	void SetPreviewBindingAsset(UGroomBindingAsset* BindingAsset);
+
+	/** Preview scene binding asset setter */
+	UGroomBindingAsset* GetPreviewBindingAsset() const;
+
+	/** Preview scene animation setter */
+	void SetPreviewAnimationAsset(UAnimationAsset* AnimationAsset);
+
+	/** Preview scene animation getter */
+	UAnimationAsset* GetPreviewAnimationAsset() const;
+
+#endif
+
+private :
+
+	/** Dataflow asset. */
+    UPROPERTY(EditAnywhere, Category = "HairDataflow")
+    TObjectPtr<UDataflow> DataflowAsset;
+
+    /** Dataflow terminal. */
+    UPROPERTY(EditAnywhere, Category = "HairDataflow")
+    FString DataflowTerminal = "GroomAssetTerminal";
+
+    /** Hair geometry that could describe external cards and meshes. */
+    UPROPERTY()
+    TObjectPtr<USkeletalMesh> SkeletalMesh;
+
+#if WITH_EDITORONLY_DATA
+	
+	/** Optional binding asset used in the dataflow Editor preview scene */
+	UPROPERTY(AssetRegistrySearchable)
+	TSoftObjectPtr<UGroomBindingAsset> PreviewBindingAsset;
+
+	/** Optional animation asset used in the dataflow Editor preview scene*/
+	UPROPERTY(AssetRegistrySearchable)
+	TSoftObjectPtr<UAnimationAsset> PreviewAnimationAsset;
+	
+#endif
+
+	/** Rest collection used to store all the dataflow attributes */
+	TSharedPtr<FManagedArrayCollection> RestCollection;
+};
+
+template<> struct TStructOpsTypeTraits<FGroomDataflowSettings> : public TStructOpsTypeTraitsBase2<FGroomDataflowSettings>
+{
+	enum
+	{
+		WithSerializer = true,
+	};
+};
\ No newline at end of file
diff --git a/Engine/Plugins/Runtime/HairStrands/Source/HairStrandsCore/Public/GroomAssetPhysics.h b/Engine/Plugins/Runtime/HairStrands/Source/HairStrandsCore/Public/GroomAssetPhysics.h
--- a/Engine/Plugins/Runtime/HairStrands/Source/HairStrandsCore/Public/GroomAssetPhysics.h
+++ b/Engine/Plugins/Runtime/HairStrands/Source/HairStrandsCore/Public/GroomAssetPhysics.h
@@ -47,9 +47,13 @@ struct HAIRSTRANDSCORE_API FHairSolverSettings
 	GENERATED_BODY()
 
 	FHairSolverSettings();
+	
+	/** Use deformer for simulation  */
+	UPROPERTY(EditAnywhere, Category = "SolverSettings", meta = (ToolTip = "Use deformer for the groom solver", DisplayName = "Enable Deformer Simulation"))
+	bool bEnableDeformation;
 
-	/** Enable the simulation for the group */
-	UPROPERTY(EditAnywhere, Category = "SolverSettings", meta = (ToolTip = "Enable the simulation on that group"))
+	/** Use niagara for simulation */
+	UPROPERTY(EditAnywhere, Category = "SolverSettings", meta = (ToolTip = "Use niagara for the groom solver", DisplayName = "Enable Niagara Simulation"))
 	bool EnableSimulation;
 
 	/** Niagara solver to be used for simulation */
diff --git a/Engine/Plugins/Runtime/HairStrands/Source/HairStrandsCore/Public/GroomBindingBuilder.h b/Engine/Plugins/Runtime/HairStrands/Source/HairStrandsCore/Public/GroomBindingBuilder.h
--- a/Engine/Plugins/Runtime/HairStrands/Source/HairStrandsCore/Public/GroomBindingBuilder.h
+++ b/Engine/Plugins/Runtime/HairStrands/Source/HairStrandsCore/Public/GroomBindingBuilder.h
@@ -64,7 +64,7 @@ struct HAIRSTRANDSCORE_API FGroomBindingBuilder
 
 namespace GroomBinding_RBFWeighting
 {
-	struct FPointsSampler
+	struct HAIRSTRANDSCORE_API FPointsSampler
 	{
 		FPointsSampler(TArray<bool>& ValidPoints, const FVector3f* PointPositions, const int32 NumSamples);
 
diff --git a/Engine/Plugins/Runtime/HairStrands/Source/HairStrandsCore/Public/GroomBuilder.h b/Engine/Plugins/Runtime/HairStrands/Source/HairStrandsCore/Public/GroomBuilder.h
--- a/Engine/Plugins/Runtime/HairStrands/Source/HairStrandsCore/Public/GroomBuilder.h
+++ b/Engine/Plugins/Runtime/HairStrands/Source/HairStrandsCore/Public/GroomBuilder.h
@@ -132,4 +132,4 @@ struct FStrandsPositionOutput
 	bool IsValid() const { return Status == 0; }
 };
 
-HAIRSTRANDSCORE_API bool RequestStrandsPosition(const UGroomComponent* Component, FStrandsPositionOutput* Output);
+HAIRSTRANDSCORE_API bool RequestStrandsPosition(const UGroomComponent* Component, TSharedPtr<FStrandsPositionOutput> Output, const bool bReadGuides = false);
diff --git a/Engine/Plugins/Runtime/HairStrands/Source/HairStrandsCore/Public/GroomCache.h b/Engine/Plugins/Runtime/HairStrands/Source/HairStrandsCore/Public/GroomCache.h
--- a/Engine/Plugins/Runtime/HairStrands/Source/HairStrandsCore/Public/GroomCache.h
+++ b/Engine/Plugins/Runtime/HairStrands/Source/HairStrandsCore/Public/GroomCache.h
@@ -2,7 +2,10 @@
 
 #pragma once
 
+#include "GroomAsset.h"
 #include "GroomCacheData.h"
+#include "HairDescription.h"
+#include "Chaos/ChaosCache.h"
 #include "Interfaces/Interface_AssetUserData.h"
 #include "Serialization/BulkData.h"
 #include "GroomCache.generated.h"
@@ -13,7 +16,7 @@ struct FGroomCacheChunk;
  * Implements an asset that is used to store an animated groom
  */
 UCLASS(BlueprintType)
-class HAIRSTRANDSCORE_API UGroomCache : public UObject, public IInterface_AssetUserData
+class HAIRSTRANDSCORE_API UGroomCache : public UObject, public IInterface_AssetUserData, public IChaosCacheData
 {
 	GENERATED_BODY()
 
@@ -29,7 +32,10 @@ public:
 	void Initialize(EGroomCacheType Type);
 	int32 GetStartFrame() const;
 	int32 GetEndFrame() const;
-	float GetDuration() const;
+
+	//~ Begin IChaosCacheData interface
+	virtual float GetDuration() const override;
+	//~ End IChaosCacheData interface
 
 	/** Get the frame number at the specified time within the animation range which might not start at 0 */
 	int32 GetFrameNumberAtTime(const float Time, bool bLooping) const;
@@ -115,3 +121,13 @@ private:
 	EGroomCacheAttributes Attributes;
 	EGroomCacheType Type;
 };
+
+namespace UE::Groom
+{
+	/** Build (create and fill) a groom cache from a processor */
+	HAIRSTRANDSCORE_API void BuildGroomCache(FGroomCacheProcessor& Processor, const FGroomAnimationInfo& AnimInfo, UGroomCache* GroomCache);
+
+	/** Build the groom groups data */
+	HAIRSTRANDSCORE_API bool BuildGroupsData(const FHairDescription& HairDescription, const TArray<FHairGroupPlatformData>& PlatformData,
+		TArray<FHairGroupInfoWithVisibility>& VisibilityData, const TArray<FHairGroupsInterpolation>& InterpolationData, TArray<FGroomCacheInputData>& GroupsData);
+}
diff --git a/Engine/Plugins/Runtime/HairStrands/Source/HairStrandsCore/Public/GroomComponent.h b/Engine/Plugins/Runtime/HairStrands/Source/HairStrandsCore/Public/GroomComponent.h
--- a/Engine/Plugins/Runtime/HairStrands/Source/HairStrandsCore/Public/GroomComponent.h
+++ b/Engine/Plugins/Runtime/HairStrands/Source/HairStrandsCore/Public/GroomComponent.h
@@ -20,6 +20,7 @@ class UGroomCache;
 class UMeshDeformer;
 class UMeshDeformerInstance;
 class UMeshDeformerInstanceSettings;
+class UGroomComponent;
 
 UCLASS(HideCategories = (Object, Physics, Activation, Mobility, "Components|Activation"), editinlinenew, meta = (BlueprintSpawnableComponent), ClassGroup = Rendering)
 class HAIRSTRANDSCORE_API UGroomComponent : public UMeshComponent, public ILODSyncInterface, public INiagaraPhysicsAssetDICollectorInterface
@@ -115,6 +116,9 @@ public:
 	/** Check if the simulation is enabled or not */
 	bool IsSimulationEnable(int32 GroupIndex, int32 LODIndex) const;
 
+	/** Check if the deformation is enabled or not (from the rigged guides or from the solver deformer) */
+	bool IsDeformationEnable(int32 GroupIndex) const;
+
 	/** Update Group Description */
 	void UpdateHairGroupsDesc();
 	void UpdateHairGroupsDescAndInvalidateRenderState(bool bInvalidate=true);
@@ -302,6 +306,12 @@ public:
 	float GetAnimationTime() const;
 	bool IsLooping() const { return bLooping; }
 
+	/** Set the groom solver onto the component */
+	void SetGroomSolver(UMeshComponent* InGroomSolver) {GroomSolver = InGroomSolver;}
+
+	/** Get the groom solver from the component */
+	UMeshComponent* GetGroomSolver() const {return GroomSolver.Get();}
+
 	/** Build the local simulation transform that could be used in strands simulation */
 	void BuildSimulationTransform(FTransform& SimulationTransform) const;
 
@@ -360,6 +370,9 @@ private:
 	friend class FGroomComponentRecreateRenderStateContext;
 	friend class FHairStrandsSceneProxy;
 	friend class FHairCardsSceneProxy;
+
+	/** Groom solver pointer */
+	TSoftObjectPtr<UMeshComponent> GroomSolver = nullptr;
 };
 
 #if WITH_EDITORONLY_DATA
diff --git a/Engine/Plugins/Runtime/HairStrands/Source/HairStrandsCore/Public/GroomInstance.h b/Engine/Plugins/Runtime/HairStrands/Source/HairStrandsCore/Public/GroomInstance.h
--- a/Engine/Plugins/Runtime/HairStrands/Source/HairStrandsCore/Public/GroomInstance.h
+++ b/Engine/Plugins/Runtime/HairStrands/Source/HairStrandsCore/Public/GroomInstance.h
@@ -262,6 +262,7 @@ struct HAIRSTRANDSCORE_API FHairGroupInstance : public FHairStrandsInstance
 		FPersistentPrimitiveIndex CachedMeshPersistentPrimitiveIndex;
 		const UMeshComponent*	MeshComponentForDebug = nullptr;
 		const UGroomComponent*	GroomComponentForDebug = nullptr; // For debug only, shouldn't be deferred on the rendering thread
+		FTransform				RigidRestLocalToWorld = FTransform::Identity;
 		FTransform				RigidCurrentLocalToWorld = FTransform::Identity;
 		FTransform				SkinningCurrentLocalToWorld = FTransform::Identity;
 		FTransform				RigidPreviousLocalToWorld = FTransform::Identity;
diff --git a/Engine/Plugins/Runtime/HairStrands/Source/HairStrandsCore/Public/HairStrandsCore.h b/Engine/Plugins/Runtime/HairStrands/Source/HairStrandsCore/Public/HairStrandsCore.h
--- a/Engine/Plugins/Runtime/HairStrands/Source/HairStrandsCore/Public/HairStrandsCore.h
+++ b/Engine/Plugins/Runtime/HairStrands/Source/HairStrandsCore/Public/HairStrandsCore.h
@@ -15,7 +15,7 @@ class UTexture2D;
 //typedef UTexture2D* (TCreateTextureHelper*)(FName Package, const FIntPoint& Resolution);
 typedef void (*TTextureAllocation)(UTexture2D* Out, const FIntPoint& Resolution, uint32 MipCount);
 
-struct FHairAssetHelper
+struct HAIRSTRANDSCORE_API FHairAssetHelper
 {
 	/* Generate a unique asset & package name */
 	typedef void (*TCreateFilename)(const FString& InAssetName, const FString& Suffix, FString& OutPackageName, FString& OutAssetName);
@@ -59,9 +59,9 @@ public:
 	HAIRSTRANDSCORE_API static UGroomBindingAsset* CreateGroomBindingAsset(EGroomBindingMeshType BindingType, UGroomAsset* GroomAsset, UObject* Source, UObject* Target, const int32 NumInterpolationPoints, const int32 MatchingSection);
 	HAIRSTRANDSCORE_API static UGroomBindingAsset* CreateGroomBindingAsset(EGroomBindingMeshType BindingType, const FString& InPackageName, UObject* InParent, UGroomAsset* GroomAsset, UObject* Source, UObject* Target, const int32 NumInterpolationPoints, const int32 MatchingSection);
 
-	static void SaveAsset(UObject* Object);
+	HAIRSTRANDSCORE_API static void SaveAsset(UObject* Object);
 
-	static FHairAssetHelper& AssetHelper();
+	HAIRSTRANDSCORE_API static FHairAssetHelper& AssetHelper();
 	
 #endif
 };
diff --git a/Engine/Plugins/Runtime/HairStrands/Source/HairStrandsEditor/HairStrandsEditor.Build.cs b/Engine/Plugins/Runtime/HairStrands/Source/HairStrandsEditor/HairStrandsEditor.Build.cs
--- a/Engine/Plugins/Runtime/HairStrands/Source/HairStrandsEditor/HairStrandsEditor.Build.cs
+++ b/Engine/Plugins/Runtime/HairStrands/Source/HairStrandsEditor/HairStrandsEditor.Build.cs
@@ -38,7 +38,8 @@ namespace UnrealBuildTool.Rules
 					"Sequencer",
 					"HairCardGeneratorFramework",
 					"CommonMenuExtensions",
-					"RenderCore"
+					"RenderCore",
+					"DataflowEditor"
 				});
 
 			PrivateIncludePathModuleNames.AddRange(
diff --git a/Engine/Plugins/Runtime/HairStrands/Source/HairStrandsEditor/Private/GroomActions.cpp b/Engine/Plugins/Runtime/HairStrands/Source/HairStrandsEditor/Private/GroomActions.cpp
--- a/Engine/Plugins/Runtime/HairStrands/Source/HairStrandsEditor/Private/GroomActions.cpp
+++ b/Engine/Plugins/Runtime/HairStrands/Source/HairStrandsEditor/Private/GroomActions.cpp
@@ -37,6 +37,7 @@
 #include "AssetToolsModule.h"
 #include "ContentBrowserMenuContexts.h"
 #include "ToolMenus.h"
+#include "Dataflow/DataflowEditor.h"
 
 #define LOCTEXT_NAMESPACE "AssetTypeActions"
 
@@ -52,8 +53,19 @@ EAssetCommandResult UAssetDefinition_GroomAsset::OpenAssets(const FAssetOpenArgs
 	{
 		if (GroomAsset != nullptr)
 		{
-			TSharedRef<FGroomCustomAssetEditorToolkit> NewCustomAssetEditor(new FGroomCustomAssetEditorToolkit());
-			NewCustomAssetEditor->InitCustomAssetEditor(OpenArgs.GetToolkitMode(), OpenArgs.ToolkitHost, GroomAsset);
+			if(!GroomAsset->GetDataflowSettings().GetDataflowAsset())
+			{
+				TSharedRef<FGroomCustomAssetEditorToolkit> NewCustomAssetEditor(new FGroomCustomAssetEditorToolkit());
+				NewCustomAssetEditor->InitCustomAssetEditor(OpenArgs.GetToolkitMode(), OpenArgs.ToolkitHost, GroomAsset);
+			}
+			else
+			{
+				UAssetEditorSubsystem* const AssetEditorSubsystem = GEditor->GetEditorSubsystem<UAssetEditorSubsystem>();
+				UDataflowEditor* const AssetEditor = NewObject<UDataflowEditor>(AssetEditorSubsystem, NAME_None, RF_Transient);
+				const TSubclassOf<AActor> ActorClass = StaticLoadClass(AActor::StaticClass(), nullptr,
+					TEXT("/HairStrands/BP_GroomPreview.BP_GroomPreview_C"), nullptr, LOAD_None, nullptr);
+				AssetEditor->Initialize({ GroomAsset}, ActorClass);
+			}
 		}
 	}
 
diff --git a/Engine/Plugins/Runtime/HairStrands/Source/HairStrandsEditor/Private/GroomAssetDetails.cpp b/Engine/Plugins/Runtime/HairStrands/Source/HairStrandsEditor/Private/GroomAssetDetails.cpp
--- a/Engine/Plugins/Runtime/HairStrands/Source/HairStrandsEditor/Private/GroomAssetDetails.cpp
+++ b/Engine/Plugins/Runtime/HairStrands/Source/HairStrandsEditor/Private/GroomAssetDetails.cpp
@@ -413,6 +413,7 @@ FName GetCategoryName(EMaterialPanelType Type)
 	case EMaterialPanelType::LODs:			return FName(TEXT("LODs"));
 	case EMaterialPanelType::Physics:		return FName(TEXT("Physics"));
 	case EMaterialPanelType::Bindings:		return FName(TEXT("Bindings"));
+	case EMaterialPanelType::Dataflow:		return FName(TEXT("Dataflow"));
 	}
 	return FName(TEXT("Unknown"));
 }
@@ -483,10 +484,10 @@ void FGroomRenderingDetails::AddNewGroupButton(IDetailCategoryBuilder& FilesCate
 	];
 }
 
-static void AddLODModeProperties(IDetailLayoutBuilder& DetailLayout, IDetailCategoryBuilder& FilesCategory)
+static void AddLODModeProperties(const IDetailLayoutBuilder& DetailLayout, IDetailCategoryBuilder& FilesCategory)
 {
-	TSharedRef<IPropertyHandle> LODModeProperty = DetailLayout.GetProperty(UGroomAsset::GetLODModeMemberName(), UGroomAsset::StaticClass());
-	TSharedRef<IPropertyHandle> LODBiasProperty = DetailLayout.GetProperty(UGroomAsset::GetAutoLODBiasMemberName(), UGroomAsset::StaticClass());
+	const TSharedRef<IPropertyHandle> LODModeProperty = DetailLayout.GetProperty(UGroomAsset::GetLODModeMemberName(), UGroomAsset::StaticClass());
+	const TSharedRef<IPropertyHandle> LODBiasProperty = DetailLayout.GetProperty(UGroomAsset::GetAutoLODBiasMemberName(), UGroomAsset::StaticClass());
 
 	FilesCategory.AddProperty(LODModeProperty);
 	FilesCategory.AddProperty(LODBiasProperty);
@@ -514,6 +515,7 @@ void FGroomRenderingDetails::CustomizeStrandsGroupProperties(IDetailLayoutBuilde
 		DetailLayout.HideProperty(DetailLayout.GetProperty(UGroomAsset::GetRiggedSkeletalMeshMemberName(), UGroomAsset::StaticClass()));
 		DetailLayout.HideProperty(DetailLayout.GetProperty(UGroomAsset::GetLODModeMemberName(), UGroomAsset::StaticClass()));
 		DetailLayout.HideProperty(DetailLayout.GetProperty(UGroomAsset::GetAutoLODBiasMemberName(), UGroomAsset::StaticClass()));
+		DetailLayout.HideProperty(DetailLayout.GetProperty(UGroomAsset::GetDataflowSettingsMemberName(), UGroomAsset::StaticClass()));
 	}
 	break;
 	case EMaterialPanelType::Meshes:
@@ -534,6 +536,7 @@ void FGroomRenderingDetails::CustomizeStrandsGroupProperties(IDetailLayoutBuilde
 		DetailLayout.HideProperty(DetailLayout.GetProperty(UGroomAsset::GetRiggedSkeletalMeshMemberName(), UGroomAsset::StaticClass()));
 		DetailLayout.HideProperty(DetailLayout.GetProperty(UGroomAsset::GetLODModeMemberName(), UGroomAsset::StaticClass()));
 		DetailLayout.HideProperty(DetailLayout.GetProperty(UGroomAsset::GetAutoLODBiasMemberName(), UGroomAsset::StaticClass()));
+		DetailLayout.HideProperty(DetailLayout.GetProperty(UGroomAsset::GetDataflowSettingsMemberName(), UGroomAsset::StaticClass()));
 	}
 	break;
 	case EMaterialPanelType::Strands:
@@ -554,6 +557,7 @@ void FGroomRenderingDetails::CustomizeStrandsGroupProperties(IDetailLayoutBuilde
 		DetailLayout.HideProperty(DetailLayout.GetProperty(UGroomAsset::GetRiggedSkeletalMeshMemberName(), UGroomAsset::StaticClass()));
 		DetailLayout.HideProperty(DetailLayout.GetProperty(UGroomAsset::GetLODModeMemberName(), UGroomAsset::StaticClass()));
 		DetailLayout.HideProperty(DetailLayout.GetProperty(UGroomAsset::GetAutoLODBiasMemberName(), UGroomAsset::StaticClass()));
+		DetailLayout.HideProperty(DetailLayout.GetProperty(UGroomAsset::GetDataflowSettingsMemberName(), UGroomAsset::StaticClass()));
 	}
 	break;
 	case EMaterialPanelType::Physics:
@@ -574,6 +578,7 @@ void FGroomRenderingDetails::CustomizeStrandsGroupProperties(IDetailLayoutBuilde
 		DetailLayout.HideProperty(DetailLayout.GetProperty(UGroomAsset::GetRiggedSkeletalMeshMemberName(), UGroomAsset::StaticClass()));
 		DetailLayout.HideProperty(DetailLayout.GetProperty(UGroomAsset::GetLODModeMemberName(), UGroomAsset::StaticClass()));
 		DetailLayout.HideProperty(DetailLayout.GetProperty(UGroomAsset::GetAutoLODBiasMemberName(), UGroomAsset::StaticClass()));
+		DetailLayout.HideProperty(DetailLayout.GetProperty(UGroomAsset::GetDataflowSettingsMemberName(), UGroomAsset::StaticClass()));
 	}
 	break;
 	case EMaterialPanelType::Interpolation:
@@ -594,6 +599,7 @@ void FGroomRenderingDetails::CustomizeStrandsGroupProperties(IDetailLayoutBuilde
 //		DetailLayout.HideProperty(DetailLayout.GetProperty(UGroomAsset::GetRiggedSkeletalMeshMemberName(), UGroomAsset::StaticClass()));
 		DetailLayout.HideProperty(DetailLayout.GetProperty(UGroomAsset::GetLODModeMemberName(), UGroomAsset::StaticClass()));
 		DetailLayout.HideProperty(DetailLayout.GetProperty(UGroomAsset::GetAutoLODBiasMemberName(), UGroomAsset::StaticClass()));
+		DetailLayout.HideProperty(DetailLayout.GetProperty(UGroomAsset::GetDataflowSettingsMemberName(), UGroomAsset::StaticClass()));
 	}
 	break;
 	case EMaterialPanelType::LODs:
@@ -614,6 +620,28 @@ void FGroomRenderingDetails::CustomizeStrandsGroupProperties(IDetailLayoutBuilde
 		DetailLayout.HideProperty(DetailLayout.GetProperty(UGroomAsset::GetRiggedSkeletalMeshMemberName(), UGroomAsset::StaticClass()));
 //		DetailLayout.HideProperty(DetailLayout.GetProperty(UGroomAsset::GetLODModeMemberName(), UGroomAsset::StaticClass()));
 //		DetailLayout.HideProperty(DetailLayout.GetProperty(UGroomAsset::GetAutoLODBiasMemberName(), UGroomAsset::StaticClass()));
+		DetailLayout.HideProperty(DetailLayout.GetProperty(UGroomAsset::GetDataflowSettingsMemberName(), UGroomAsset::StaticClass()));
+	}
+	break;
+	case EMaterialPanelType::Dataflow:
+	{
+		DetailLayout.HideProperty(DetailLayout.GetProperty(UGroomAsset::GetHairGroupsInterpolationMemberName(), UGroomAsset::StaticClass()));
+		DetailLayout.HideProperty(DetailLayout.GetProperty(UGroomAsset::GetHairGroupsRenderingMemberName(), UGroomAsset::StaticClass()));
+		DetailLayout.HideProperty(DetailLayout.GetProperty(UGroomAsset::GetHairGroupsPhysicsMemberName(), UGroomAsset::StaticClass()));
+		DetailLayout.HideProperty(DetailLayout.GetProperty(UGroomAsset::GetHairGroupsCardsMemberName(), UGroomAsset::StaticClass()));
+		DetailLayout.HideProperty(DetailLayout.GetProperty(UGroomAsset::GetHairGroupsMeshesMemberName(), UGroomAsset::StaticClass()));
+		DetailLayout.HideProperty(DetailLayout.GetProperty(UGroomAsset::GetHairGroupsMaterialsMemberName(), UGroomAsset::StaticClass()));
+		DetailLayout.HideProperty(DetailLayout.GetProperty(UGroomAsset::GetHairGroupsLODMemberName(), UGroomAsset::StaticClass()));
+		DetailLayout.HideProperty(DetailLayout.GetProperty(UGroomAsset::GetHairGroupsInfoMemberName(), UGroomAsset::StaticClass()));
+		DetailLayout.HideProperty(DetailLayout.GetProperty(UGroomAsset::GetEnableGlobalInterpolationMemberName(), UGroomAsset::StaticClass()));
+		DetailLayout.HideProperty(DetailLayout.GetProperty(UGroomAsset::GetEnableSimulationCacheMemberName(), UGroomAsset::StaticClass()));
+		DetailLayout.HideProperty(DetailLayout.GetProperty(UGroomAsset::GetHairInterpolationTypeMemberName(), UGroomAsset::StaticClass()));
+		DetailLayout.HideProperty(DetailLayout.GetProperty(UGroomAsset::GetMinLODMemberName(), UGroomAsset::StaticClass()));
+		DetailLayout.HideProperty(DetailLayout.GetProperty(UGroomAsset::GetDisableBelowMinLodStrippingMemberName(), UGroomAsset::StaticClass()));
+		DetailLayout.HideProperty(DetailLayout.GetProperty(UGroomAsset::GetRiggedSkeletalMeshMemberName(), UGroomAsset::StaticClass()));
+		DetailLayout.HideProperty(DetailLayout.GetProperty(UGroomAsset::GetLODModeMemberName(), UGroomAsset::StaticClass()));
+		DetailLayout.HideProperty(DetailLayout.GetProperty(UGroomAsset::GetAutoLODBiasMemberName(), UGroomAsset::StaticClass()));
+		//DetailLayout.HideProperty(DetailLayout.GetProperty(UGroomAsset::GetDataflowSettingsMemberName(), UGroomAsset::StaticClass()));
 	}
 	break;
 	}
@@ -707,6 +735,12 @@ void FGroomRenderingDetails::CustomizeStrandsGroupProperties(IDetailLayoutBuilde
 			}
 		}
 		break;
+		case EMaterialPanelType::Dataflow:
+		{
+			FilesCategory.AddProperty(DetailLayout.GetProperty(FGroomDataflowSettings::GetDataflowAssetMemberName(), UGroomAsset::StaticClass()));
+			FilesCategory.AddProperty(DetailLayout.GetProperty(FGroomDataflowSettings::GetDataflowTerminalMemberName(), UGroomAsset::StaticClass()));
+		}
+		break;
 	}
 }
 
@@ -1837,6 +1871,11 @@ void FGroomRenderingDetails::OnGenerateElementForHairGroup(TSharedRef<IPropertyH
 			AddPropertyWithCustomReset(ChildHandle, ChildrenBuilder, GroupIndex, -1);
 		}
 		break;
+		case EMaterialPanelType::Dataflow:
+		{
+			AddPropertyWithCustomReset(ChildHandle, ChildrenBuilder, GroupIndex, -1);
+		}
+		break;
 		default:
 		{
 			ChildrenBuilder.AddProperty(ChildHandle.ToSharedRef());
diff --git a/Engine/Plugins/Runtime/HairStrands/Source/HairStrandsEditor/Private/GroomAssetDetails.h b/Engine/Plugins/Runtime/HairStrands/Source/HairStrandsEditor/Private/GroomAssetDetails.h
--- a/Engine/Plugins/Runtime/HairStrands/Source/HairStrandsEditor/Private/GroomAssetDetails.h
+++ b/Engine/Plugins/Runtime/HairStrands/Source/HairStrandsEditor/Private/GroomAssetDetails.h
@@ -29,7 +29,8 @@ enum class EMaterialPanelType
 	Interpolation,
 	LODs,
 	Physics,
-	Bindings
+	Bindings,
+	Dataflow
 };
 
 class FGroomRenderingDetails : public IDetailCustomization
diff --git a/Engine/Plugins/Runtime/HairStrands/Source/HairStrandsEditor/Private/GroomCacheImporter.cpp b/Engine/Plugins/Runtime/HairStrands/Source/HairStrandsEditor/Private/GroomCacheImporter.cpp
--- a/Engine/Plugins/Runtime/HairStrands/Source/HairStrandsEditor/Private/GroomCacheImporter.cpp
+++ b/Engine/Plugins/Runtime/HairStrands/Source/HairStrandsEditor/Private/GroomCacheImporter.cpp
@@ -95,14 +95,9 @@ static UGroomCache* CreateGroomCache(EGroomCacheType Type, UObject*& InParent, c
 UGroomCache* FGroomCacheImporter::ProcessToGroomCache(FGroomCacheProcessor& Processor, const FGroomAnimationInfo& AnimInfo, FHairImportContext& ImportContext, const FString& ObjectNameSuffix)
 {
 	TRACE_CPUPROFILER_EVENT_SCOPE(FGroomCacheImporter::ProcessToGroomCache);
-
 	if (UGroomCache* GroomCache = CreateGroomCache(Processor.GetType(), ImportContext.Parent, ObjectNameSuffix, ImportContext.Flags))
 	{
-		Processor.TransferChunks(GroomCache);
-		GroomCache->SetGroomAnimationInfo(AnimInfo);
-
-		GroomCache->MarkPackageDirty();
-		GroomCache->PostEditChange();
+		UE::Groom::BuildGroomCache(Processor, AnimInfo, GroomCache);
 		return GroomCache;
 	}
 	return nullptr;
diff --git a/Engine/Plugins/Runtime/HairStrands/Source/HairStrandsEditor/Private/GroomCustomAssetEditorToolkit.cpp b/Engine/Plugins/Runtime/HairStrands/Source/HairStrandsEditor/Private/GroomCustomAssetEditorToolkit.cpp
--- a/Engine/Plugins/Runtime/HairStrands/Source/HairStrandsEditor/Private/GroomCustomAssetEditorToolkit.cpp
+++ b/Engine/Plugins/Runtime/HairStrands/Source/HairStrandsEditor/Private/GroomCustomAssetEditorToolkit.cpp
@@ -44,6 +44,7 @@ const FName FGroomCustomAssetEditorToolkit::TabId_PhysicsProperties(TEXT("GroomC
 const FName FGroomCustomAssetEditorToolkit::TabId_PreviewGroomComponent(TEXT("GroomCustomAssetEditor_PreviewGroomComponent"));
 const FName FGroomCustomAssetEditorToolkit::TabId_BindingProperties(TEXT("GroomCustomAssetEditor_BindingProperties"));
 const FName FGroomCustomAssetEditorToolkit::TabId_PreviewSceneProperties(TEXT("GroomCustomAssetEditor_PreviewSceneProperties"));
+const FName FGroomCustomAssetEditorToolkit::TabId_DataflowProperties(TEXT("GroomCustomAssetEditor_DataflowProperties"));
 
 void FGroomCustomAssetEditorToolkit::RegisterTabSpawners(const TSharedRef<class FTabManager>& InTabManager)
 {
@@ -91,6 +92,11 @@ void FGroomCustomAssetEditorToolkit::RegisterTabSpawners(const TSharedRef<class
 		.SetGroup(WorkspaceMenuCategory.ToSharedRef())
 		.SetIcon(FSlateIcon(FAppStyle::GetAppStyleSetName(), "ClassIcon.Material"));
 
+	InTabManager->RegisterTabSpawner(TabId_DataflowProperties, FOnSpawnTab::CreateSP(this, &FGroomCustomAssetEditorToolkit::SpawnTab_DataflowProperties))
+		.SetDisplayName(LOCTEXT("DataflowPropertiesTab", "Dataflow"))
+		.SetGroup(WorkspaceMenuCategory.ToSharedRef())
+		.SetIcon(FSlateIcon(FName("DataflowEditorStyle"), "ClassIcon.Dataflow"));
+
 #if GROOMEDITOR_ENABLE_COMPONENT_PANEL
 	InTabManager->RegisterTabSpawner(TabId_PreviewGroomComponent, FOnSpawnTab::CreateSP(this, &FGroomCustomAssetEditorToolkit::SpawnTab_PreviewGroomComponent))
 		.SetDisplayName(LOCTEXT("PreviewGroomComponentTab", "Preview Component"))
@@ -126,6 +132,7 @@ void FGroomCustomAssetEditorToolkit::UnregisterTabSpawners(const TSharedRef<clas
 #endif
 	InTabManager->UnregisterTabSpawner(TabId_BindingProperties);
 	InTabManager->UnregisterTabSpawner(TabId_PreviewSceneProperties);
+	InTabManager->UnregisterTabSpawner(TabId_DataflowProperties);
 }
 
 void FGroomCustomAssetEditorToolkit::DocPropChanged(UObject *InObject, FPropertyChangedEvent &Property)
@@ -352,6 +359,7 @@ void FGroomCustomAssetEditorToolkit::OnClose()
 #endif
 	DetailView_BindingProperties.Reset();
 	DetailView_PreviewSceneProperties.Reset();
+	DetailView_DataflowProperties.Reset();
 }
 
 static void ListAllBindingAssets(const UGroomAsset* InGroomAsset, TWeakObjectPtr<UGroomBindingAssetList>& Out)
@@ -494,7 +502,8 @@ void FGroomCustomAssetEditorToolkit::InitCustomAssetEditor(const EToolkitMode::T
 #endif
 	DetailView_BindingProperties		= PropertyEditorModule.CreateDetailView(DetailsViewArgs);
 	DetailView_PreviewSceneProperties	= PropertyEditorModule.CreateDetailView(DetailsViewArgs); 
-
+	DetailView_DataflowProperties	= PropertyEditorModule.CreateDetailView(DetailsViewArgs);
+	
 	// Customization
 	DetailView_CardsProperties->SetGenericLayoutDetailsDelegate(FOnGetDetailCustomizationInstance::CreateStatic(&FGroomRenderingDetails::MakeInstance, (IGroomCustomAssetEditorToolkit*)this, EMaterialPanelType::Cards));
 	DetailView_MeshesProperties->SetGenericLayoutDetailsDelegate(FOnGetDetailCustomizationInstance::CreateStatic(&FGroomRenderingDetails::MakeInstance, (IGroomCustomAssetEditorToolkit*)this, EMaterialPanelType::Meshes));
@@ -504,7 +513,8 @@ void FGroomCustomAssetEditorToolkit::InitCustomAssetEditor(const EToolkitMode::T
 	DetailView_LODProperties->SetGenericLayoutDetailsDelegate(FOnGetDetailCustomizationInstance::CreateStatic(&FGroomRenderingDetails::MakeInstance, (IGroomCustomAssetEditorToolkit*)this, EMaterialPanelType::LODs));
 	DetailView_MaterialProperties->SetGenericLayoutDetailsDelegate(FOnGetDetailCustomizationInstance::CreateStatic(&FGroomMaterialDetails::MakeInstance, (IGroomCustomAssetEditorToolkit*)this));
 	DetailView_BindingProperties->SetGenericLayoutDetailsDelegate(FOnGetDetailCustomizationInstance::CreateStatic(&FGroomRenderingDetails::MakeInstance, (IGroomCustomAssetEditorToolkit*)this, EMaterialPanelType::Bindings));
-
+	DetailView_DataflowProperties->SetGenericLayoutDetailsDelegate(FOnGetDetailCustomizationInstance::CreateStatic(&FGroomRenderingDetails::MakeInstance, (IGroomCustomAssetEditorToolkit*)this, EMaterialPanelType::Dataflow));
+	
 	SEditorViewport::FArguments args;
 	
 	// Default layout
@@ -541,6 +551,7 @@ void FGroomCustomAssetEditorToolkit::InitCustomAssetEditor(const EToolkitMode::T
 				#endif
 					->AddTab(TabId_BindingProperties,		ETabState::OpenedTab)
 					->AddTab(TabId_PreviewSceneProperties,	ETabState::OpenedTab)
+					->AddTab(TabId_DataflowProperties,		ETabState::OpenedTab)
 				)
 			)
 		);
@@ -579,6 +590,7 @@ void FGroomCustomAssetEditorToolkit::InitCustomAssetEditor(const EToolkitMode::T
 	FProperty* P6 = FindFProperty<FProperty>(GroomAsset->GetClass(), UGroomAsset::GetHairGroupsMaterialsMemberName());
 	FProperty* P7 = FindFProperty<FProperty>(GroomAsset->GetClass(), UGroomAsset::GetHairGroupsInfoMemberName());
 	FProperty* P9 = FindFProperty<FProperty>(GroomAsset->GetClass(), GET_MEMBER_NAME_CHECKED(UGroomAsset, AssetUserData));
+	FProperty* P10 = FindFProperty<FProperty>(GroomAsset->GetClass(), UGroomAsset::GetDataflowSettingsMemberName());
 
 	P7->SetMetaData(TEXT("Category"), TEXT("Hidden"));
 	P0->RemoveMetaData(TEXT("ShowOnlyInnerProperties"));
@@ -590,6 +602,7 @@ void FGroomCustomAssetEditorToolkit::InitCustomAssetEditor(const EToolkitMode::T
 	P6->RemoveMetaData(TEXT("ShowOnlyInnerProperties"));
 	P7->RemoveMetaData(TEXT("ShowOnlyInnerProperties"));
 	P9->RemoveMetaData(TEXT("ShowOnlyInnerProperties"));
+	P10->RemoveMetaData(TEXT("ShowOnlyInnerProperties"));
 
 	// Override the display name so that it show correct/readable name when the details view is displayed as grid view
 	P0->SetMetaData(TEXT("DisplayName"), TEXT("Interpolation"));
@@ -600,6 +613,7 @@ void FGroomCustomAssetEditorToolkit::InitCustomAssetEditor(const EToolkitMode::T
 	P5->SetMetaData(TEXT("DisplayName"), TEXT("Meshes"));
 	P6->SetMetaData(TEXT("DisplayName"), TEXT("Materials"));
 	P7->SetMetaData(TEXT("DisplayName"), TEXT("Info"));
+	P10->SetMetaData(TEXT("DisplayName"), TEXT("Dataflow"));
 
 	// Set the asset we are editing in the details view
 	if (DetailView_InterpolationProperties.IsValid())
@@ -617,6 +631,11 @@ void FGroomCustomAssetEditorToolkit::InitCustomAssetEditor(const EToolkitMode::T
 		DetailView_PhysicsProperties->SetObject(Cast<UObject>(GroomAsset));
 	}
 
+	if (DetailView_DataflowProperties.IsValid())
+	{
+		DetailView_DataflowProperties->SetObject(Cast<UObject>(GroomAsset));
+	}
+
 	if (DetailView_CardsProperties.IsValid())
 	{
 		DetailView_CardsProperties->SetObject(Cast<UObject>(GroomAsset));
@@ -693,6 +712,7 @@ void FGroomCustomAssetEditorToolkit::InitCustomAssetEditor(const EToolkitMode::T
 			if (LocalToolKit->DetailView_MaterialProperties)		{ LocalToolKit->DetailView_MaterialProperties->ForceRefresh(); }
 			if (LocalToolKit->DetailView_BindingProperties)			{ LocalToolKit->DetailView_BindingProperties->ForceRefresh(); }
 			if (LocalToolKit->DetailView_PreviewSceneProperties)	{ LocalToolKit->DetailView_PreviewSceneProperties->ForceRefresh(); }
+			if (LocalToolKit->DetailView_DataflowProperties)		{ LocalToolKit->DetailView_DataflowProperties->ForceRefresh(); }
 		};
 
 		PropertyListenDelegatesResourceChanged.Add(GroomAsset->GetOnGroomAssetResourcesChanged().AddLambda(InvalidateDetailViews));
@@ -874,6 +894,18 @@ TSharedRef<SDockTab> FGroomCustomAssetEditorToolkit::SpawnTab_PhysicsProperties(
 		];
 }
 
+TSharedRef<SDockTab> FGroomCustomAssetEditorToolkit::SpawnTab_DataflowProperties(const FSpawnTabArgs& Args)
+{
+	check(Args.GetTabId() == TabId_DataflowProperties);
+
+	return SNew(SDockTab)
+		.Label(LOCTEXT("DataflowPropertiesTab", "Dataflow"))
+		.TabColorScale(GetTabColorScale())
+		[
+			DetailView_DataflowProperties.ToSharedRef()
+		];
+}
+
 TSharedRef<SDockTab> FGroomCustomAssetEditorToolkit::SpawnTab_RenderingProperties(const FSpawnTabArgs& Args)
 {
 	check(Args.GetTabId() == TabId_RenderingProperties);
diff --git a/Engine/Plugins/Runtime/HairStrands/Source/HairStrandsEditor/Private/GroomMaterialDetails.cpp b/Engine/Plugins/Runtime/HairStrands/Source/HairStrandsEditor/Private/GroomMaterialDetails.cpp
--- a/Engine/Plugins/Runtime/HairStrands/Source/HairStrandsEditor/Private/GroomMaterialDetails.cpp
+++ b/Engine/Plugins/Runtime/HairStrands/Source/HairStrandsEditor/Private/GroomMaterialDetails.cpp
@@ -214,6 +214,9 @@ void FGroomMaterialDetails::CustomizeDetails( IDetailLayoutBuilder& DetailLayout
 		DetailLayout.HideProperty(DetailLayout.GetProperty(UGroomAsset::GetDisableBelowMinLodStrippingMemberName(), UGroomAsset::StaticClass()));
 		DetailLayout.HideProperty(DetailLayout.GetProperty(UGroomAsset::GetMinLODMemberName(), UGroomAsset::StaticClass()));
 		DetailLayout.HideProperty(DetailLayout.GetProperty(UGroomAsset::GetLODModeMemberName(), UGroomAsset::StaticClass()));
+		DetailLayout.HideProperty(DetailLayout.GetProperty(UGroomAsset::GetEnableSimulationCacheMemberName(), UGroomAsset::StaticClass()));
+		DetailLayout.HideProperty(DetailLayout.GetProperty(UGroomAsset::GetRiggedSkeletalMeshMemberName(), UGroomAsset::StaticClass()));
+		DetailLayout.HideProperty(DetailLayout.GetProperty(UGroomAsset::GetDataflowSettingsMemberName(), UGroomAsset::StaticClass()));
 	}
 
 	GroomDetailLayout = &DetailLayout;
diff --git a/Engine/Plugins/Runtime/HairStrands/Source/HairStrandsEditor/Public/GroomCustomAssetEditorToolkit.h b/Engine/Plugins/Runtime/HairStrands/Source/HairStrandsEditor/Public/GroomCustomAssetEditorToolkit.h
--- a/Engine/Plugins/Runtime/HairStrands/Source/HairStrandsEditor/Public/GroomCustomAssetEditorToolkit.h
+++ b/Engine/Plugins/Runtime/HairStrands/Source/HairStrandsEditor/Public/GroomCustomAssetEditorToolkit.h
@@ -148,6 +148,8 @@ private:
 	TSharedRef<SDockTab> SpawnTab_PreviewGroomComponent(const FSpawnTabArgs& Args);
 	TSharedRef<SDockTab> SpawnTab_BindingProperties(const FSpawnTabArgs& Args);
 	TSharedRef<SDockTab> SpawnTab_PreviewSettings(const FSpawnTabArgs& Args);
+	TSharedRef<SDockTab> SpawnTab_DataflowProperties(const FSpawnTabArgs& Args);
+	
 private:
 
 	/** Dockable tab for properties */
@@ -164,6 +166,7 @@ private:
 	TSharedPtr<class IDetailsView> DetailView_PreviewGroomComponent;
 	TSharedPtr<class IDetailsView> DetailView_BindingProperties;
 	TSharedPtr<class IDetailsView> DetailView_PreviewSceneProperties;
+	TSharedPtr<class IDetailsView> DetailView_DataflowProperties;
 
 	static const FName ToolkitFName;
 	static const FName TabId_Viewport;
@@ -178,6 +181,7 @@ private:
 	static const FName TabId_PreviewGroomComponent;
 	static const FName TabId_BindingProperties;
 	static const FName TabId_PreviewSceneProperties;
+	static const FName TabId_DataflowProperties;
 
 
 	bool bIsTabManagerInitialized = false;
diff --git a/Engine/Source/Runtime/Renderer/Private/SceneVisibility.cpp b/Engine/Source/Runtime/Renderer/Private/SceneVisibility.cpp
--- a/Engine/Source/Runtime/Renderer/Private/SceneVisibility.cpp
+++ b/Engine/Source/Runtime/Renderer/Private/SceneVisibility.cpp
@@ -4817,12 +4817,14 @@ static bool IsLargeCameraMovement(FSceneView& View, const FMatrix& PrevViewMatri
 void FSceneRenderer::PreVisibilityFrameSetup(FRDGBuilder& GraphBuilder)
 {
 	FRHICommandListImmediate& RHICmdList = GraphBuilder.RHICmdList;
-
+	
 	if (GetRendererOutput() == ERendererOutput::FinalSceneColor)
 	{
+		FHairStrandsBookmarkParameters Parameters;
+        bool bHasHairStrandsBookmarks = false;
+		
 		if (Views.Num() > 0 && !ViewFamily.EngineShowFlags.HitProxies)
 		{
-			FHairStrandsBookmarkParameters Parameters; 
 			CreateHairStrandsBookmarkParameters(Scene, Views, AllFamilyViews, Parameters, false /*bComputeVisibleInstances*/);
 			Parameters.TransientResources = AllocateHairTransientResources(GraphBuilder, Scene);
 			if (Parameters.HasInstances())
@@ -4838,6 +4840,7 @@ void FSceneRenderer::PreVisibilityFrameSetup(FRDGBuilder& GraphBuilder)
 					if (Views[0].AllowGPUParticleUpdate())
 					{
 						RunHairStrandsBookmark(GraphBuilder, EHairStrandsBookmark::ProcessGuideInterpolation, Parameters);
+						bHasHairStrandsBookmarks = true;
 					}
 				}
 			}
@@ -4861,6 +4864,11 @@ void FSceneRenderer::PreVisibilityFrameSetup(FRDGBuilder& GraphBuilder)
 			FXSystem->PreInitViews(GraphBuilder, Views[0].AllowGPUParticleUpdate() && !ViewFamily.EngineShowFlags.HitProxies, AllFamilies, &ViewFamily);
 		}
 
+		if(bHasHairStrandsBookmarks)
+		{
+			RunHairStrandsBookmark(GraphBuilder, EHairStrandsBookmark::ProcessGuideDeformation, Parameters);
+		}
+
 #if WITH_EDITOR
 		// Draw lines to lights affecting this mesh if its selected.
 		if (ViewFamily.EngineShowFlags.LightInfluences && Scene)
diff --git a/Engine/Source/Runtime/Renderer/Public/HairStrandsInterface.h b/Engine/Source/Runtime/Renderer/Public/HairStrandsInterface.h
--- a/Engine/Source/Runtime/Renderer/Public/HairStrandsInterface.h
+++ b/Engine/Source/Runtime/Renderer/Public/HairStrandsInterface.h
@@ -514,7 +514,8 @@ enum class EHairStrandsBookmark : uint8
 	ProcessCardsAndMeshesInterpolation_ShadowView,
 	ProcessStrandsInterpolation,
 	ProcessDebug,
-	ProcessEndOfFrame
+	ProcessEndOfFrame,
+	ProcessGuideDeformation
 };
 
 enum class EHairInstanceVisibilityType : uint8
