5-Groom dataflow nodes

diff --git a/Engine/Plugins/Runtime/HairStrands/HairStrands.uplugin b/Engine/Plugins/Runtime/HairStrands/HairStrands.uplugin
--- a/Engine/Plugins/Runtime/HairStrands/HairStrands.uplugin
+++ b/Engine/Plugins/Runtime/HairStrands/HairStrands.uplugin
@@ -43,6 +43,11 @@
 			"Name": "HairCardGeneratorFramework",
 			"Type": "Editor",
 			"LoadingPhase": "Default"
+		},
+		{
+			"Name": "HairStrandsDataflow",
+			"Type": "Editor",
+			"LoadingPhase": "Default"
 		}
 	],
 	"Plugins": [
diff --git a/Engine/Plugins/Runtime/HairStrands/Source/HairStrandsDataflow/HairStrandsDataflow.Build.cs b/Engine/Plugins/Runtime/HairStrands/Source/HairStrandsDataflow/HairStrandsDataflow.Build.cs
new file mode 100644
--- /dev/null
+++ b/Engine/Plugins/Runtime/HairStrands/Source/HairStrandsDataflow/HairStrandsDataflow.Build.cs
@@ -0,0 +1,37 @@
+﻿// Copyright Epic Games, Inc. All Rights Reserved.
+
+namespace UnrealBuildTool.Rules
+{
+	public class HairStrandsDataflow: ModuleRules
+	{
+		public HairStrandsDataflow(ReadOnlyTargetRules Target) : base(Target)
+		{
+			PublicDependencyModuleNames.AddRange(
+				new string[]
+				{
+					"Core",
+					"CoreUObject",
+					"ChaosCore",
+					"Chaos",
+					"DataflowCore",
+					"DataflowEngine",
+					"DataflowEditor",
+					"DataflowNodes",
+					"DataflowEnginePlugin",
+					"HairStrandsCore",
+					"Engine",
+					"GeometryAlgorithms",
+					"GeometryCore",
+					"MeshConversion",
+					"MeshConversionEngineTypes",
+					"MeshDescription",
+					"DynamicMesh",
+					"SkeletalMeshDescription",
+					"StaticMeshDescription",
+					"Slate",
+					"SlateCore"
+				}
+			);
+		}
+	}
+}
\ No newline at end of file
diff --git a/Engine/Plugins/Runtime/HairStrands/Source/HairStrandsDataflow/Private/BuildGuidesSkinningNode.cpp b/Engine/Plugins/Runtime/HairStrands/Source/HairStrandsDataflow/Private/BuildGuidesSkinningNode.cpp
new file mode 100644
--- /dev/null
+++ b/Engine/Plugins/Runtime/HairStrands/Source/HairStrandsDataflow/Private/BuildGuidesSkinningNode.cpp
@@ -0,0 +1,120 @@
+// Copyright Epic Games, Inc. All Rights Reserved.
+
+#include "BuildGuidesSkinningNode.h"
+#include "GroomCollectionFacades.h"
+#include "MeshDescriptionToDynamicMesh.h"
+#include "SkeletalMeshAttributes.h"
+#include "SkeletalMeshLODRenderDataToDynamicMesh.h"
+#include "DynamicMesh/DynamicMesh3.h"
+#include "Dataflow/DataflowObjectInterface.h"
+#include "DynamicMesh/MeshTransforms.h"
+#include "Operations/TransferBoneWeights.h"
+#include "Rendering/SkeletalMeshRenderData.h"
+
+#include UE_INLINE_GENERATED_CPP_BY_NAME(BuildGuidesSkinningNode)
+
+namespace UE::Groom::Private
+{
+	FORCEINLINE bool SkeletalMeshToDynamicMesh(const USkeletalMesh* SkeletalMesh, const int32 LodIndex, UE::Geometry::FDynamicMesh3& ToDynamicMesh)
+	{
+		if (SkeletalMesh->HasMeshDescription(LodIndex))
+		{
+			const FMeshDescription* SourceMesh = SkeletalMesh->GetMeshDescription(LodIndex);
+			if (!SourceMesh)
+			{
+				return false;
+			}
+
+			FMeshDescriptionToDynamicMesh Converter;
+			Converter.Convert(SourceMesh, ToDynamicMesh);
+		}
+		else
+		{
+			const FSkeletalMeshRenderData* RenderData = SkeletalMesh->GetResourceForRendering();
+			if (!RenderData)
+			{
+				return false;
+			}
+
+			if (!RenderData->LODRenderData.IsValidIndex(LodIndex))
+			{
+				return false;
+			}
+
+			const FSkeletalMeshLODRenderData* SkeletalMeshLODRenderData = &(RenderData->LODRenderData[LodIndex]);
+
+			UE::Geometry::FSkeletalMeshLODRenderDataToDynamicMesh::ConversionOptions ConversionOptions;
+			ConversionOptions.bWantUVs = false;
+			ConversionOptions.bWantVertexColors = false;
+			ConversionOptions.bWantMaterialIDs = false;
+			ConversionOptions.bWantSkinWeights = true;
+
+			UE::Geometry::FSkeletalMeshLODRenderDataToDynamicMesh::Convert(SkeletalMeshLODRenderData, SkeletalMesh->GetRefSkeleton(), ConversionOptions, ToDynamicMesh);
+		}
+
+		return true;
+	}
+}
+
+void FBuildGuidesSkinningDataflowNode::Evaluate(UE::Dataflow::FContext& Context, const FDataflowOutput* Out) const
+{
+	if (Out->IsA<FManagedArrayCollection>(&Collection))
+	{
+		FManagedArrayCollection GroomCollection = GetValue<FManagedArrayCollection>(Context, &Collection);
+		
+		UE::Groom::FGroomGuidesFacade CurvesFacade(GroomCollection);
+
+		if(CurvesFacade.IsValid())
+		{
+			const USkeletalMesh* LocalSkelmesh = SkeletalMesh;
+			if (!LocalSkelmesh)
+			{
+				if (const UE::Dataflow::FEngineContext* EngineContext = Context.AsType<UE::Dataflow::FEngineContext>())
+				{
+					if(UGroomAsset* GroomAsset = Cast<UGroomAsset>(EngineContext->Owner))
+					{
+						LocalSkelmesh = GroomAsset->GetDataflowSettings().GetSkeletalMesh();
+					}
+				}
+			}
+			if(LocalSkelmesh && LocalSkelmesh->IsValidLODIndex(LODIndex))
+			{
+				TMap<FName, FBoneIndexType> TargetBoneToIndex;
+				TargetBoneToIndex.Reserve(LocalSkelmesh->GetRefSkeleton().GetRawBoneNum());
+				for (int32 BoneIdx = 0; BoneIdx < LocalSkelmesh->GetRefSkeleton().GetRawBoneNum(); ++BoneIdx)
+				{
+					TargetBoneToIndex.Add(LocalSkelmesh->GetRefSkeleton().GetRawRefBoneInfo()[BoneIdx].Name, BoneIdx);
+				}
+				
+				UE::Geometry::FDynamicMesh3 DynamicMesh;
+				if (UE::Groom::Private::SkeletalMeshToDynamicMesh(LocalSkelmesh, LODIndex, DynamicMesh))
+				{
+					MeshTransforms::ApplyTransform(DynamicMesh, RelativeTransform, true);
+
+					UE::Geometry::FTransferBoneWeights TransferBoneWeights(&DynamicMesh, FSkeletalMeshAttributes::DefaultSkinWeightProfileName);
+					TransferBoneWeights.bUseParallel = true;
+					TransferBoneWeights.MaxNumInfluences = 4;
+					TransferBoneWeights.TransferMethod = UE::Geometry::FTransferBoneWeights::ETransferBoneWeightsMethod::ClosestPointOnSurface;
+
+					if (TransferBoneWeights.Validate() == UE::Geometry::EOperationValidationResult::Ok)
+					{
+						ParallelFor(CurvesFacade.GetNumPoints(), [&TransferBoneWeights, &TargetBoneToIndex, &CurvesFacade](int32 PointIndex)
+						{
+							TArray<int32> BoneIndices;
+							TArray<float> BoneWeights;
+							TransferBoneWeights.TransferWeightsToPoint(BoneIndices, BoneWeights,
+								CurvesFacade.GetPointRestPositions()[PointIndex],&TargetBoneToIndex);
+
+							CurvesFacade.SetPointBoneIndices(PointIndex, BoneIndices);
+							CurvesFacade.SetPointBoneWeights(PointIndex, BoneWeights);
+
+						}, TransferBoneWeights.bUseParallel  ? EParallelForFlags::None : EParallelForFlags::ForceSingleThread);
+					}
+				}
+			}
+		}
+		
+		SetValue(Context, MoveTemp(GroomCollection), &Collection);
+	}
+}
+
diff --git a/Engine/Plugins/Runtime/HairStrands/Source/HairStrandsDataflow/Private/GenerateGuidesCurvesNode.cpp b/Engine/Plugins/Runtime/HairStrands/Source/HairStrandsDataflow/Private/GenerateGuidesCurvesNode.cpp
new file mode 100644
--- /dev/null
+++ b/Engine/Plugins/Runtime/HairStrands/Source/HairStrandsDataflow/Private/GenerateGuidesCurvesNode.cpp
@@ -0,0 +1,105 @@
+// Copyright Epic Games, Inc. All Rights Reserved.
+
+#include "GenerateGuidesCurvesNode.h"
+
+#include "GroomBindingBuilder.h"
+#include "GroomCollectionFacades.h"
+#include "Dataflow/DataflowObjectInterface.h"
+
+#include UE_INLINE_GENERATED_CPP_BY_NAME(GenerateGuidesCurvesNode)
+
+namespace UE::Groom::Private
+{
+	FORCEINLINE uint32 SampleStrandsCurves(const UE::Groom::FGroomStrandsFacade& StrandsFacade, const uint32 GuidesCount, const uint32 GuidesOffset,
+		const uint32 ObjectIndex, const uint32 PrevStrands, const uint32 NextStrands, TArray<uint32>& SampleIndices)
+	{
+		const uint32 NumStrands = NextStrands-PrevStrands;
+		const uint32 NumGuides = (ObjectIndex == StrandsFacade.GetNumObjects()-1) ? GuidesCount - GuidesOffset :
+			GuidesCount * static_cast<float>(NumStrands)/StrandsFacade.GetNumCurves();
+
+		TArray<FVector3f> RootPositions;
+		RootPositions.Init(FVector3f::ZeroVector, NumStrands);
+		for(uint32 StrandIndex = 0; StrandIndex < NumStrands; ++StrandIndex)
+		{
+			const uint32 RootIndex = (StrandIndex+PrevStrands == 0) ? 0 : StrandsFacade.GetCurvePointOffsets()[StrandIndex+PrevStrands-1];
+			RootPositions[StrandIndex] = StrandsFacade.GetPointRestPositions()[RootIndex];
+		}
+
+		TArray<bool> ValidPoints;
+		ValidPoints.Init(true, NumStrands);
+				
+		GroomBinding_RBFWeighting::FPointsSampler PointsSampler(ValidPoints, RootPositions.GetData(), NumGuides);
+
+		SampleIndices = PointsSampler.SampleIndices;
+		return NumGuides;
+	}
+
+	FORCEINLINE void BuildGuidesCurves(const UE::Groom::FGroomStrandsFacade& StrandsFacade, const uint32 PrevStrands, const TArray<uint32>& SampleIndices,
+		TArray<FVector3f>& PointRestPositions, TArray<int32>& ObjectCurveOffsets, TArray<int32>& CurvePointOffsets, TArray<int32>& CurveStrandIndices)
+	{
+		for(uint32 SampleIndex : SampleIndices)
+		{
+			const uint32 GuideIndex = SampleIndex + PrevStrands;
+			const int32 PointBegin = (GuideIndex == 0) ? 0 : StrandsFacade.GetCurvePointOffsets()[GuideIndex-1];
+			const int32 PointEnd = StrandsFacade.GetCurvePointOffsets()[GuideIndex];
+
+			for(int32 PointIndex = PointBegin; PointIndex < PointEnd; ++PointIndex)
+			{
+				PointRestPositions.Add(StrandsFacade.GetPointRestPositions()[PointIndex]);
+			}
+			CurvePointOffsets.Add(PointRestPositions.Num());
+			CurveStrandIndices.Add(SampleIndex);
+		}
+		ObjectCurveOffsets.Add(CurvePointOffsets.Num());
+	}
+}
+
+void FGenerateGuidesCurvesDataflowNode::Evaluate(UE::Dataflow::FContext& Context, const FDataflowOutput* Out) const
+{
+	if (Out->IsA<FManagedArrayCollection>(&Collection))
+	{
+		FManagedArrayCollection GroomCollection = GetValue<FManagedArrayCollection>(Context, &Collection);
+		
+		UE::Groom::FGroomStrandsFacade StrandsFacade(GroomCollection);
+		if(StrandsFacade.IsValid())
+		{
+			TArray<FVector3f> PointRestPositions;
+			TArray<int32> ObjectCurveOffsets, CurvePointOffsets, CurveStrandIndices;
+			
+			uint32 PrevOffset = 0, GuidesOffset = 0;;
+			for (int32 ObjectIndex = 0; ObjectIndex < StrandsFacade.GetNumObjects(); ++ObjectIndex)
+			{
+				const uint32 NextOffset = StrandsFacade.GetObjectCurveOffsets()[ObjectIndex];
+
+				TArray<uint32> SampleIndices;
+				
+				// Sample guides among the input strands
+				const uint32 NumGuides = UE::Groom::Private::SampleStrandsCurves(StrandsFacade, GuidesCount, GuidesOffset, ObjectIndex, PrevOffset, NextOffset, SampleIndices);
+
+				// Build the sampled guides
+				UE::Groom::Private::BuildGuidesCurves(StrandsFacade, PrevOffset, SampleIndices, PointRestPositions, ObjectCurveOffsets, CurvePointOffsets, CurveStrandIndices);
+
+				GuidesOffset += NumGuides;
+				PrevOffset = NextOffset;
+			}
+			
+			UE::Groom::FGroomGuidesFacade GuidesFacade(GroomCollection);
+			const TArray<int32> ObjectPointSamples = GuidesFacade.GetObjectPointSamples();
+
+			// Init the groom collection
+			GuidesFacade.InitGroomCollection(PointRestPositions, CurvePointOffsets, ObjectCurveOffsets);
+
+			// Set the curve strand indices
+			GuidesFacade.SetCurveStrandIndices(CurveStrandIndices);
+
+			// Set the point samples if already defined and if the size is matching
+			if(ObjectPointSamples.Num() == GuidesFacade.GetNumObjects())
+			{
+				GuidesFacade.SetObjectPointSamples(ObjectPointSamples);
+			}
+		}
+		
+		SetValue(Context, MoveTemp(GroomCollection), &Collection);
+	}
+}
+
diff --git a/Engine/Plugins/Runtime/HairStrands/Source/HairStrandsDataflow/Private/GetGroomAssetNode.cpp b/Engine/Plugins/Runtime/HairStrands/Source/HairStrandsDataflow/Private/GetGroomAssetNode.cpp
new file mode 100644
--- /dev/null
+++ b/Engine/Plugins/Runtime/HairStrands/Source/HairStrandsDataflow/Private/GetGroomAssetNode.cpp
@@ -0,0 +1,103 @@
+﻿// Copyright Epic Games, Inc. All Rights Reserved.
+
+#include "GetGroomAssetNode.h"
+#include "GroomEdit.h"
+#include "GroomInstance.h"
+#include "GroomCollectionFacades.h"
+#include "Dataflow/DataflowInputOutput.h"
+#include "Dataflow/DataflowObjectInterface.h"
+
+#include UE_INLINE_GENERATED_CPP_BY_NAME(GetGroomAssetNode)
+
+namespace UE::Groom::Private
+{
+	template<typename FacadeType>
+	FORCEINLINE const UGroomAsset* BuildGroomCollection(Dataflow::FContext& Context, const UGroomAsset* NodeGroom, FManagedArrayCollection& GroomCollection, FEditableGroom& EditGroom)
+	{
+		const UGroomAsset* GroomAsset = NodeGroom;
+		if (!GroomAsset)
+		{
+			if (const Dataflow::FEngineContext* EngineContext = Context.AsType<Dataflow::FEngineContext>())
+			{
+				GroomAsset = Cast<UGroomAsset>(EngineContext->Owner);
+			}
+		}
+
+		if (GroomAsset)
+		{
+			ConvertFromGroomAsset(const_cast<UGroomAsset*>(GroomAsset), &EditGroom);
+
+			TArray<FVector3f> PointRestPositions;
+			TArray<int32> ObjectCurveOffsets;
+			TArray<int32> CurvePointOffsets;
+
+			for(FEditableGroomGroup& GroomGroup : EditGroom.Groups)
+			{
+				for( const typename FacadeType::FEditableType& EditType : FacadeType::GetEditableGroom(GroomGroup))
+				{
+					for(auto& GuidePoint : EditType.ControlPoints)
+					{
+						PointRestPositions.Add(GuidePoint.Position);
+					}
+					CurvePointOffsets.Add(PointRestPositions.Num());
+				}
+				ObjectCurveOffsets.Add(CurvePointOffsets.Num());
+			}
+			FacadeType CurvesFacade(GroomCollection);
+			CurvesFacade.InitGroomCollection(PointRestPositions, CurvePointOffsets, ObjectCurveOffsets);
+
+			return GroomAsset;
+		}
+		return nullptr;
+	}
+}
+
+void FGetGroomAssetDataflowNode::Evaluate(UE::Dataflow::FContext& Context, const FDataflowOutput* Out) const
+{
+	if (Out->IsA<FManagedArrayCollection>(&Collection))
+	{
+		FManagedArrayCollection GroomCollection;
+		FEditableGroom EditGroom;
+		const UGroomAsset* LocalGroom = nullptr;
+		if(CurvesType == EGroomCollectionType::Guides)
+		{
+			LocalGroom = UE::Groom::Private::BuildGroomCollection<UE::Groom::FGroomGuidesFacade>(
+				Context, GroomAsset, GroomCollection, EditGroom);
+		}
+		else if(CurvesType == EGroomCollectionType::Strands)
+		{
+			LocalGroom = UE::Groom::Private::BuildGroomCollection<UE::Groom::FGroomStrandsFacade>(
+				Context, GroomAsset, GroomCollection, EditGroom);
+		}
+
+		if(LocalGroom)
+		{
+			TArray<int32> ObjectPointSamples, CurveStrandIndices;
+			const TArray<FHairGroupsPhysics>& GroupsPhysics = LocalGroom->GetHairGroupsPhysics();
+			for(const FHairGroupsPhysics& GroupPhysics : GroupsPhysics)
+			{
+				ObjectPointSamples.Add(static_cast<uint8>(GroupPhysics.StrandsParameters.StrandsSize));
+			}
+				
+			UE::Groom::FGroomGuidesFacade GuidesFacade(GroomCollection);
+			if(GroomCollection.NumElements(UE::Groom::FGroomGuidesFacade::ObjectsGroup) == 0)
+			{
+				GroomCollection.AddElements(ObjectPointSamples.Num(), UE::Groom::FGroomGuidesFacade::ObjectsGroup);
+			}
+			GuidesFacade.SetObjectPointSamples(ObjectPointSamples);
+		}
+		SetValue(Context, MoveTemp(GroomCollection), &Collection);
+	}
+}
+
+TArray<UE::Dataflow::FRenderingParameter> FGetGroomAssetDataflowNode::GetRenderParametersImpl() const
+{
+	if(CurvesType == EGroomCollectionType::Guides)
+	{
+		return { {TEXT("GuidesRender"), FName("FGroomCollection"), {TEXT("Collection")}}};
+	}
+	else
+	{
+		return { {TEXT("StrandsRender"), FName("FGroomCollection"), {TEXT("Collection")}}};
+	}
+}
diff --git a/Engine/Plugins/Runtime/HairStrands/Source/HairStrandsDataflow/Private/GetGuidesAttributesNodes.cpp b/Engine/Plugins/Runtime/HairStrands/Source/HairStrandsDataflow/Private/GetGuidesAttributesNodes.cpp
new file mode 100644
--- /dev/null
+++ b/Engine/Plugins/Runtime/HairStrands/Source/HairStrandsDataflow/Private/GetGuidesAttributesNodes.cpp
@@ -0,0 +1,18 @@
+// Copyright Epic Games, Inc. All Rights Reserved.
+
+#include "GetGuidesAttributesNodes.h"
+#include "GroomCollectionFacades.h"
+#include "Dataflow/DataflowObjectInterface.h"
+
+#include UE_INLINE_GENERATED_CPP_BY_NAME(GetGuidesAttributesNodes)
+
+void FGetGuidesKinematicWeightsDataflowNode::Evaluate(UE::Dataflow::FContext& Context, const FDataflowOutput* Out) const
+{
+	if (Out->IsA<FCollectionAttributeKey>(&AttributeKey))
+	{
+		FCollectionAttributeKey KinematicWeightKey(
+			UE::Groom::FGroomGuidesFacade::VertexKinematicWeightsAttribute.ToString(), UE::Groom::FGroomGuidesFacade::VerticesGroup.ToString());
+		SetValue(Context, MoveTemp(KinematicWeightKey), &AttributeKey);
+	}
+}
+
diff --git a/Engine/Plugins/Runtime/HairStrands/Source/HairStrandsDataflow/Private/GroomAssetTerminalNode.cpp b/Engine/Plugins/Runtime/HairStrands/Source/HairStrandsDataflow/Private/GroomAssetTerminalNode.cpp
new file mode 100644
--- /dev/null
+++ b/Engine/Plugins/Runtime/HairStrands/Source/HairStrandsDataflow/Private/GroomAssetTerminalNode.cpp
@@ -0,0 +1,146 @@
+﻿// Copyright Epic Games, Inc. All Rights Reserved.
+
+#include "GroomAssetTerminalNode.h"
+
+#include "GroomCollectionFacades.h"
+#include "GroomEdit.h"
+#include "Dataflow/DataflowObjectInterface.h"
+
+#include UE_INLINE_GENERATED_CPP_BY_NAME(GroomAssetTerminalNode)
+
+namespace UE::Groom::Private
+{
+	FORCEINLINE void BuildEditableGuides(const UE::Groom::FGroomGuidesFacade& GuidesFacade, FEditableGroom& EditGroom)
+	{
+		int32 ObjectIndex = 0, CurveIndex = 0;
+		int32 PrevCurve = 0, PrevPoint = 0;;
+		for (FEditableGroomGroup& Group : EditGroom.Groups)
+		{
+			const int32 NextCurve = GuidesFacade.GetObjectCurveOffsets()[ObjectIndex];
+			Group.Guides.SetNum(NextCurve-PrevCurve);
+
+			// Only rebuild the guides since the strands are the same
+			for (FEditableHairGuide& Guide : Group.Guides)
+			{
+				const int32 NextPoint = GuidesFacade.GetCurvePointOffsets()[CurveIndex];
+				const int32 StrandIndex = GuidesFacade.GetCurveStrandIndices()[CurveIndex];
+						
+				Guide.ControlPoints.Reset();
+				for(int32 PointIndex = PrevPoint; PointIndex < NextPoint; ++PointIndex)
+				{
+					Guide.ControlPoints.Add({GuidesFacade.GetPointRestPositions()[PointIndex],
+						FMath::Clamp(static_cast<float>(PointIndex-PrevPoint) / (NextPoint-PrevPoint-1), 0.f, 1.f)});
+				}
+						
+				if(StrandIndex != INDEX_NONE)
+				{
+					Guide.GuideID = CurveIndex;
+					Guide.RootUV = Group.Strands[StrandIndex].RootUV;
+				}
+				
+				PrevPoint = NextPoint;
+				++CurveIndex;
+			}
+			PrevCurve = NextCurve;
+			++ObjectIndex;
+		}
+	}
+
+	FORCEINLINE void CopyCollectionAttributes(const FManagedArrayCollection* InputCollection, FManagedArrayCollection* OutputCollection,
+		const TArray<TTuple<FName, FName>> AttributesToCopy = TArray<TTuple<FName, FName>>())
+	{
+		for(const TTuple<FName, FName>& AttributeToCopy : AttributesToCopy)
+		{
+			const FName& AttributeName = AttributeToCopy.Value;
+			const FName& GroupName = AttributeToCopy.Key;
+
+			if (InputCollection->HasGroup(GroupName))
+			{
+				if(!OutputCollection->HasGroup(GroupName))
+				{
+					OutputCollection->AddGroup(GroupName);
+				}
+				if (InputCollection->NumElements(GroupName) != OutputCollection->NumElements(GroupName))
+				{
+					OutputCollection->EmptyGroup(GroupName);
+					OutputCollection->AddElements(InputCollection->NumElements(GroupName), GroupName);
+				}
+				OutputCollection->CopyAttribute(*InputCollection, AttributeName, GroupName);
+			}
+		}
+	}
+
+	template<typename FacadeType>
+	FORCEINLINE void TransferVerticesAttributes(const FManagedArrayCollection& InCollection, FManagedArrayCollection* OutCollection, const int32 NumPoints)
+	{
+		// Transfer vertices weight maps onto the points to be stored onto the rest collection
+		const TArray<FName> AttributeNames = InCollection.AttributeNames(FacadeType::VerticesGroup);
+		for(const FName& AttributeName : AttributeNames)
+		{
+			if(InCollection.GetAttributeType(AttributeName, FacadeType::VerticesGroup) == EManagedArrayType::FFloatType)
+			{
+				const TManagedArray<float>& VerticesAttribute = InCollection.GetAttribute<float>(AttributeName, FacadeType::VerticesGroup);
+				TManagedArray<float>& PointsAttribute = OutCollection->AddAttribute<float>(AttributeName, FacadeType::PointsGroup);
+				
+				for(int32 PointIndex = 0; PointIndex < NumPoints; ++PointIndex)
+				{
+					PointsAttribute[PointIndex] = 0.5f * (VerticesAttribute[2*PointIndex] + VerticesAttribute[2*PointIndex+1]);
+				}
+			}
+		}
+	}
+}
+
+void FGroomAssetTerminalDataflowNode::SetAssetValue(TObjectPtr<UObject> Asset, UE::Dataflow::FContext& Context) const
+{
+	if (UGroomAsset* GroomAsset = Cast<UGroomAsset>(Asset.Get()))
+	{
+		const FManagedArrayCollection& InCollection = GetValue<FManagedArrayCollection>(Context, &Collection);
+		UE::Groom::FGroomGuidesFacade GuidesFacade(InCollection);
+
+		FManagedArrayCollection* OutCollection = (GuidesFacade.IsValid()) ? new FManagedArrayCollection() : nullptr;
+
+		if(GuidesFacade.IsValid())
+		{
+			// Only copy the attributes not stored already in the asset
+			const TArray<TTuple<FName, FName>> AttributesToCopy = {
+				{UE::Groom::FGroomGuidesFacade::PointsGroup, UE::Groom::FGroomGuidesFacade::PointBoneIndicesAttribute},
+				{UE::Groom::FGroomGuidesFacade::PointsGroup, UE::Groom::FGroomGuidesFacade::PointBoneWeightsAttribute},
+				{UE::Groom::FGroomGuidesFacade::ObjectsGroup, UE::Groom::FGroomGuidesFacade::ObjectCurveOffsetsAttribute},
+				{UE::Groom::FGroomGuidesFacade::CurvesGroup, UE::Groom::FGroomGuidesFacade::CurvePointOffsetsAttribute}};
+
+			// Copy attributes from input collection
+			UE::Groom::Private::CopyCollectionAttributes(&InCollection, OutCollection, AttributesToCopy);
+			
+			// Transfer vertices weight maps onto the points to be stored onto the rest collection
+			UE::Groom::Private::TransferVerticesAttributes<UE::Groom::FGroomGuidesFacade>(InCollection, OutCollection, GuidesFacade.GetNumPoints());
+
+			// Build an editable groom asset for the strands
+			FEditableGroom EditGroom;
+			ConvertFromGroomAsset(const_cast<UGroomAsset*>(GroomAsset), &EditGroom);
+
+			if(GuidesFacade.GetNumObjects() == EditGroom.Groups.Num())
+			{
+				// Build the editable guides
+				UE::Groom::Private::BuildEditableGuides(GuidesFacade,EditGroom);
+
+				// Convert to groom asset
+				ConvertToGroomAsset(const_cast<UGroomAsset*>(GroomAsset), &EditGroom, EEditableGroomOperations::ControlPoints_Modified);
+			}
+		}
+
+		if(OutCollection)
+		{
+			GroomAsset->GetDataflowSettings().SetRestCollection(OutCollection);
+		}
+	}
+}
+
+void FGroomAssetTerminalDataflowNode::Evaluate(UE::Dataflow::FContext& Context, const FDataflowOutput* Out) const
+{
+	const FManagedArrayCollection& InCollection = GetValue<FManagedArrayCollection>(Context, &Collection);
+	SetValue(Context, InCollection, &Collection);
+}
+
+
+
diff --git a/Engine/Plugins/Runtime/HairStrands/Source/HairStrandsDataflow/Private/GroomCollectionFacades.cpp b/Engine/Plugins/Runtime/HairStrands/Source/HairStrandsDataflow/Private/GroomCollectionFacades.cpp
new file mode 100644
--- /dev/null
+++ b/Engine/Plugins/Runtime/HairStrands/Source/HairStrandsDataflow/Private/GroomCollectionFacades.cpp
@@ -0,0 +1,288 @@
+﻿// Copyright Epic Games, Inc. All Rights Reserved.
+
+#include "GroomCollectionFacades.h"
+#include "GeometryCollection/GeometryCollection.h"
+
+namespace UE::Groom
+{
+	// Strands/Guides Groups
+	template<typename DerivedType>
+	const FName FGroomCollectionFacade<DerivedType>::CurvesGroup(DerivedType::GroupPrefix.ToString() + FString("Curves"));
+	template<typename DerivedType>
+	const FName FGroomCollectionFacade<DerivedType>::EdgesGroup(DerivedType::GroupPrefix.ToString() + FString("Edges"));
+	template<typename DerivedType>
+	const FName FGroomCollectionFacade<DerivedType>::ObjectsGroup(DerivedType::GroupPrefix.ToString() + FString("Objects"));
+	template<typename DerivedType>
+	const FName FGroomCollectionFacade<DerivedType>::PointsGroup(DerivedType::GroupPrefix.ToString() + FString("Points"));
+	
+	// Strands/Guides Attributes
+	template<typename DerivedType>
+	const FName FGroomCollectionFacade<DerivedType>::CurvePointOffsetsAttribute("CurvePointOffsets");
+	template<typename DerivedType>
+	const FName FGroomCollectionFacade<DerivedType>::PointCurveIndicesAttribute("PointCurveIndices");
+	template<typename DerivedType>
+	const FName FGroomCollectionFacade<DerivedType>::ObjectCurveOffsetsAttribute("ObjectCurveOffsets");
+	template<typename DerivedType>
+	const FName FGroomCollectionFacade<DerivedType>::CurveObjectIndicesAttribute("CurveObjectIndices");
+	template<typename DerivedType>
+	const FName FGroomCollectionFacade<DerivedType>::EdgeRestOrientationsAttribute("EdgeRestOrientations");
+	template<typename DerivedType>
+	const FName FGroomCollectionFacade<DerivedType>::PointRestPositionsAttribute("PointRestPositions");
+
+	// Strands/Guides Prefix
+	const FName FGroomStrandsFacade::GroupPrefix("Strands");
+	const FName FGroomGuidesFacade::GroupPrefix("Guides");
+
+	// Guides Groups
+	const FName FGroomGuidesFacade::VerticesGroup(FGroomGuidesFacade::GroupPrefix.ToString() + FString("Vertices"));
+	//const FName FGroomGuidesFacade::VerticesGroup(FString("GuidesVertices"));
+
+	// Guides Attributes
+	const FName FGroomGuidesFacade::VertexKinematicWeightsAttribute("VertexKinematicWeights");
+	const FName FGroomGuidesFacade::PointBoneIndicesAttribute("PointBoneIndices");
+	const FName FGroomGuidesFacade::PointBoneWeightsAttribute("PointBoneWeights");
+	const FName FGroomGuidesFacade::ObjectPointSamplesAttribute("ObjectPointSamples");
+	const FName FGroomGuidesFacade::CurveStrandIndicesAttribute("CurveStrandIndices");
+	const FName FGroomGuidesFacade::CurveLevelIndicesAttribute("CurveLevelIndices");
+	const FName FGroomGuidesFacade::PointInterpIndicesAttribute("PointInterpIndices");
+	const FName FGroomGuidesFacade::PointInterpWeightsAttribute("PointInterpWeights");
+	
+	template<typename DerivedType>
+	FGroomCollectionFacade<DerivedType>::FGroomCollectionFacade(FManagedArrayCollection& InCollection) :
+		ConstCollection(InCollection), Collection(&InCollection),
+		EdgeRestOrientations(InCollection, EdgeRestOrientationsAttribute, EdgesGroup),
+		PointRestPositions(InCollection, PointRestPositionsAttribute, PointsGroup),
+		CurvePointOffsets(InCollection, CurvePointOffsetsAttribute, CurvesGroup),
+		ObjectCurveOffsets(InCollection, ObjectCurveOffsetsAttribute, ObjectsGroup),
+		PointCurveIndices(InCollection, PointCurveIndicesAttribute, PointsGroup),
+		CurveObjectIndices(InCollection, CurveObjectIndicesAttribute, CurvesGroup)
+	{}
+
+	template<typename DerivedType>
+	FGroomCollectionFacade<DerivedType>::FGroomCollectionFacade(const FManagedArrayCollection& InCollection) :
+		ConstCollection(InCollection), Collection(nullptr),
+		EdgeRestOrientations(InCollection, EdgeRestOrientationsAttribute, EdgesGroup),
+		PointRestPositions(InCollection, PointRestPositionsAttribute, PointsGroup),
+		CurvePointOffsets(InCollection, CurvePointOffsetsAttribute, CurvesGroup),
+		ObjectCurveOffsets(InCollection, ObjectCurveOffsetsAttribute, ObjectsGroup),
+		PointCurveIndices(InCollection, PointCurveIndicesAttribute, PointsGroup),
+		CurveObjectIndices(InCollection, CurveObjectIndicesAttribute, CurvesGroup)
+	{
+	}
+
+	template<typename DerivedType>
+	bool FGroomCollectionFacade<DerivedType>::IsValid() const
+	{
+		return EdgeRestOrientations.IsValid() && PointRestPositions.IsValid() && 
+			CurvePointOffsets.IsValid() && ObjectCurveOffsets.IsValid() && PointCurveIndices.IsValid() &&
+					CurveObjectIndices.IsValid() && static_cast<const DerivedType*>(this)->IsFacadeValid();
+	}
+
+	template<typename DerivedType>
+	void FGroomCollectionFacade<DerivedType>::DefineSchema()
+	{
+		check(!IsConst());
+		EdgeRestOrientations.Add();
+		PointRestPositions.Add();
+		CurvePointOffsets.Add();
+		ObjectCurveOffsets.Add();
+		PointCurveIndices.Add();
+		CurveObjectIndices.Add();
+		
+		static_cast<DerivedType*>(this)->DefineFacadeSchema();
+	}
+
+	template<typename DerivedType>
+	void FGroomCollectionFacade<DerivedType>::InitGroomCollection(const TArray<FVector3f>& InPointRestPositions,
+		const TArray<int32>& InCurvePointOffsets, const TArray<int32>& InObjectCurveOffsets)
+	{
+		if(Collection)
+		{
+			// Curves group
+			Collection->EmptyGroup(CurvesGroup);
+			Collection->AddElements(InCurvePointOffsets.Num(), CurvesGroup);
+
+			// Object group
+			Collection->EmptyGroup(ObjectsGroup);
+			Collection->AddElements(InObjectCurveOffsets.Num(), ObjectsGroup);
+
+			// Points group
+			Collection->EmptyGroup(PointsGroup);
+			Collection->AddElements(InPointRestPositions.Num(), PointsGroup);
+			
+			// Edges group
+			Collection->EmptyGroup(EdgesGroup);
+			Collection->AddElements(InPointRestPositions.Num()-InCurvePointOffsets.Num(), EdgesGroup);
+
+			// Fill attributes
+			SetObjectCurveOffsets(InObjectCurveOffsets);
+			SetCurvePointOffsets(InCurvePointOffsets);
+			SetPointRestPositions(InPointRestPositions);
+
+			static_cast<DerivedType*>(this)->InitFacadeCollection();
+		}
+	}
+
+	template<typename DerivedType>
+	void FGroomCollectionFacade<DerivedType>::UpdateCurveObjectIndices()
+	{
+		const int32 NumObjects = GetNumObjects();
+
+		int32 CurveOffset = 0;
+		for(int32 ObjectIndex = 0; ObjectIndex < NumObjects; ++ObjectIndex)
+		{
+			for(int32 CurveIndex = CurveOffset, CurveEnd = ObjectCurveOffsets[ObjectIndex];
+					  CurveIndex < CurveEnd; ++CurveIndex)
+			{
+				CurveObjectIndices.ModifyAt(CurveIndex,ObjectIndex);
+			}
+			CurveOffset = ObjectCurveOffsets[ObjectIndex];
+		}
+	}
+
+	template<typename DerivedType>
+	void FGroomCollectionFacade<DerivedType>::UpdatePointCurveIndices()
+	{
+		const int32 NumCurves = GetNumCurves();
+
+		int32 PointOffset = 0;
+		for(int32 CurveIndex = 0; CurveIndex < NumCurves; ++CurveIndex)
+		{
+			for (int32 PointIndex = PointOffset, PointEnd = CurvePointOffsets[CurveIndex];
+					   PointIndex < PointEnd; ++PointIndex)
+			{
+				PointCurveIndices.ModifyAt(PointIndex, CurveIndex);
+			}
+			PointOffset = CurvePointOffsets[CurveIndex];
+		}
+	}
+
+	template<typename DerivedType>
+	void FGroomCollectionFacade<DerivedType>::UpdateEdgeRestOrientations()
+	{
+		const int32 NumCurves = GetNumCurves();
+
+		int32 EdgeOffset = 0;
+		int32 PointOffset = 0;
+		for(int32 CurveIndex = 0; CurveIndex < NumCurves; ++CurveIndex)
+		{
+			FVector3f TangentPrev = FVector3f(0.,0.,1.), TangentNext = FVector3f::Zero();
+			FQuat4f EdgeOrientation = FQuat4f::Identity;
+
+			const int32 NumEdges = (CurvePointOffsets[CurveIndex] - PointOffset)-1;
+			for (int32 EdgeIndex = 0; EdgeIndex < NumEdges; ++EdgeIndex)
+			{
+				TangentPrev = TangentNext;
+				TangentNext = (PointRestPositions[PointOffset+EdgeIndex+1]-PointRestPositions[PointOffset+EdgeIndex]).GetSafeNormal();
+
+				EdgeOrientation = (FQuat4f::FindBetweenNormals(TangentPrev,TangentNext) * EdgeOrientation).GetNormalized();
+				EdgeRestOrientations.ModifyAt(EdgeOffset+EdgeIndex, EdgeOrientation);
+			}
+			EdgeOffset += NumEdges;
+			PointOffset = CurvePointOffsets[CurveIndex];
+		}
+	}
+
+	template class FGroomCollectionFacade<FGroomStrandsFacade>;
+	template class FGroomCollectionFacade<FGroomGuidesFacade>;
+
+	// Strands
+	
+	FGroomStrandsFacade::FGroomStrandsFacade(FManagedArrayCollection& InCollection) :
+		FGroomCollectionFacade<FGroomStrandsFacade>(InCollection)
+	{
+		DefineSchema();
+	}
+
+	FGroomStrandsFacade::FGroomStrandsFacade(const FManagedArrayCollection& InCollection) :
+		FGroomCollectionFacade<FGroomStrandsFacade>(InCollection)
+	{
+	}
+
+	// Guides
+	
+	FGroomGuidesFacade::FGroomGuidesFacade(FManagedArrayCollection& InCollection) :
+		FGroomCollectionFacade<FGroomGuidesFacade>(InCollection),
+		VertexKinematicWeights(InCollection, VertexKinematicWeightsAttribute, VerticesGroup),
+		PointBoneIndices(InCollection, PointBoneIndicesAttribute, PointsGroup),
+		PointBoneWeights(InCollection, PointBoneWeightsAttribute, PointsGroup),
+		ObjectPointSamples(InCollection, ObjectPointSamplesAttribute, ObjectsGroup),
+		CurveStrandIndices(InCollection, CurveStrandIndicesAttribute, CurvesGroup)
+		// CurveLevelIndices(InCollection, CurveLevelIndicesAttribute, CurvesGroup),
+		// PointInterpIndices(InCollection, PointInterpIndicesAttribute, PointsGroup),
+		// PointInterpWeights(InCollection, PointInterpWeightsAttribute, PointsGroup)
+	{
+		DefineSchema();
+	}
+
+	FGroomGuidesFacade::FGroomGuidesFacade(const FManagedArrayCollection& InCollection) :
+		FGroomCollectionFacade<FGroomGuidesFacade>(InCollection),
+		VertexKinematicWeights(InCollection, VertexKinematicWeightsAttribute, VerticesGroup),
+		PointBoneIndices(InCollection, PointBoneIndicesAttribute, PointsGroup),
+		PointBoneWeights(InCollection, PointBoneWeightsAttribute, PointsGroup),
+		ObjectPointSamples(InCollection, ObjectPointSamplesAttribute, ObjectsGroup),
+		CurveStrandIndices(InCollection, CurveStrandIndicesAttribute, CurvesGroup)
+		// CurveLevelIndices(InCollection, CurveLevelIndicesAttribute, CurvesGroup),
+		// PointInterpIndices(InCollection, PointInterpIndicesAttribute, PointsGroup),
+		// PointInterpWeights(InCollection, PointInterpWeightsAttribute, PointsGroup)
+	{
+	}
+
+	bool FGroomGuidesFacade::IsFacadeValid() const
+	{
+		return VertexKinematicWeights.IsValid() && PointBoneIndices.IsValid() && PointBoneWeights.IsValid() && ObjectPointSamples.IsValid();
+	}
+
+	bool FGroomGuidesFacade::HasSkinningAttributes() const
+	{
+		return PointBoneIndices.IsValid() && PointBoneWeights.IsValid() && ObjectCurveOffsets.IsValid() && CurvePointOffsets.IsValid();
+	}
+
+	void FGroomGuidesFacade::DefineFacadeSchema()
+	{
+		check(!IsConst());
+		VertexKinematicWeights.Add();
+		PointBoneWeights.Add();
+		PointBoneIndices.Add();
+		ObjectPointSamples.Add();
+		CurveStrandIndices.Add();
+		// CurveLevelIndices.Add();
+		// PointInterpIndices.Add();
+		// PointInterpWeights.Add();
+	}
+
+	void FGroomGuidesFacade::InitFacadeCollection()
+	{
+		// Reset groups
+		Collection->EmptyGroup(VerticesGroup);
+		Collection->AddElements(GetNumPoints() * 2, VerticesGroup);
+
+		// Reset attributes
+		VertexKinematicWeights.Modify().Fill(0.0f);
+		PointBoneIndices.Modify().Fill(TArray<int32>());
+		PointBoneWeights.Modify().Fill(TArray<float>());
+		ObjectPointSamples.Modify().Fill(4);
+		CurveStrandIndices.Modify().Fill(INDEX_NONE);
+		// CurveLevelIndices.Modify().Fill(0);
+		// PointInterpIndices.Modify().Fill(FVector4f(INDEX_NONE));
+		// PointInterpWeights.Modify().Fill(FVector4f(0.0f));
+	}
+
+	void FGroomGuidesFacade::ResizePointsGroups(const int32 NumPoints) const 
+	{
+		if(Collection)
+		{
+			const int32 NumCurves = GetNumCurves();
+			
+			Collection->EmptyGroup(PointsGroup);
+			Collection->AddElements(NumPoints, PointsGroup);
+
+			Collection->EmptyGroup(EdgesGroup);
+			Collection->AddElements(NumPoints-NumCurves, EdgesGroup);
+
+			Collection->EmptyGroup(VerticesGroup);
+			Collection->AddElements(NumPoints*2, VerticesGroup);
+		}
+	}
+}
+
diff --git a/Engine/Plugins/Runtime/HairStrands/Source/HairStrandsDataflow/Private/GroomDataflowRendering.cpp b/Engine/Plugins/Runtime/HairStrands/Source/HairStrandsDataflow/Private/GroomDataflowRendering.cpp
new file mode 100644
--- /dev/null
+++ b/Engine/Plugins/Runtime/HairStrands/Source/HairStrandsDataflow/Private/GroomDataflowRendering.cpp
@@ -0,0 +1,250 @@
+﻿// Copyright Epic Games, Inc. All Rights Reserved.
+
+#include "GroomDataflowRendering.h"
+#include "GroomCollectionFacades.h"
+#include "GeometryCollection/GeometryCollection.h"
+
+namespace UE::Groom
+{
+	int32 GGroomDataflowDebugDraw = 0;
+	static FAutoConsoleVariableRef CVarGroomDataflowDebugDraw(
+		 TEXT("p.Groom.Dataflow.DebugDraw"),
+		 GGroomDataflowDebugDraw,
+		 TEXT("Type of information we want to draw for debug (0-DefaultColor, 1-SkinWeights)")
+	 );
+
+	int32 GGroomDataflowBoneIndex = 0;
+	static FAutoConsoleVariableRef CVarGroomDataflowBoneIndex(
+		 TEXT("p.Groom.Dataflow.BoneIndex"),
+		 GGroomDataflowBoneIndex,
+		 TEXT("Bone index we want to visualize the skin weights")
+	 );
+
+	float GGroomDataflowRenderingThickness = 0.5f;
+	static FAutoConsoleVariableRef CVarRestGuidesRenderingThickness(
+		 TEXT("p.Groom.Dataflow.RenderingThickness"),
+		 GGroomDataflowRenderingThickness,
+		 TEXT("Thickness used to render the groom in dataflow editor.")
+	 );
+
+	FORCEINLINE void BuildVertexPositions(const int32 PointIndex, const FVector3f& PointPosition, const FVector3f& SideVector, const float RenderThickness, TArray<FVector3f>& VertexPositions)
+	{
+		const int32 VertexIndex = 2 * PointIndex;
+		
+		const FVector3f PointPositionA = PointPosition + SideVector * RenderThickness;
+		const FVector3f PointPositionB = PointPosition - SideVector * RenderThickness;
+	
+		VertexPositions[VertexIndex] = PointPositionA;
+		VertexPositions[VertexIndex+1] = PointPositionB;
+	}
+	
+	FORCEINLINE void BuildFaceVertices(const int32 PointIndex, const int32 CurveIndex, TArray<FIntVector>& FaceVertices)
+	{
+		const int32 VertexIndex = 2 * PointIndex;
+		const int32 FaceIndex = 2 * (PointIndex - CurveIndex);
+				
+		FaceVertices[FaceIndex] = FIntVector(VertexIndex, VertexIndex + 1, VertexIndex + 3);
+		FaceVertices[FaceIndex+1] = FIntVector(VertexIndex, VertexIndex + 3, VertexIndex + 2);
+	}
+
+	FORCEINLINE void BuildVertexNormals(const int32 PointIndex, const FVector3f& EdgeNormal, TArray<FVector3f>& VertexNormals)
+	{
+		const int32 VertexIndex = 2 * PointIndex;
+
+		VertexNormals[VertexIndex] = EdgeNormal;
+		VertexNormals[VertexIndex+1] = EdgeNormal;
+	}
+
+	template<typename FacadeType>
+	FORCEINLINE void BuildRenderingDatas(const FacadeType& GroomFacade, TArray<FVector3f>& VertexPositions, TArray<FVector3f>& VertexNormals, TArray<FIntVector>& FaceVertices)
+	{
+		const int32 NumPoints = GroomFacade.GetNumPoints();
+		const int32 NumEdges = GroomFacade.GetNumEdges();
+		const int32 NumCurves = GroomFacade.GetNumCurves();
+
+		VertexPositions.SetNum(NumPoints * 2);
+		VertexNormals.SetNum(NumPoints * 2);
+		FaceVertices.SetNum(NumEdges * 2);
+
+		int32 PointOffset = 0;
+		for(int32 CurveIndex = 0; CurveIndex < NumCurves; ++CurveIndex)
+		{
+			for (int32 PointIndex = PointOffset, PointEnd = GroomFacade.GetCurvePointOffsets()[CurveIndex];
+					   PointIndex < PointEnd; ++PointIndex)
+			{
+				// Build 2 vertices for each points
+				const int32 EdgeIndex = (PointIndex == (PointEnd-1)) ? PointIndex-CurveIndex-1 : PointIndex-CurveIndex;
+								
+				const FVector3f SideVector = GroomFacade.GetEdgeRestOrientations()[EdgeIndex].GetAxisX();
+				BuildVertexPositions(PointIndex, GroomFacade.GetPointRestPositions()[PointIndex], SideVector, GGroomDataflowRenderingThickness, VertexPositions);
+
+				const FVector3f NormalVector = GroomFacade.GetEdgeRestOrientations()[EdgeIndex].GetAxisZ();
+				BuildVertexNormals(PointIndex, NormalVector, VertexNormals);
+					
+				// Build 2 faces for each edges perpendicular to the up-vector curve
+				if(PointIndex < (PointEnd-1))
+				{
+					BuildFaceVertices(PointIndex, CurveIndex, FaceVertices);
+				}
+			}
+			PointOffset = GroomFacade.GetCurvePointOffsets()[CurveIndex];
+		}
+	}
+
+	template<typename FacadeType>
+	FORCEINLINE void BuildGeometryDatas(const FacadeType& GroomFacade, TArray<float>& VertexPositions, TArray<int32>& FaceVertices, TArray<int32>& ObjectFaceOffsets, TArray<int32>& ObjectVertexOffsets)
+	{
+		TArray<FVector3f> RenderingPositions, RenderingNormals;
+		TArray<FIntVector> RenderingVertices;
+
+		BuildRenderingDatas(GroomFacade, RenderingPositions, RenderingNormals, RenderingVertices);
+
+		VertexPositions.SetNum(RenderingPositions.Num() * 3);
+		FaceVertices.SetNum(RenderingVertices.Num() * 3);
+
+		int32 VertexIndex = 0;
+		for(const FVector3f& VertexPosition : RenderingPositions)
+		{
+			VertexPositions[VertexIndex++] = VertexPosition.X;
+			VertexPositions[VertexIndex++] = VertexPosition.Y;
+			VertexPositions[VertexIndex++] = VertexPosition.Z;
+		}
+
+		int32 FaceIndex = 0;
+		for(const FIntVector& FaceVertex : RenderingVertices)
+		{
+			FaceVertices[FaceIndex++] = FaceVertex.X;
+			FaceVertices[FaceIndex++] = FaceVertex.Y;
+			FaceVertices[FaceIndex++] = FaceVertex.Z;
+		}
+
+		const int32 NumObjects = GroomFacade.GetNumObjects();
+		ObjectFaceOffsets.SetNum(NumObjects);
+		ObjectVertexOffsets.SetNum(NumObjects);
+		
+		for(int32 ObjectIndex = 0; ObjectIndex < NumObjects; ++ObjectIndex)
+		{
+			const int32 CurveOffset = GroomFacade.GetObjectCurveOffsets()[ObjectIndex]-1;
+			ObjectVertexOffsets[ObjectIndex] = 2 * GroomFacade.GetCurvePointOffsets()[CurveOffset];
+			ObjectFaceOffsets[ObjectIndex] = ObjectVertexOffsets[ObjectIndex] - 2 * CurveOffset;
+		}
+	}
+
+	template<typename FacadeType>
+	FORCEINLINE void SetupGeometryCollection(const FacadeType& GroomFacade, FGeometryCollection& GeometryCollection)
+	{
+		TArray<float> RawVertexPositions;
+		TArray<int32> RawFaceIndices;
+			
+		TArray<int32> ObjectFaceOffsets;
+		TArray<int32> ObjectVertexOffsets;
+
+		BuildGeometryDatas(GroomFacade, RawVertexPositions, RawFaceIndices, ObjectFaceOffsets, ObjectVertexOffsets);
+		FGeometryCollection::Init(&GeometryCollection, RawVertexPositions, RawFaceIndices, false);
+
+		// add a material section
+		TManagedArray<FGeometryCollectionSection>& Sections = GeometryCollection.Sections;
+		GeometryCollection.Resize(GroomFacade.GetNumObjects(), FGeometryCollection::MaterialGroup);
+			
+		int32 FaceOffset = 0;
+		int32 VertexOffset = 0;
+		for(int32 SectionIndex = 0; SectionIndex < GroomFacade.GetNumObjects(); ++SectionIndex)
+		{
+			Sections[SectionIndex].MaterialID = SectionIndex;
+			Sections[SectionIndex].FirstIndex = FaceOffset;
+			Sections[SectionIndex].NumTriangles = ObjectFaceOffsets[SectionIndex] - FaceOffset;
+			Sections[SectionIndex].MinVertexIndex = VertexOffset;
+			Sections[SectionIndex].MaxVertexIndex = ObjectVertexOffsets[SectionIndex] - 1;
+
+			FaceOffset = ObjectFaceOffsets[SectionIndex];
+			VertexOffset = ObjectVertexOffsets[SectionIndex];
+
+			SectionIndex++;
+		}
+	}
+	
+	template<typename FacadeType>
+	FORCEINLINE void RenderGroomCollection(GeometryCollection::Facades::FRenderingFacade& RenderCollection, const Dataflow::FGraphRenderingState& State,
+		const TFunction<void(const FacadeType&, TArray<FLinearColor>&)>& ColorLambda)
+	{
+		if (State.GetRenderOutputs().Num())
+		{
+			const FManagedArrayCollection Default;
+			checkf(State.GetRenderOutputs().Num() == 1, TEXT("Expected FGraphRenderingState object to have one render output"));
+			const FName PrimaryOutput = State.GetRenderOutputs()[0];
+			FManagedArrayCollection Collection = State.GetValue<FManagedArrayCollection>(PrimaryOutput, Default);
+
+			const TSharedRef<const FManagedArrayCollection> GroomCollection = MakeShared<FManagedArrayCollection>(MoveTemp(Collection));
+			const FacadeType GroomFacade(*GroomCollection);
+			if (GroomFacade.IsValid())
+			{
+				TArray<FVector3f> VertexPositions;
+				TArray<FVector3f> VertexNormals;
+				TArray<FIntVector> FaceIndices;
+
+				if (State.GetViewMode().GetName() == Dataflow::FDataflowConstruction3DViewMode::Name)
+				{
+					BuildRenderingDatas(GroomFacade, VertexPositions, VertexNormals, FaceIndices);
+				}
+				else
+				{
+					checkf(false, TEXT("Invalid View Mode for FClothCollection rendering"));
+				}
+
+				TArray<FLinearColor> Colors;
+				Colors.Init(FLinearColor::Blue, VertexPositions.Num());
+				ColorLambda(GroomFacade, Colors);
+ 
+				const int32 GeometryIndex = RenderCollection.StartGeometryGroup(State.GetGuid().ToString());
+				RenderCollection.AddSurface(MoveTemp(VertexPositions), MoveTemp(FaceIndices), MoveTemp(VertexNormals), MoveTemp(Colors));
+				RenderCollection.EndGeometryGroup(GeometryIndex);
+			}
+		}
+	}
+	
+	Dataflow::FRenderKey FGroomGuidesRenderingCallbacks::RenderKey = { "GuidesRender", FName("FGroomCollection") };
+	
+	void FGroomGuidesRenderingCallbacks::Render(GeometryCollection::Facades::FRenderingFacade& RenderCollection, const Dataflow::FGraphRenderingState& State)
+	{
+		RenderGroomCollection<FGroomGuidesFacade>(RenderCollection, State,
+			[](const FGroomGuidesFacade& GuidesFacade, TArray<FLinearColor>& VertexColors)
+			{
+				if(GGroomDataflowDebugDraw == 1)
+				{
+					for(int32 PointIndex = 0, PointEnd = GuidesFacade.GetNumPoints(); PointIndex < PointEnd; ++PointIndex)
+					{
+						const TArray<int32>& BoneIndices = GuidesFacade.GetPointBoneIndices(PointIndex);
+						const TArray<float>& BoneWeights = GuidesFacade.GetPointBoneWeights(PointIndex);
+						for(int32 BoneIndex = 0, BoneEnd = BoneIndices.Num(); BoneIndex < BoneEnd; ++BoneIndex)
+						{
+							if(BoneIndices[BoneIndex] == GGroomDataflowBoneIndex)
+							{
+								VertexColors[2*PointIndex] = FLinearColor::LerpUsingHSV(
+									FLinearColor::Black, FLinearColor::Yellow, BoneWeights[BoneIndex]);
+								VertexColors[2*PointIndex+1] = VertexColors[2*PointIndex];
+							}
+						}	
+					}
+				}
+			});
+	}
+
+	Dataflow::FRenderKey FGroomStrandsRenderingCallbacks::RenderKey = { "StrandsRender", FName("FGroomCollection") };
+	
+	void FGroomStrandsRenderingCallbacks::Render(GeometryCollection::Facades::FRenderingFacade& RenderCollection, const Dataflow::FGraphRenderingState& State)
+	{
+		RenderGroomCollection<FGroomStrandsFacade>(RenderCollection, State, [](const FGroomStrandsFacade& StrandsFacade, TArray<FLinearColor>& VertexColors){});
+	}
+
+	void RegisterRenderingCallbacks()
+	{
+		Dataflow::FRenderingFactory::GetInstance()->RegisterCallbacks(MakeUnique<FGroomGuidesRenderingCallbacks>());
+		Dataflow::FRenderingFactory::GetInstance()->RegisterCallbacks(MakeUnique<FGroomStrandsRenderingCallbacks>());
+	}
+
+	void DeregisterRenderingCallbacks()
+	{
+		Dataflow::FRenderingFactory::GetInstance()->DeregisterCallbacks(FGroomGuidesRenderingCallbacks::RenderKey);
+		Dataflow::FRenderingFactory::GetInstance()->DeregisterCallbacks(FGroomStrandsRenderingCallbacks::RenderKey);
+	}
+}
diff --git a/Engine/Plugins/Runtime/HairStrands/Source/HairStrandsDataflow/Private/GroomDataflowVisualization.cpp b/Engine/Plugins/Runtime/HairStrands/Source/HairStrandsDataflow/Private/GroomDataflowVisualization.cpp
new file mode 100644
--- /dev/null
+++ b/Engine/Plugins/Runtime/HairStrands/Source/HairStrandsDataflow/Private/GroomDataflowVisualization.cpp
@@ -0,0 +1,140 @@
+﻿// Copyright Epic Games, Inc. All Rights Reserved.
+
+#include "GroomDataflowVisualization.h"
+#include "GroomComponent.h"
+#include "GroomVisualizationData.h"
+#include "Dataflow/DataflowEditorToolkit.h"
+#include "Dataflow/DataflowSimulationViewportClient.h"
+#include "Dataflow/DataflowSimulationScene.h"
+#include "Elements/Columns/TypedElementAlertColumns.h"
+
+#define LOCTEXT_NAMESPACE "ChaosClothAssetDataflowSimulationVisualization"
+
+namespace UE::Groom
+{
+	const FName FGroomDataflowSimulationVisualization::Name = FName("GroomDataflowSimulationVisualization");
+
+	FName FGroomDataflowSimulationVisualization::GetName() const
+	{
+		return Name;
+	}
+
+	FGroomDataflowSimulationVisualization::FGroomDataflowSimulationVisualization()
+	{
+		VisualizationFlags.Init(false, static_cast<uint8>(EGroomViewMode::Count));
+	}
+
+	void FGroomDataflowSimulationVisualization::ExtendSimulationVisualizationMenu(const TSharedPtr<FDataflowSimulationViewportClient>& ViewportClient, FMenuBuilder& MenuBuilder)
+	{
+		if (ViewportClient)
+		{
+			if (TSharedPtr<FDataflowEditorToolkit> Toolkit = ViewportClient->GetDataflowEditorToolkit().Pin())
+			{
+				if (const TSharedPtr<FDataflowSimulationScene>& SimulationScene = Toolkit->GetSimulationScene())
+				{
+					if (GetGroomComponent(SimulationScene.Get()))
+					{
+						static IConsoleVariable* CVarGroomViewMode = nullptr; 
+						if (CVarGroomViewMode == nullptr)
+						{
+							if (IConsoleVariable* ConsoleVariable = IConsoleManager::Get().FindConsoleVariable(
+								FGroomVisualizationData::GetVisualizeConsoleCommandName()))
+							{
+								CVarGroomViewMode = ConsoleVariable->AsVariable();
+							}
+						}
+						static IConsoleVariable* CVarGroomStrandWidth = nullptr; 
+						if (CVarGroomStrandWidth == nullptr)
+						{
+							if (IConsoleVariable* ConsoleVariable = IConsoleManager::Get().FindConsoleVariable(
+								TEXT("r.HairStrands.StrandWidth")))
+							{
+								CVarGroomStrandWidth = ConsoleVariable->AsVariable();
+							}
+						}
+						constexpr uint8 NumFlags = static_cast<uint8>(EGroomViewMode::Count);
+						auto AddSimulationVisualisation = [this, &MenuBuilder, NumFlags](const TAttribute<FText>& FlagLabel, const uint8 FlagIndex, const int32 ViewMode, const float StrandWidth)
+						{
+							// Handler for visualization entry being clicked
+							const FExecuteAction ExecuteAction = FExecuteAction::CreateLambda([this, FlagIndex, ViewMode, StrandWidth]()
+							{
+								VisualizationFlags[FlagIndex] = !VisualizationFlags[FlagIndex];
+								if(VisualizationFlags[FlagIndex])
+								{
+									CVarGroomViewMode->SetWithCurrentPriority(ViewMode);
+									CVarGroomStrandWidth->SetWithCurrentPriority(StrandWidth);
+									for(int32 OtherIndex = 0; OtherIndex < NumFlags; ++OtherIndex)
+									{
+										if(OtherIndex != FlagIndex)
+										{
+											VisualizationFlags[OtherIndex] = false;
+										}
+									}
+								}
+								else
+								{
+									CVarGroomViewMode->SetWithCurrentPriority(0);
+									CVarGroomStrandWidth->SetWithCurrentPriority(0.0f);
+								}
+							});
+
+							// Checkstate function for visualization entries
+							const FIsActionChecked IsActionChecked = FIsActionChecked::CreateLambda([this, FlagIndex]()
+							{
+								return VisualizationFlags[FlagIndex];
+							});
+
+							const FUIAction Action(ExecuteAction, FCanExecuteAction(), IsActionChecked);
+
+							// Add menu entry
+							MenuBuilder.AddMenuEntry(FlagLabel, FText(), FSlateIcon(), Action, NAME_None, EUserInterfaceActionType::ToggleButton);
+						};
+						MenuBuilder.BeginSection(TEXT("GroomSimulation_Visualizations"), LOCTEXT("GroomSimulationVisualization", "Groom Simulation Visualization"));
+
+						AddSimulationVisualisation(LOCTEXT("SimulatedGuides", "Simulated guides"), static_cast<uint8>(EGroomViewMode::SimHairStrands), 21, 0.5f);
+						AddSimulationVisualisation(LOCTEXT("RenderedStrands", "Rendered strands"), static_cast<uint8>(EGroomViewMode::RenderHairStrands), 22, 0.0f);
+						AddSimulationVisualisation(LOCTEXT("StrandsGroups", "Strands groups"), static_cast<uint8>(EGroomViewMode::Group), 35, 0.0f);
+						AddSimulationVisualisation(LOCTEXT("StrandsClumps", "Strands clumps"), static_cast<uint8>(EGroomViewMode::ClumpID), 39, 0.0f);
+						AddSimulationVisualisation(LOCTEXT("StrandsClusters", "Strands clusters"), static_cast<uint8>(EGroomViewMode::Cluster), 31, 0.0f);
+						AddSimulationVisualisation(LOCTEXT("MeshProjection", "Mesh projection"), static_cast<uint8>(EGroomViewMode::MeshProjection), 12, 0.0f);
+						AddSimulationVisualisation(LOCTEXT("CardsGuides", "Cards guides"), static_cast<uint8>(EGroomViewMode::CardGuides), 37, 0.0f);
+						AddSimulationVisualisation(LOCTEXT("LODColoration", "LOD Coloration"), static_cast<uint8>(EGroomViewMode::LODColoration), 36, 0.0f);
+						
+						MenuBuilder.EndSection();
+					}
+				}
+			}
+		}
+	}
+
+	const UGroomComponent* FGroomDataflowSimulationVisualization::GetGroomComponent(const FDataflowSimulationScene* SimulationScene) const
+	{
+		if (SimulationScene)
+		{
+			if (const TObjectPtr<AActor> PreviewActor = SimulationScene->GetPreviewActor())
+			{
+				return PreviewActor->GetComponentByClass<UGroomComponent>();
+			}
+		}
+		return nullptr;
+	}
+
+	void FGroomDataflowSimulationVisualization::Draw(const FDataflowSimulationScene* SimulationScene, FPrimitiveDrawInterface* PDI)
+	{
+		
+	}
+
+	void FGroomDataflowSimulationVisualization::DrawCanvas(const FDataflowSimulationScene* SimulationScene, FCanvas* Canvas, const FSceneView* SceneView)
+	{
+		
+	}
+
+	FText FGroomDataflowSimulationVisualization::GetDisplayString(const FDataflowSimulationScene* SimulationScene) const
+	{
+		return FText();
+	}
+
+}  // End namespace UE::Groom
+
+#undef LOCTEXT_NAMESPACE 
+ 
\ No newline at end of file
diff --git a/Engine/Plugins/Runtime/HairStrands/Source/HairStrandsDataflow/Private/GroomDataflowVisualization.h b/Engine/Plugins/Runtime/HairStrands/Source/HairStrandsDataflow/Private/GroomDataflowVisualization.h
new file mode 100644
--- /dev/null
+++ b/Engine/Plugins/Runtime/HairStrands/Source/HairStrandsDataflow/Private/GroomDataflowVisualization.h
@@ -0,0 +1,36 @@
+﻿// Copyright Epic Games, Inc. All Rights Reserved.
+
+#pragma once
+
+#include "Dataflow/DataflowSimulationVisualization.h"
+
+class UGroomComponent;
+
+namespace UE::Groom
+{
+	/** Dataflow simulation customisation for groom */
+	class FGroomDataflowSimulationVisualization : public Dataflow::IDataflowSimulationVisualization
+	{
+	public:
+		static const FName Name;
+		virtual ~FGroomDataflowSimulationVisualization() = default;
+		FGroomDataflowSimulationVisualization();
+
+	private:
+
+		//~ Begin IDataflowSimulationVisualization interface
+		virtual FName GetName() const override;
+		virtual void ExtendSimulationVisualizationMenu(const TSharedPtr<FDataflowSimulationViewportClient>& ViewportClient, FMenuBuilder& MenuBuilder) override;
+		virtual void Draw(const FDataflowSimulationScene* SimulationScene, FPrimitiveDrawInterface* PDI) override;
+		virtual void DrawCanvas(const FDataflowSimulationScene* SimulationScene, FCanvas* Canvas, const FSceneView* SceneView) override;
+		virtual FText GetDisplayString(const FDataflowSimulationScene* SimulationScene) const override;
+		//~ End IDataflowSimulationVisualization interface
+
+		/** Get the groom component given a simulation scene */
+		const UGroomComponent* GetGroomComponent(const FDataflowSimulationScene* SimulationScene) const;
+
+		/** Visualization flags */
+		TArray<bool> VisualizationFlags;
+	};
+
+}  // End namespace UE::Chaos::ClothAsset
diff --git a/Engine/Plugins/Runtime/HairStrands/Source/HairStrandsDataflow/Private/HairStrandsDataflowModule.cpp b/Engine/Plugins/Runtime/HairStrands/Source/HairStrandsDataflow/Private/HairStrandsDataflowModule.cpp
new file mode 100644
--- /dev/null
+++ b/Engine/Plugins/Runtime/HairStrands/Source/HairStrandsDataflow/Private/HairStrandsDataflowModule.cpp
@@ -0,0 +1,77 @@
+﻿// Copyright Epic Games, Inc. All Rights Reserved.
+
+#include "HairStrandsDataflowModule.h"
+#include "GetGroomAssetNode.h"
+#include "GroomCollectionFacades.h"
+#include "GroomAssetTerminalNode.h"
+#include "GetGuidesAttributesNodes.h"
+#include "ResampleGuidesPointsNode.h"
+#include "BuildGuidesSkinningNode.h"
+#include "GenerateGuidesCurvesNode.h"
+#include "SmoothGuidesCurvesNode.h"
+#include "GroomDataflowRendering.h"
+#include "GroomDataflowVisualization.h"
+#include "Dataflow/DataflowCollectionAddScalarVertexPropertyNode.h"
+#include "Modules/ModuleManager.h"
+
+#define LOCTEXT_NAMESPACE "HairStrandsDataflow"
+
+class FGroomCollectionAddScalarVertexPropertyCallbacks : public IDataflowAddScalarVertexPropertyCallbacks
+{
+public:
+
+	const static FName Name;
+
+	virtual ~FGroomCollectionAddScalarVertexPropertyCallbacks() override = default;
+
+	virtual FName GetName() const override
+	{
+		return Name;
+	}
+
+	virtual TArray<FName> GetTargetGroupNames() const override
+	{
+		return {UE::Groom::FGroomGuidesFacade::VerticesGroup};
+	}
+
+	virtual TArray<UE::Dataflow::FRenderingParameter> GetRenderingParameters() const override
+	{
+		return { {TEXT("GuidesRender"), FName("FGroomCollection"), {TEXT("Collection")}}};
+	}
+};
+
+const FName FGroomCollectionAddScalarVertexPropertyCallbacks::Name = FName(TEXT("FGroomCollectionAddScalarVertexPropertyCallbacks"));
+
+void FHairStrandsDataflowModule::StartupModule()
+{
+	DATAFLOW_NODE_REGISTER_CREATION_FACTORY(FGetGroomAssetDataflowNode);
+	DATAFLOW_NODE_REGISTER_CREATION_FACTORY(FGroomAssetTerminalDataflowNode);
+	DATAFLOW_NODE_REGISTER_CREATION_FACTORY(FResampleGuidesPointsDataflowNode);
+	DATAFLOW_NODE_REGISTER_CREATION_FACTORY(FBuildGuidesSkinningDataflowNode);
+	DATAFLOW_NODE_REGISTER_CREATION_FACTORY(FGenerateGuidesCurvesDataflowNode);
+	DATAFLOW_NODE_REGISTER_CREATION_FACTORY(FGetGuidesKinematicWeightsDataflowNode);
+	DATAFLOW_NODE_REGISTER_CREATION_FACTORY(FSmoothGuidesCurvesDataflowNode);
+	
+	UE::Groom::RegisterRenderingCallbacks();
+	
+	FDataflowAddScalarVertexPropertyCallbackRegistry::Get().RegisterCallbacks(
+		MakeUnique<FGroomCollectionAddScalarVertexPropertyCallbacks>());
+
+	UE::Dataflow::FDataflowSimulationVisualizationRegistry::GetInstance().RegisterVisualization(
+		MakeUnique<UE::Groom::FGroomDataflowSimulationVisualization>());
+}
+
+void FHairStrandsDataflowModule::ShutdownModule()
+{
+	UE::Groom::DeregisterRenderingCallbacks();
+
+	FDataflowAddScalarVertexPropertyCallbackRegistry::Get().DeregisterCallbacks(
+		FGroomCollectionAddScalarVertexPropertyCallbacks::Name);
+
+	UE::Dataflow::FDataflowSimulationVisualizationRegistry::GetInstance().DeregisterVisualization(
+		UE::Groom::FGroomDataflowSimulationVisualization::Name);
+}
+
+IMPLEMENT_MODULE(FHairStrandsDataflowModule, HairStrandsDataflow)
+
+#undef LOCTEXT_NAMESPACE
diff --git a/Engine/Plugins/Runtime/HairStrands/Source/HairStrandsDataflow/Private/ResampleGuidesPointsNode.cpp b/Engine/Plugins/Runtime/HairStrands/Source/HairStrandsDataflow/Private/ResampleGuidesPointsNode.cpp
new file mode 100644
--- /dev/null
+++ b/Engine/Plugins/Runtime/HairStrands/Source/HairStrandsDataflow/Private/ResampleGuidesPointsNode.cpp
@@ -0,0 +1,118 @@
+﻿// Copyright Epic Games, Inc. All Rights Reserved.
+
+#include "ResampleGuidesPointsNode.h"
+#include "GroomCollectionFacades.h"
+#include "Dataflow/DataflowObjectInterface.h"
+
+#include UE_INLINE_GENERATED_CPP_BY_NAME(ResampleGuidesPointsNode)
+
+namespace UE::Groom::Private
+{
+	FORCEINLINE int32 ComputeNumPoints(const FGroomGuidesFacade& GroomFacade, const int32 GuidePoints) 
+	{
+		int32 CurveOffset = 0, NumPoints = 0;
+		for(int32 ObjectIndex = 0, NumObjects = GroomFacade.GetNumObjects(); ObjectIndex < NumObjects; ++ObjectIndex)
+		{
+			const uint32 SamplesCount = (GuidePoints != 0) ? GuidePoints : GroomFacade.GetObjectPointSamples()[ObjectIndex];
+			NumPoints +=  (GroomFacade.GetObjectCurveOffsets()[ObjectIndex] - CurveOffset) * SamplesCount;
+			CurveOffset = GroomFacade.GetObjectCurveOffsets()[ObjectIndex];
+		}
+		return NumPoints;
+	}
+
+	FORCEINLINE void BuildGuidePoints(const FGroomGuidesFacade& GroomFacade, const int32 CurveIndex, const int32 PointOffset, const int32 SampleOffset,
+		TArray<FVector3f>& SamplePositions, const int32 GuidePoints) 
+	{
+		const int32 ObjectIndex = GroomFacade.GetCurveObjectIndices()[CurveIndex];
+		const int32 PointsCount = (GroomFacade.GetCurvePointOffsets()[CurveIndex]-PointOffset)-1;
+		const int32 SamplesCount = (GuidePoints != 0) ? GuidePoints-1 : GroomFacade.GetObjectPointSamples()[ObjectIndex]-1;
+
+		float CurveLength = 0.0;
+		TArray<float> EdgeLengths; EdgeLengths.Init(0.0f, PointsCount);
+		for(int32 EdgeIndex = 0; EdgeIndex < PointsCount; ++EdgeIndex)
+		{
+			EdgeLengths[EdgeIndex] = (GroomFacade.GetPointRestPositions()[EdgeIndex+PointOffset+1] -
+				GroomFacade.GetPointRestPositions()[EdgeIndex+PointOffset]).Length();
+			CurveLength += EdgeLengths[EdgeIndex];
+		}
+
+		SamplePositions[SampleOffset + 1] = GroomFacade.GetPointRestPositions()[PointOffset];
+		for(int32 SampleIndex = 1; SampleIndex < SamplesCount-1; ++SampleIndex)
+		{
+			const float SampleCoord = static_cast<float>(SampleIndex) / (SamplesCount - 1.0f);
+			const float SampleLength = CurveLength * SampleCoord;
+
+			float LocalLength = 0.0;
+			for(int32 EdgeIndex = 0; EdgeIndex < PointsCount; ++EdgeIndex)
+			{
+				LocalLength += EdgeLengths[EdgeIndex];
+				if(LocalLength >= SampleLength)
+				{
+					const int32 PrevPoint = EdgeIndex+PointOffset;
+					const int32 NextPoint = PrevPoint+1;
+					
+					const float SampleAlpha = (LocalLength - SampleLength) / EdgeLengths[EdgeIndex];
+					SamplePositions[SampleOffset + SampleIndex + 1] = GroomFacade.GetPointRestPositions()[PrevPoint] * SampleAlpha +
+						GroomFacade.GetPointRestPositions()[NextPoint] * (1.0-SampleAlpha);
+					break;
+				}
+			}
+		}
+		SamplePositions[SampleOffset+SamplesCount] = GroomFacade.GetPointRestPositions()[PointOffset+PointsCount];
+		SamplePositions[SampleOffset] = 2 * SamplePositions[SampleOffset+1] - SamplePositions[SampleOffset+2];
+	}
+	
+	FORCEINLINE void ResampleGuidesPoints(FGroomGuidesFacade& GroomFacade, const int32 GuidePoints)
+	{
+		if(GroomFacade.IsValid())
+		{
+			const int32 NumSamples = ComputeNumPoints(GroomFacade, GuidePoints);
+			const int32 NumCurves = GroomFacade.GetNumCurves();
+
+			TArray<FVector3f> SamplePositions;
+			SamplePositions.Init(FVector3f::Zero(), NumSamples);
+
+			TArray<int32> CurveOffsets;
+			CurveOffsets.Init(0, NumCurves);
+
+			int32 PointOffset = 0;
+			int32 SampleOffset = 0;
+			for(int32 CurveIndex = 0; CurveIndex < NumCurves; ++CurveIndex)
+			{
+				BuildGuidePoints(GroomFacade, CurveIndex, PointOffset, SampleOffset, SamplePositions, GuidePoints);
+				const uint32 SamplesCount = (GuidePoints != 0) ? GuidePoints : GroomFacade.GetObjectPointSamples()[GroomFacade.GetCurveObjectIndices()[CurveIndex]];
+					
+				PointOffset = GroomFacade.GetCurvePointOffsets()[CurveIndex];
+				SampleOffset += SamplesCount;
+
+				CurveOffsets[CurveIndex] = SampleOffset;
+			}
+			// Resize the points groups
+			GroomFacade.ResizePointsGroups(NumSamples);
+
+			// Set the curve offsets
+			GroomFacade.SetCurvePointOffsets(CurveOffsets);
+
+			// Set the point positions and kinematic weights
+			GroomFacade.SetPointRestPositions(SamplePositions);
+		}
+	}
+}
+
+void FResampleGuidesPointsDataflowNode::Evaluate(UE::Dataflow::FContext& Context, const FDataflowOutput* Out) const
+{
+	if (Out->IsA<FManagedArrayCollection>(&Collection))
+	{
+		FManagedArrayCollection GroomCollection = GetValue<FManagedArrayCollection>(Context, &Collection);
+		
+		UE::Groom::FGroomGuidesFacade GuidesFacade(GroomCollection);
+
+		if(GuidesFacade.IsValid())
+		{
+			UE::Groom::Private::ResampleGuidesPoints(GuidesFacade, static_cast<uint8>(PointsCount));
+		}
+		
+		SetValue(Context, MoveTemp(GroomCollection), &Collection);
+	}
+}
+
diff --git a/Engine/Plugins/Runtime/HairStrands/Source/HairStrandsDataflow/Private/SmoothGuidesCurvesNode.cpp b/Engine/Plugins/Runtime/HairStrands/Source/HairStrandsDataflow/Private/SmoothGuidesCurvesNode.cpp
new file mode 100644
--- /dev/null
+++ b/Engine/Plugins/Runtime/HairStrands/Source/HairStrandsDataflow/Private/SmoothGuidesCurvesNode.cpp
@@ -0,0 +1,68 @@
+// Copyright Epic Games, Inc. All Rights Reserved.
+
+#include "SmoothGuidesCurvesNode.h"
+#include "GroomCollectionFacades.h"
+#include "Dataflow/DataflowObjectInterface.h"
+
+#include UE_INLINE_GENERATED_CPP_BY_NAME(SmoothGuidesCurvesNode)
+
+namespace UE::Groom::Private
+{
+	FORCEINLINE void SmoothGuidesPoints(FGroomGuidesFacade& GroomFacade, const float SmoothingFactor)
+	{
+		if(GroomFacade.IsValid())
+		{
+			const int32 NumPoints = GroomFacade.GetNumPoints();
+			const int32 NumCurves = GroomFacade.GetNumCurves();
+
+			TArray<FVector3f> SmoothedPositions;
+			SmoothedPositions.Init(FVector3f::Zero(), NumPoints);
+
+			int32 PointOffset = 0;
+			for(int32 CurveIndex = 0; CurveIndex < NumCurves; ++CurveIndex)
+			{
+				FVector3f DirM1 = GroomFacade.GetPointRestPositions()[PointOffset+1] - GroomFacade.GetPointRestPositions()[PointOffset];
+				FVector3f DirM2 = DirM1;
+
+				const float Gamma1 = 2.0 * (1.0-SmoothingFactor);
+				const float Gamma2 = - (1.0-SmoothingFactor)*(1.0-SmoothingFactor);
+				const float Gamma3 = SmoothingFactor*SmoothingFactor;
+				
+				SmoothedPositions[PointOffset] = GroomFacade.GetPointRestPositions()[PointOffset];
+				
+				for(int32 PointIndex = PointOffset; PointIndex < GroomFacade.GetCurvePointOffsets()[CurveIndex]-1; ++PointIndex)
+				{
+					const FVector3f DirM3 = GroomFacade.GetPointRestPositions()[PointIndex+1] - GroomFacade.GetPointRestPositions()[PointIndex];
+					const FVector3f DirMi = Gamma1 * DirM1 + Gamma2 * DirM2 + Gamma3 * DirM3;
+
+					SmoothedPositions[PointIndex+1] = SmoothedPositions[PointIndex] + DirMi;
+					
+					DirM2 = DirM1;
+					DirM1 = DirMi;
+				}
+				PointOffset = GroomFacade.GetCurvePointOffsets()[CurveIndex];
+			}
+
+			// Set the point positions and kinematic weights
+			GroomFacade.SetPointRestPositions(SmoothedPositions);
+		}
+	}
+}
+
+void FSmoothGuidesCurvesDataflowNode::Evaluate(UE::Dataflow::FContext& Context, const FDataflowOutput* Out) const
+{
+	if (Out->IsA<FManagedArrayCollection>(&Collection))
+	{
+		FManagedArrayCollection GroomCollection = GetValue<FManagedArrayCollection>(Context, &Collection);
+		
+		UE::Groom::FGroomGuidesFacade GuidesFacade(GroomCollection);
+
+		if(GuidesFacade.IsValid())
+		{
+			UE::Groom::Private::SmoothGuidesPoints(GuidesFacade, 1.0-SmoothingFactor);
+		}
+		
+		SetValue(Context, MoveTemp(GroomCollection), &Collection);
+	}
+}
+
diff --git a/Engine/Plugins/Runtime/HairStrands/Source/HairStrandsDataflow/Public/BuildGuidesSkinningNode.h b/Engine/Plugins/Runtime/HairStrands/Source/HairStrandsDataflow/Public/BuildGuidesSkinningNode.h
new file mode 100644
--- /dev/null
+++ b/Engine/Plugins/Runtime/HairStrands/Source/HairStrandsDataflow/Public/BuildGuidesSkinningNode.h
@@ -0,0 +1,50 @@
+// Copyright Epic Games, Inc. All Rights Reserved.
+
+#pragma once 
+
+#include "CoreMinimal.h"
+#include "GroomAsset.h"
+#include "ChaosLog.h"
+#include "Dataflow/DataflowCore.h"
+
+#include "BuildGuidesSkinningNode.generated.h"
+
+/** Build the guides skinning by transfering the indices weights from a skelmesh */
+USTRUCT(meta = (DataflowGroom))
+struct FBuildGuidesSkinningDataflowNode : public FDataflowNode
+{
+	GENERATED_BODY()
+	
+	DATAFLOW_NODE_DEFINE_INTERNAL(FBuildGuidesSkinningDataflowNode, "TransferSkinWeights", "Groom", "")
+	DATAFLOW_NODE_RENDER_TYPE("GuidesRender", FName("FGroomCollection"), "Collection")
+
+public:
+	
+	FBuildGuidesSkinningDataflowNode(const UE::Dataflow::FNodeParameters& InParam, FGuid InGuid = FGuid::NewGuid())
+	: FDataflowNode(InParam, InGuid)
+	{
+		RegisterInputConnection(&Collection);
+		RegisterOutputConnection(&Collection, &Collection);
+	}
+
+	//~ Begin FDataflowNode interface
+	virtual void Evaluate(UE::Dataflow::FContext& Context, const FDataflowOutput* Out) const override;
+	//~ End FDataflowNode interface
+
+	/** Managed array collection to be used to store datas */
+	UPROPERTY(meta = (DataflowInput, DataflowOutput, DisplayName = "Collection", DataflowPassthrough  = "Collection"))
+	FManagedArrayCollection Collection;
+
+	/** Managed array collection to be used to store datas */
+	UPROPERTY(EditAnywhere, Category="Skinning")
+	TObjectPtr<USkeletalMesh> SkeletalMesh = nullptr;
+
+	/** Managed array collection to be used to store datas */
+	UPROPERTY(EditAnywhere, Category="Skinning", meta = (DisplayName = "LOD Index"))
+	int32 LODIndex = 0;
+
+	/** The relative transform between the skeletal mesh and the groom asset. */
+	UPROPERTY(EditAnywhere, Category = "Skinning")
+	FTransform RelativeTransform;
+};
+
diff --git a/Engine/Plugins/Runtime/HairStrands/Source/HairStrandsDataflow/Public/GenerateGuidesCurvesNode.h b/Engine/Plugins/Runtime/HairStrands/Source/HairStrandsDataflow/Public/GenerateGuidesCurvesNode.h
new file mode 100644
--- /dev/null
+++ b/Engine/Plugins/Runtime/HairStrands/Source/HairStrandsDataflow/Public/GenerateGuidesCurvesNode.h
@@ -0,0 +1,42 @@
+// Copyright Epic Games, Inc. All Rights Reserved.
+
+#pragma once 
+
+#include "CoreMinimal.h"
+#include "GroomAsset.h"
+#include "ChaosLog.h"
+#include "Dataflow/DataflowCore.h"
+
+#include "GenerateGuidesCurvesNode.generated.h"
+
+/** Build the guides curves from the strands */
+USTRUCT(meta = (DataflowGroom))
+struct FGenerateGuidesCurvesDataflowNode : public FDataflowNode
+{
+	GENERATED_BODY()
+	
+	DATAFLOW_NODE_DEFINE_INTERNAL(FGenerateGuidesCurvesDataflowNode, "GenerateGuidesCurves", "Groom", "")
+	DATAFLOW_NODE_RENDER_TYPE("GuidesRender", FName("FGroomCollection"), "Collection")
+
+public:
+	
+	FGenerateGuidesCurvesDataflowNode(const UE::Dataflow::FNodeParameters& InParam, FGuid InGuid = FGuid::NewGuid())
+	: FDataflowNode(InParam, InGuid)
+	{
+		RegisterInputConnection(&Collection);
+		RegisterOutputConnection(&Collection, &Collection);
+	}
+
+	//~ Begin FDataflowNode interface
+	virtual void Evaluate(UE::Dataflow::FContext& Context, const FDataflowOutput* Out) const override;
+	//~ End FDataflowNode interface
+
+	/** Managed array collection to be used to store datas */
+	UPROPERTY(meta = (DataflowInput, DataflowOutput, DisplayName = "Collection", DataflowPassthrough  = "Collection"))
+	FManagedArrayCollection Collection;
+
+	/** Max number of guides */
+	UPROPERTY(EditAnywhere, Category="Geometry", meta = (DisplayName = "Guides Count"))
+	int32 GuidesCount = 0;
+};
+
diff --git a/Engine/Plugins/Runtime/HairStrands/Source/HairStrandsDataflow/Public/GetGroomAssetNode.h b/Engine/Plugins/Runtime/HairStrands/Source/HairStrandsDataflow/Public/GetGroomAssetNode.h
new file mode 100644
--- /dev/null
+++ b/Engine/Plugins/Runtime/HairStrands/Source/HairStrandsDataflow/Public/GetGroomAssetNode.h
@@ -0,0 +1,46 @@
+﻿// Copyright Epic Games, Inc. All Rights Reserved.
+
+#pragma once 
+
+#include "CoreMinimal.h"
+#include "GroomAsset.h"
+#include "ChaosLog.h"
+#include "Dataflow/DataflowCore.h"
+#include "Dataflow/DataflowEngine.h"
+#include "GroomCollectionFacades.h"
+
+#include "GetGroomAssetNode.generated.h"
+
+/** Get the groom asset guides */
+USTRUCT(meta = (DataflowGroom))
+struct FGetGroomAssetDataflowNode : public FDataflowNode
+{
+	GENERATED_BODY()
+	
+	DATAFLOW_NODE_DEFINE_INTERNAL(FGetGroomAssetDataflowNode, "GetGroomAsset", "Groom", "")
+
+public:
+	
+	FGetGroomAssetDataflowNode(const UE::Dataflow::FNodeParameters& InParam, FGuid InGuid = FGuid::NewGuid())
+		: FDataflowNode(InParam, InGuid)
+	{
+		RegisterOutputConnection(&Collection);
+	}
+
+	//~ Begin FDataflowNode interface
+	virtual void Evaluate(UE::Dataflow::FContext& Context, const FDataflowOutput* Out) const override;
+	virtual TArray<UE::Dataflow::FRenderingParameter> GetRenderParametersImpl() const override;
+	//~ End FDataflowNode interface
+
+	/** Input asset to read the guides from */
+	UPROPERTY(EditAnywhere, Category = "Groom", meta = (DisplayName = "GroomAsset"))
+	TObjectPtr<const UGroomAsset> GroomAsset = nullptr;
+
+	/** Type of curves to use to fill the groom collection (guides/strands) */
+	UPROPERTY(EditAnywhere, Category = "Groom", meta = (DisplayName = "CurvesType"))
+	EGroomCollectionType CurvesType = EGroomCollectionType::Guides;
+
+	/** Managed array collection used to store the guides */
+	UPROPERTY(meta = (DataflowOutput, DisplayName = "Collection"))
+	FManagedArrayCollection Collection;
+};
diff --git a/Engine/Plugins/Runtime/HairStrands/Source/HairStrandsDataflow/Public/GetGuidesAttributesNodes.h b/Engine/Plugins/Runtime/HairStrands/Source/HairStrandsDataflow/Public/GetGuidesAttributesNodes.h
new file mode 100644
--- /dev/null
+++ b/Engine/Plugins/Runtime/HairStrands/Source/HairStrandsDataflow/Public/GetGuidesAttributesNodes.h
@@ -0,0 +1,36 @@
+// Copyright Epic Games, Inc. All Rights Reserved.
+
+#pragma once 
+
+#include "CoreMinimal.h"
+#include "ChaosLog.h"
+#include "Dataflow/DataflowCore.h"
+#include "Dataflow/DataflowConnectionTypes.h"
+
+#include "GetGuidesAttributesNodes.generated.h"
+
+/** Get the guides attributes names */
+USTRUCT(meta = (DataflowGroom))
+struct FGetGuidesKinematicWeightsDataflowNode : public FDataflowNode
+{
+	GENERATED_BODY()
+	
+	DATAFLOW_NODE_DEFINE_INTERNAL(FGetGuidesKinematicWeightsDataflowNode, "GetGuidesKinematicWeights", "Groom", "")
+
+public:
+	
+	FGetGuidesKinematicWeightsDataflowNode(const UE::Dataflow::FNodeParameters& InParam, FGuid InGuid = FGuid::NewGuid())
+	: FDataflowNode(InParam, InGuid)
+	{
+		RegisterOutputConnection(&AttributeKey);
+	}
+
+	//~ Begin FDataflowNode interface
+	virtual void Evaluate(UE::Dataflow::FContext& Context, const FDataflowOutput* Out) const override;
+	//~ End FDataflowNode interface
+
+	/** Vertex kinematic weights attribute to use */
+	UPROPERTY(meta = (DisplayName = "AttributeKey", DataflowOutput))
+	FCollectionAttributeKey AttributeKey;
+};
+
diff --git a/Engine/Plugins/Runtime/HairStrands/Source/HairStrandsDataflow/Public/GroomAssetTerminalNode.h b/Engine/Plugins/Runtime/HairStrands/Source/HairStrandsDataflow/Public/GroomAssetTerminalNode.h
new file mode 100644
--- /dev/null
+++ b/Engine/Plugins/Runtime/HairStrands/Source/HairStrandsDataflow/Public/GroomAssetTerminalNode.h
@@ -0,0 +1,36 @@
+﻿// Copyright Epic Games, Inc. All Rights Reserved.
+
+#pragma once 
+
+#include "CoreMinimal.h"
+#include "GroomAsset.h"
+#include "ChaosLog.h"
+#include "Dataflow/DataflowCore.h"
+#include "Dataflow/DataflowEngine.h"
+
+#include "GroomAssetTerminalNode.generated.h"
+
+USTRUCT(meta = (DataflowGroom, DataflowTerminal))
+struct FGroomAssetTerminalDataflowNode : public FDataflowTerminalNode
+{
+	GENERATED_USTRUCT_BODY()
+	DATAFLOW_NODE_DEFINE_INTERNAL(FGroomAssetTerminalDataflowNode, "GroomAssetTerminal", "Groom", "")
+
+public:
+
+	FGroomAssetTerminalDataflowNode(const UE::Dataflow::FNodeParameters& InParam, FGuid InGuid = FGuid::NewGuid())
+	: FDataflowTerminalNode(InParam, InGuid)
+	{
+		RegisterInputConnection(&Collection);
+		RegisterOutputConnection(&Collection, &Collection);
+	}
+
+	//~ Begin FDataflowTerminalNode interface
+	virtual void SetAssetValue(TObjectPtr<UObject> Asset, UE::Dataflow::FContext& Context) const override;
+	virtual void Evaluate(UE::Dataflow::FContext& Context, const FDataflowOutput* Out) const override;
+	//~ End FDataflowTerminalNode interface
+
+	UPROPERTY(meta = (DataflowInput, DataflowOutput, DisplayName = "Collection", DataflowPassthrough  = "Collection"))
+	FManagedArrayCollection Collection;
+};
+
diff --git a/Engine/Plugins/Runtime/HairStrands/Source/HairStrandsDataflow/Public/GroomCollectionFacades.h b/Engine/Plugins/Runtime/HairStrands/Source/HairStrandsDataflow/Public/GroomCollectionFacades.h
new file mode 100644
--- /dev/null
+++ b/Engine/Plugins/Runtime/HairStrands/Source/HairStrandsDataflow/Public/GroomCollectionFacades.h
@@ -0,0 +1,279 @@
+// Copyright Epic Games, Inc. All Rights Reserved.
+
+#pragma once
+
+#include "CoreMinimal.h"
+#include "GroomEdit.h"
+#include "GeometryCollection/ManagedArrayAccessor.h"
+#include "GeometryCollection/ManagedArrayCollection.h"
+
+#include "GroomCollectionFacades.generated.h"
+
+class FGeometryCollection;
+namespace Chaos
+{
+	class FChaosArchive;
+}
+
+/** Enum to pick strands or guides in dataflow nodes */
+UENUM()
+enum class EGroomCollectionType : uint8
+{
+	/** Strands type (Rendering Mesh) */
+	Strands UMETA(DisplayName = "Strands"),
+
+	/** Guides type (Simulation Mesh)*/
+	Guides UMETA(DisplayName = "Guides"),
+};
+
+namespace UE::Groom
+{
+	/**
+	* FGroomCollectionFacade 
+	* @brief Base facade to store the TArray<T> groups necessary to setup the groom asset.
+	*/
+	template<typename DerivedType>
+	class HAIRSTRANDSDATAFLOW_API FGroomCollectionFacade
+	{
+	public:
+		/** Groom collection group names */
+		static const FName CurvesGroup;
+		static const FName EdgesGroup;
+		static const FName ObjectsGroup;
+		static const FName PointsGroup;
+		
+		/** Groom collection attribute names */
+		static const FName CurvePointOffsetsAttribute;
+		static const FName ObjectCurveOffsetsAttribute;
+		static const FName EdgeRestOrientationsAttribute;
+		static const FName PointRestPositionsAttribute;
+		static const FName PointCurveIndicesAttribute;
+		static const FName CurveObjectIndicesAttribute;
+
+		FGroomCollectionFacade(FManagedArrayCollection& InCollection);
+		FGroomCollectionFacade(const FManagedArrayCollection& InCollection);
+
+		/** Create the facade attributes. */
+		void DefineSchema();
+
+		/** Is the facade defined constant. */
+		bool IsConst() const { return Collection == nullptr; }
+
+		/** Is the Facade defined on the collection? */
+		bool IsValid() const;
+		
+		/** Get the number of curves */
+		int32 GetNumCurves() const { return CurvePointOffsets.Num(); }
+
+		/** Get the number of sections */
+		int32 GetNumObjects() const { return ObjectCurveOffsets.Num(); }
+
+		/** Get the number of points */
+		int32 GetNumPoints() const { return PointRestPositions.Num(); }
+
+		/** Get the number of edges */
+		int32 GetNumEdges() const { return EdgeRestOrientations.Num(); }
+		
+		/** Get the point rest positions */
+		const TArray<FVector3f>& GetPointRestPositions() const { return PointRestPositions.Get().GetConstArray(); }
+
+		/** Get the edge rest orientations */
+		const TArray<FQuat4f>& GetEdgeRestOrientations() const { return EdgeRestOrientations.Get().GetConstArray(); }
+
+		/** Get the curve point offsets */
+		const TArray<int32>& GetCurvePointOffsets() const { return CurvePointOffsets.Get().GetConstArray(); }
+
+		/** Get the object curve offsets */
+		const TArray<int32>& GetObjectCurveOffsets() const { return ObjectCurveOffsets.Get().GetConstArray(); }
+
+		/** Get the point curve indices */
+		const TArray<int32>& GetPointCurveIndices() const { return PointCurveIndices.Get().GetConstArray(); }
+
+		/** Get the curve object indices */
+		const TArray<int32>& GetCurveObjectIndices() const { return CurveObjectIndices.Get().GetConstArray(); }
+		
+		/** Set the point rest positions */
+		void SetPointRestPositions(const TArray<FVector3f>& InPointRestPositions) { PointRestPositions.Modify() = InPointRestPositions; UpdateEdgeRestOrientations(); }	
+		
+		/** Set the curve point offsets */
+		void SetCurvePointOffsets(const TArray<int32>& InCurvePointOffsets) { CurvePointOffsets.Modify() = InCurvePointOffsets; UpdatePointCurveIndices(); }
+
+		/** Set the object curve offsets */
+		void SetObjectCurveOffsets(const TArray<int32>& InObjectCurveOffsets) { ObjectCurveOffsets.Modify() = InObjectCurveOffsets; UpdateCurveObjectIndices(); }
+		
+		/** Initialize the curve vertices from point positions and objects offsets */
+		void InitGroomCollection(const TArray<FVector3f>& InPointRestPositions,
+			const TArray<int32>& InCurvePointOffsets, const TArray<int32>& InObjectCurveOffsets);
+		
+	protected :
+		
+		/** Update the edge rest orientations with parallel transport */
+		void UpdateEdgeRestOrientations();
+		
+		/** Update the point curve indices from the offsets  */
+		void UpdatePointCurveIndices();
+
+		/** Update the curve object indices from the offsets  */
+		void UpdateCurveObjectIndices();
+		
+		/** Const collection the facade is linked to */
+		const FManagedArrayCollection& ConstCollection;
+
+		/** Non-const collection the facade is linked to */
+		FManagedArrayCollection* Collection = nullptr;
+		
+		/** Groom edges local orientation */
+		TManagedArrayAccessor<FQuat4f> EdgeRestOrientations;
+
+		/** Groom points local position */
+		TManagedArrayAccessor<FVector3f> PointRestPositions;
+
+		/** Groom curves point offset */
+		TManagedArrayAccessor<int32> CurvePointOffsets;
+                                                 
+		/** Groom objects curve offset */
+		TManagedArrayAccessor<int32> ObjectCurveOffsets;
+
+		/** Groom points curve index */
+		TManagedArrayAccessor<int32> PointCurveIndices;
+
+		/** Groom curves object index */
+		TManagedArrayAccessor<int32> CurveObjectIndices;
+	};
+	
+	/**
+	* FGroomStrandsFacade 
+	* @brief Strands facade to store the TArray<T> groups necessary to setup the groom strands.
+	*/
+	class HAIRSTRANDSDATAFLOW_API FGroomStrandsFacade : public FGroomCollectionFacade<FGroomStrandsFacade>
+	{
+	public:
+		/** Groom collection group names */
+		static const FName GroupPrefix;
+		using FEditableType = FEditableHairStrand;
+
+		FGroomStrandsFacade(FManagedArrayCollection& InCollection);
+		FGroomStrandsFacade(const FManagedArrayCollection& InCollection);
+
+		/** Create the facade attributes. */
+		void DefineFacadeSchema() {}
+
+		/** Is the Facade defined on the collection? */
+		bool IsFacadeValid() const {return true;}
+
+		/** Init facade collection attributes */
+		void InitFacadeCollection() {};
+
+		/** Get the editable groom types */
+		static const TArray<FEditableType>& GetEditableGroom(const FEditableGroomGroup& GroomGroup)
+		{
+			return GroomGroup.Strands;
+		}
+	};
+
+	/**
+	* FGroomGuidesFacade 
+	* @brief Guides facade to store the TArray<T> groups necessary to setup the groom guides.
+	*/
+	class HAIRSTRANDSDATAFLOW_API FGroomGuidesFacade : public FGroomCollectionFacade<FGroomGuidesFacade>
+	{
+	public:
+		/** Groom collection group names */
+		static const FName GroupPrefix;
+		using FEditableType = FEditableHairGuide;
+
+		FGroomGuidesFacade(FManagedArrayCollection& InCollection);
+		FGroomGuidesFacade(const FManagedArrayCollection& InCollection);
+		
+		static const FName VerticesGroup;
+
+		static const FName VertexKinematicWeightsAttribute;
+		static const FName PointBoneIndicesAttribute;
+		static const FName PointBoneWeightsAttribute;
+		static const FName ObjectPointSamplesAttribute;
+		static const FName CurveStrandIndicesAttribute;
+		static const FName CurveLevelIndicesAttribute;
+		static const FName PointInterpIndicesAttribute;
+		static const FName PointInterpWeightsAttribute;
+		
+		/** Create the facade attributes. */
+		void DefineFacadeSchema();
+
+		/** Is the Facade defined on the collection? */
+		bool IsFacadeValid() const;
+
+		/** Check if the facade have skinning attributes*/
+		bool HasSkinningAttributes() const;
+
+		/** Init facade collection attributes */
+		void InitFacadeCollection();
+
+		/** Get the number of Vertices */
+		int32 GetNumVertices() const { return VertexKinematicWeights.Num(); }
+
+		/** Get the editable groom types */
+		static const TArray<FEditableType>& GetEditableGroom(const FEditableGroomGroup& GroomGroup)
+		{
+			return GroomGroup.Guides;
+		}
+		
+		/** Get the vertex kinematic weights */
+		const TArray<float>& GetVertexKinematicWeights() const { return VertexKinematicWeights.Get().GetConstArray(); }
+		
+		/** Get the point bone indices */
+		const TArray<int32>& GetPointBoneIndices(const int32 PointIndex) const { return PointBoneIndices.Get()[PointIndex]; }
+
+		/** Get the point bone weights*/
+		const TArray<float>& GetPointBoneWeights(const int32 PointIndex) const { return PointBoneWeights.Get()[PointIndex]; }
+
+		/** Get the object point samples */
+		const TArray<int32>& GetObjectPointSamples() const { return ObjectPointSamples.Get().GetConstArray(); }
+		
+		/** Get the guide strand indices */
+		const TArray<int32>& GetCurveStrandIndices() const { return CurveStrandIndices.Get().GetConstArray(); }
+
+		/** Set the vertex kinematic weights */
+		void SetVertexKinematicWeights(const TArray<float>& InKinematicWeights) { VertexKinematicWeights.Modify() = InKinematicWeights; }
+
+		/** Set the point bone indices */
+		void SetPointBoneIndices(const int32 PointIndex, const TArray<int32>& InBoneIndices) { return PointBoneIndices.ModifyAt(PointIndex,InBoneIndices); }
+
+		/** Set the point bone weights*/
+		void SetPointBoneWeights(const int32 PointIndex, const TArray<float>& InBoneWeights) { return PointBoneWeights.ModifyAt(PointIndex, InBoneWeights); }
+		
+		/** Set the object point samples */
+		void SetObjectPointSamples(const TArray<int32>& NumPointSamples) { ObjectPointSamples.Modify() = NumPointSamples; }
+		
+		/** Set the guide strand indices */
+        void SetCurveStrandIndices(const TArray<int32>& StrandIndices) { CurveStrandIndices.Modify() = StrandIndices; }
+
+		/** Resize the geometry groups (vertex, points, edges) based on a new total number of points */
+		void ResizePointsGroups(const int32 NumPoints) const;
+
+	protected :
+		
+		/** Max distance from the kinematic target */
+		TManagedArrayAccessor<float> VertexKinematicWeights;
+
+		/** Point bone indices */
+		TManagedArrayAccessor<TArray<int32>> PointBoneIndices;
+
+		/** Point bone weights */
+		TManagedArrayAccessor<TArray<float>> PointBoneWeights;
+
+		/** Object point samples */
+		TManagedArrayAccessor<int32> ObjectPointSamples;
+		
+		/** Strand index from which the guide has been generated */
+		TManagedArrayAccessor<int32> CurveStrandIndices;
+
+		// /** Level index for the curves hierarchy */
+		// TManagedArrayAccessor<int32> CurveLevelIndices;
+		//
+		// /** Point (Level L) to point (Level L+1) interpolation indices+coords */
+		// TManagedArrayAccessor<FVector4f> PointInterpIndices;
+		//
+		// /** Point (Level L) to point (Level L+1) interpolation weights */
+		// TManagedArrayAccessor<FVector4f> PointInterpWeights;
+	};
+}
diff --git a/Engine/Plugins/Runtime/HairStrands/Source/HairStrandsDataflow/Public/GroomDataflowRendering.h b/Engine/Plugins/Runtime/HairStrands/Source/HairStrandsDataflow/Public/GroomDataflowRendering.h
new file mode 100644
--- /dev/null
+++ b/Engine/Plugins/Runtime/HairStrands/Source/HairStrandsDataflow/Public/GroomDataflowRendering.h
@@ -0,0 +1,62 @@
+﻿// Copyright Epic Games, Inc. All Rights Reserved.
+
+#pragma once
+
+#include "Dataflow/DataflowRenderingFactory.h"
+#include "Dataflow/DataflowRenderingViewMode.h"
+
+namespace UE::Groom
+{
+/** Strands rendering callback for the dataflow editor */
+class FGroomStrandsRenderingCallbacks : public Dataflow::FRenderingFactory::ICallbackInterface
+{
+public :
+	static Dataflow::FRenderKey RenderKey;
+	
+private:
+
+	//¬ Begin ICallbackInterface interface
+	virtual Dataflow::FRenderKey GetRenderKey() const override
+	{
+		return RenderKey;
+	}
+
+	virtual bool CanRender(const Dataflow::IDataflowConstructionViewMode& ViewMode) const override
+	{
+		return (ViewMode.GetName() == Dataflow::FDataflowConstruction3DViewMode::Name);
+	}
+
+	virtual void Render(GeometryCollection::Facades::FRenderingFacade& RenderCollection, const Dataflow::FGraphRenderingState& State) override;
+	//¬ End ICallbackInterface interface
+};
+
+/** Guides rendering callback for the dataflow editor */
+class FGroomGuidesRenderingCallbacks : public Dataflow::FRenderingFactory::ICallbackInterface
+{
+	public :
+		static Dataflow::FRenderKey RenderKey;
+	
+private:
+
+	//¬ Begin ICallbackInterface interface
+	virtual Dataflow::FRenderKey GetRenderKey() const override
+	{
+		return RenderKey;
+	}
+
+	virtual bool CanRender(const Dataflow::IDataflowConstructionViewMode& ViewMode) const override
+	{
+		return (ViewMode.GetName() == Dataflow::FDataflowConstruction3DViewMode::Name);
+	}
+
+	virtual void Render(GeometryCollection::Facades::FRenderingFacade& RenderCollection, const Dataflow::FGraphRenderingState& State) override;
+	//¬ End ICallbackInterface interface
+};
+
+/** Register rendering callbacks */
+void RegisterRenderingCallbacks();
+
+/** Deregister rendering callbacks */
+void DeregisterRenderingCallbacks();
+	
+}  // End namespace UE::Groom
diff --git a/Engine/Plugins/Runtime/HairStrands/Source/HairStrandsDataflow/Public/HairStrandsDataflowModule.h b/Engine/Plugins/Runtime/HairStrands/Source/HairStrandsDataflow/Public/HairStrandsDataflowModule.h
new file mode 100644
--- /dev/null
+++ b/Engine/Plugins/Runtime/HairStrands/Source/HairStrandsDataflow/Public/HairStrandsDataflowModule.h
@@ -0,0 +1,17 @@
+﻿// Copyright Epic Games, Inc. All Rights Reserved.
+
+#pragma once
+
+#include "CoreMinimal.h"
+#include "Modules/ModuleInterface.h"
+
+/** Hair dataflow construction module public interface */
+class FHairStrandsDataflowModule : public IModuleInterface
+{
+public :
+	
+	//~ IModuleInterface begin interface
+	virtual void StartupModule() override;
+	virtual void ShutdownModule() override;
+	//~ IModuleInterface end interface
+};
diff --git a/Engine/Plugins/Runtime/HairStrands/Source/HairStrandsDataflow/Public/ResampleGuidesPointsNode.h b/Engine/Plugins/Runtime/HairStrands/Source/HairStrandsDataflow/Public/ResampleGuidesPointsNode.h
new file mode 100644
--- /dev/null
+++ b/Engine/Plugins/Runtime/HairStrands/Source/HairStrandsDataflow/Public/ResampleGuidesPointsNode.h
@@ -0,0 +1,66 @@
+﻿// Copyright Epic Games, Inc. All Rights Reserved.
+
+#pragma once 
+
+#include "CoreMinimal.h"
+#include "GroomAsset.h"
+#include "ChaosLog.h"
+#include "Dataflow/DataflowCore.h"
+#include "Dataflow/DataflowEngine.h"
+
+#include "ResampleGuidesPointsNode.generated.h"
+
+/** Enum to pick the number of points per guide */
+UENUM()
+enum class EGroomNumPoints : uint8
+{
+	/** Default behavior coming from the strands size in the physics settings */
+	Default  = 0 UMETA(DisplayName = "Default"),
+	
+	/** 4 points / curve */
+	Size4  = 4 UMETA(DisplayName = "4 Points"),
+
+	/** 8 points / curve */
+	Size8  = 8 UMETA(DisplayName = "8 Points"),
+
+	/** 16 points / curve */
+	Size16  = 16 UMETA(DisplayName = "16 Points"),
+
+	/** 32 points / curve */
+	Size32  = 32 UMETA(DisplayName = "32 Points"),
+
+	/** 64 points / curve */
+	Size64 = 64 UMETA(DisplayName = "64 Points"),
+};
+
+/** Resample the groom guides with a fixed number of points (in physics strands size) */
+USTRUCT(meta = (DataflowGroom))
+struct FResampleGuidesPointsDataflowNode : public FDataflowNode
+{
+	GENERATED_BODY()
+	
+	DATAFLOW_NODE_DEFINE_INTERNAL(FResampleGuidesPointsDataflowNode, "ResampleGuidesPoints", "Groom", "")
+	DATAFLOW_NODE_RENDER_TYPE("GuidesRender", FName("FGroomCollection"), "Collection")
+
+public:
+	
+	FResampleGuidesPointsDataflowNode(const UE::Dataflow::FNodeParameters& InParam, FGuid InGuid = FGuid::NewGuid())
+	: FDataflowNode(InParam, InGuid)
+	{
+		RegisterInputConnection(&Collection);
+		RegisterOutputConnection(&Collection, &Collection);
+	}
+
+	//~ Begin FDataflowNode interface
+	virtual void Evaluate(UE::Dataflow::FContext& Context, const FDataflowOutput* Out) const override;
+	//~ End FDataflowNode interface
+
+	/** Managed array collection to be used to store datas */
+	UPROPERTY(meta = (DataflowInput, DataflowOutput, DisplayName = "Collection", DataflowPassthrough  = "Collection"))
+	FManagedArrayCollection Collection;
+
+	/** Max number of points */
+	UPROPERTY(EditAnywhere, Category="Geometry", meta = (DisplayName = "Points Count"))
+	EGroomNumPoints PointsCount = EGroomNumPoints::Default;
+};
+
diff --git a/Engine/Plugins/Runtime/HairStrands/Source/HairStrandsDataflow/Public/SmoothGuidesCurvesNode.h b/Engine/Plugins/Runtime/HairStrands/Source/HairStrandsDataflow/Public/SmoothGuidesCurvesNode.h
new file mode 100644
--- /dev/null
+++ b/Engine/Plugins/Runtime/HairStrands/Source/HairStrandsDataflow/Public/SmoothGuidesCurvesNode.h
@@ -0,0 +1,44 @@
+// Copyright Epic Games, Inc. All Rights Reserved.
+
+#pragma once 
+
+#include "CoreMinimal.h"
+#include "GroomAsset.h"
+#include "ChaosLog.h"
+#include "Dataflow/DataflowCore.h"
+#include "Dataflow/DataflowEngine.h"
+
+#include "SmoothGuidesCurvesNode.generated.h"
+
+
+/** Smooth the guides for more stable simulation */
+USTRUCT(meta = (DataflowGroom))
+struct FSmoothGuidesCurvesDataflowNode : public FDataflowNode
+{
+	GENERATED_BODY()
+	
+	DATAFLOW_NODE_DEFINE_INTERNAL(FSmoothGuidesCurvesDataflowNode, "SmoothGuidesCurves", "Groom", "")
+	DATAFLOW_NODE_RENDER_TYPE("GuidesRender", FName("FGroomCollection"), "Collection")
+
+public:
+	
+	FSmoothGuidesCurvesDataflowNode(const UE::Dataflow::FNodeParameters& InParam, FGuid InGuid = FGuid::NewGuid())
+	: FDataflowNode(InParam, InGuid)
+	{
+		RegisterInputConnection(&Collection);
+		RegisterOutputConnection(&Collection, &Collection);
+	}
+
+	//~ Begin FDataflowNode interface
+	virtual void Evaluate(UE::Dataflow::FContext& Context, const FDataflowOutput* Out) const override;
+	//~ End FDataflowNode interface
+
+	/** Managed array collection to be used to store data */
+	UPROPERTY(meta = (DataflowInput, DataflowOutput, DisplayName = "Collection", DataflowPassthrough  = "Collection"))
+	FManagedArrayCollection Collection;
+
+	/** Smoothing factor between 0 and 1 */
+	UPROPERTY(EditAnywhere, Category="Geometry", meta = (DisplayName = "Smoothing Factor", ClampMin="0", ClampMax="1", UIMin="0", UIMax="1"))
+	float SmoothingFactor = 0.0f;
+};
+
