Fix CIS for section both incremented in lop and in body

diff --git a/Engine/Plugins/Runtime/HairStrands/Source/HairStrandsDataflow/Private/GroomDataflowRendering.cpp b/Engine/Plugins/Runtime/HairStrands/Source/HairStrandsDataflow/Private/GroomDataflowRendering.cpp
--- a/Engine/Plugins/Runtime/HairStrands/Source/HairStrandsDataflow/Private/GroomDataflowRendering.cpp
+++ b/Engine/Plugins/Runtime/HairStrands/Source/HairStrandsDataflow/Private/GroomDataflowRendering.cpp
@@ -158,8 +158,6 @@ namespace UE::Groom
 
 			FaceOffset = ObjectFaceOffsets[SectionIndex];
 			VertexOffset = ObjectVertexOffsets[SectionIndex];
-
-			SectionIndex++;
 		}
 	}
 	
