7-Groom solver/caching

#virtualized

diff --git a/Engine/Plugins/Experimental/ChaosCaching/Source/ChaosCaching/Public/Chaos/CacheManagerActor.h b/Engine/Plugins/Experimental/ChaosCaching/Source/ChaosCaching/Public/Chaos/CacheManagerActor.h
--- a/Engine/Plugins/Experimental/ChaosCaching/Source/ChaosCaching/Public/Chaos/CacheManagerActor.h
+++ b/Engine/Plugins/Experimental/ChaosCaching/Source/ChaosCaching/Public/Chaos/CacheManagerActor.h
@@ -130,6 +130,9 @@ struct FObservedComponent
 
 	/** Check if the Observed component is enabled for a specific cache mode */
 	CHAOSCACHING_API bool IsEnabled(ECacheMode CacheMode) const;
+	
+	/** Gets the chaos cache from the observed component */
+	CHAOSCACHING_API UChaosCache* GetChaosCache() const {return Cache;}
 
 	/** Gets the component from the internal component ref */
 	CHAOSCACHING_API UPrimitiveComponent* GetComponent(AActor* OwningActor);
diff --git a/Engine/Plugins/Experimental/Dataflow/Source/DataflowEditor/Private/Dataflow/DataflowSimulationScene.cpp b/Engine/Plugins/Experimental/Dataflow/Source/DataflowEditor/Private/Dataflow/DataflowSimulationScene.cpp
--- a/Engine/Plugins/Experimental/Dataflow/Source/DataflowEditor/Private/Dataflow/DataflowSimulationScene.cpp
+++ b/Engine/Plugins/Experimental/Dataflow/Source/DataflowEditor/Private/Dataflow/DataflowSimulationScene.cpp
@@ -225,7 +225,9 @@ void FDataflowSimulationScene::RecordSimulationCache()
 	else
 	{
 		bIsRecordingCache = true;
-        SimulationTime = TimeRange[0];
+		const float DeltaTime = (SceneDescription->CacheParams.FrameRate > 0) ?
+			(1.0f / SceneDescription->CacheParams.FrameRate) : 0.0f; 
+        SimulationTime = TimeRange[0] - DeltaTime;
 		if(AChaosCacheManager* CacheManager = Cast<AChaosCacheManager>(RootSceneActor))
 		{
 			GetWorld()->GetSubsystem<UDataflowSimulationManager>()->SetSimulationEnabled(true);
diff --git a/Engine/Plugins/Runtime/HairStrands/HairStrands.uplugin b/Engine/Plugins/Runtime/HairStrands/HairStrands.uplugin
--- a/Engine/Plugins/Runtime/HairStrands/HairStrands.uplugin
+++ b/Engine/Plugins/Runtime/HairStrands/HairStrands.uplugin
@@ -48,6 +48,11 @@
 			"Name": "HairStrandsDataflow",
 			"Type": "Editor",
 			"LoadingPhase": "Default"
+		},
+		{
+			"Name": "HairStrandsSolver",
+			"Type": "Runtime",
+			"LoadingPhase": "Default"
 		}
 	],
 	"Plugins": [
diff --git a/Engine/Plugins/Runtime/HairStrands/Source/HairStrandsCore/Private/GroomAsset.cpp b/Engine/Plugins/Runtime/HairStrands/Source/HairStrandsCore/Private/GroomAsset.cpp
--- a/Engine/Plugins/Runtime/HairStrands/Source/HairStrandsCore/Private/GroomAsset.cpp
+++ b/Engine/Plugins/Runtime/HairStrands/Source/HairStrandsCore/Private/GroomAsset.cpp
@@ -3991,3 +3991,14 @@ void UGroomAsset::ReadDataflowContent(const TObjectPtr<UDataflowBaseContent>& Da
 #endif
 	}
 }
+
+void UGroomAsset::ValidateSimulationCache()
+{
+	Modify();
+	SetEnableSimulationCache(true);
+
+#if WITH_EDITORONLY_DATA
+	CacheDerivedDatas();
+#endif
+}
+
diff --git a/Engine/Plugins/Runtime/HairStrands/Source/HairStrandsCore/Private/GroomComponent.cpp b/Engine/Plugins/Runtime/HairStrands/Source/HairStrandsCore/Private/GroomComponent.cpp
--- a/Engine/Plugins/Runtime/HairStrands/Source/HairStrandsCore/Private/GroomComponent.cpp
+++ b/Engine/Plugins/Runtime/HairStrands/Source/HairStrandsCore/Private/GroomComponent.cpp
@@ -3968,9 +3968,7 @@ void UGroomComponent::PostEditChangeProperty(FPropertyChangedEvent& PropertyChan
 		{
 			if (!GroomAsset->GetEnableSimulationCache())
 			{
-				GroomAsset->Modify();
-				GroomAsset->SetEnableSimulationCache(true);
-				GroomAsset->CacheDerivedDatas();
+				GroomAsset->ValidateSimulationCache();
 				break;
 			}
 		}
diff --git a/Engine/Plugins/Runtime/HairStrands/Source/HairStrandsCore/Public/GroomAsset.h b/Engine/Plugins/Runtime/HairStrands/Source/HairStrandsCore/Public/GroomAsset.h
--- a/Engine/Plugins/Runtime/HairStrands/Source/HairStrandsCore/Public/GroomAsset.h
+++ b/Engine/Plugins/Runtime/HairStrands/Source/HairStrandsCore/Public/GroomAsset.h
@@ -799,6 +799,9 @@ public:
 
 	uint8 GenerateClassStripFlags(FArchive& Ar);
 
+	/** Enable the simulation cache and recompute the cached derived data */
+	void ValidateSimulationCache();
+
 private:
 	void ApplyStripFlags(uint8 StripFlags, const class ITargetPlatform* CookTarget);
 
diff --git a/Engine/Plugins/Runtime/HairStrands/Source/HairStrandsEditor/Private/GroomActions.cpp b/Engine/Plugins/Runtime/HairStrands/Source/HairStrandsEditor/Private/GroomActions.cpp
--- a/Engine/Plugins/Runtime/HairStrands/Source/HairStrandsEditor/Private/GroomActions.cpp
+++ b/Engine/Plugins/Runtime/HairStrands/Source/HairStrandsEditor/Private/GroomActions.cpp
@@ -63,7 +63,7 @@ EAssetCommandResult UAssetDefinition_GroomAsset::OpenAssets(const FAssetOpenArgs
 				UAssetEditorSubsystem* const AssetEditorSubsystem = GEditor->GetEditorSubsystem<UAssetEditorSubsystem>();
 				UDataflowEditor* const AssetEditor = NewObject<UDataflowEditor>(AssetEditorSubsystem, NAME_None, RF_Transient);
 				const TSubclassOf<AActor> ActorClass = StaticLoadClass(AActor::StaticClass(), nullptr,
-					TEXT("/HairStrands/BP_GroomPreview.BP_GroomPreview_C"), nullptr, LOAD_None, nullptr);
+					TEXT("/HairStrands/BP_PreviewGroom.BP_PreviewGroom_C"), nullptr, LOAD_None, nullptr);
 				AssetEditor->Initialize({ GroomAsset}, ActorClass);
 			}
 		}
diff --git a/Engine/Plugins/Runtime/HairStrands/Source/HairStrandsSolver/HairStrandsSolver.Build.cs b/Engine/Plugins/Runtime/HairStrands/Source/HairStrandsSolver/HairStrandsSolver.Build.cs
new file mode 100644
--- /dev/null
+++ b/Engine/Plugins/Runtime/HairStrands/Source/HairStrandsSolver/HairStrandsSolver.Build.cs
@@ -0,0 +1,28 @@
+// Copyright Epic Games, Inc. All Rights Reserved.
+
+namespace UnrealBuildTool.Rules
+{
+	public class HairStrandsSolver : ModuleRules
+	{
+		public HairStrandsSolver(ReadOnlyTargetRules Target) : base(Target)
+		{
+			PublicDependencyModuleNames.AddRange(
+				new string[]
+				{
+					"Core",
+					"CoreUObject",
+					"HairStrandsCore",
+					"DataflowCore",
+					"DataflowEngine",
+					"DataflowSimulation",
+					"Engine",
+					"RHI",
+					"RenderCore",
+					"OptimusCore",
+					"ComputeFramework", 
+					"ChaosCaching"
+				}
+			);
+		}
+	}
+}
\ No newline at end of file
diff --git a/Engine/Plugins/Runtime/HairStrands/Source/HairStrandsSolver/Private/AddSolverDeformerNode.cpp b/Engine/Plugins/Runtime/HairStrands/Source/HairStrandsSolver/Private/AddSolverDeformerNode.cpp
new file mode 100644
--- /dev/null
+++ b/Engine/Plugins/Runtime/HairStrands/Source/HairStrandsSolver/Private/AddSolverDeformerNode.cpp
@@ -0,0 +1,629 @@
+// Copyright Epic Games, Inc. All Rights Reserved.
+
+#include "AddSolverDeformerNode.h"
+#include "GroomSolverComponent.h"
+#include "Dataflow/DataflowConnection.h"
+#include "OptimusDataTypeRegistry.h"
+#include "OptimusVariableDescription.h"
+#include "OptimusDeformer.h"
+#include "OptimusDeformerInstance.h"
+#include "OptimusDeformerDynamicInstanceManager.h"
+#include "Dataflow/DataflowSimulationManager.h"
+
+#include UE_INLINE_GENERATED_CPP_BY_NAME(AddSolverDeformerNode)
+
+namespace UE::Groom::Private
+{
+	template<typename DataflowType>
+	FORCEINLINE UE::Dataflow::TConnectionReference<DataflowType> GetConnectionReference(const TArray<DataflowType>& DataflowInputs, int32 InputIndex)
+	{
+		return { &DataflowInputs[InputIndex], InputIndex, &DataflowInputs };
+	}
+
+	template<typename DataflowType>
+	FORCEINLINE void RegisterArrayConnection(const TArray<DataflowType>& DataflowInputs, FDataflowNode& DataflowNode)
+	{
+		if(!DataflowInputs.IsEmpty())
+		{
+			for(int32 InputIndex = 0; InputIndex < DataflowInputs.Num(); ++InputIndex)
+			{
+				DataflowNode.FindOrRegisterInputArrayConnection(GetConnectionReference(DataflowInputs, InputIndex));
+			}
+		}
+	}
+
+	template<typename DataflowType, typename PolicyType = typename DataflowType::FPolicyType>
+	FORCEINLINE void UnregisterArrayConnection(TArray<DataflowType>& DataflowInputs, FDataflowNode& DataflowNode)
+	{
+		TArray<FDataflowInput*> NodeInputs = DataflowNode.GetInputs();
+		int32 NumRegisteredInputs = 0;
+		for (FDataflowInput* NodeInput: NodeInputs)
+		{
+			if(PolicyType::SupportsTypeStatic(NodeInput->GetType()))
+			{
+				NumRegisteredInputs++;
+			}
+		}
+		const int32 NumDataflowInputs = DataflowInputs.Num();
+		if(NumRegisteredInputs > NumDataflowInputs)
+		{
+			DataflowInputs.SetNum(NumRegisteredInputs);
+			for (int32 InputIndex = NumDataflowInputs; InputIndex < NumRegisteredInputs; ++InputIndex)
+			{
+				DataflowNode.UnregisterInputConnection(GetConnectionReference(DataflowInputs, InputIndex));
+			}
+			DataflowInputs.SetNum(NumDataflowInputs);
+		}
+	}
+	
+	template<typename DataflowType>
+	FORCEINLINE void RemoveOptionPin(TArray<DataflowType>& DataflowInputs, const FDataflowNode& DataflowNode, const UE::Dataflow::FPin& Pin)
+	{
+		if(!DataflowInputs.IsEmpty() && (Pin.Direction == UE::Dataflow::FPin::EDirection::INPUT))
+		{
+			for(int32 InputIndex = 0; InputIndex < DataflowInputs.Num(); ++InputIndex)
+			{
+				if (const FDataflowInput* const DeformerInput = DataflowNode.FindInput(GetConnectionReference(DataflowInputs, InputIndex)))
+				{
+					if(Pin.Type == DeformerInput->GetType() && Pin.Name == DeformerInput->GetName() )
+					{
+						DataflowInputs.RemoveAt(InputIndex);
+						break;
+					}		
+				}
+			}
+		}
+	}
+
+	template<typename DataflowType>
+	FORCEINLINE void GatherOptionPins(const TArray<DataflowType>& DataflowInputs, const FDataflowNode& DataflowNode, TArray<UE::Dataflow::FPin>& Pins)
+	{
+		for(int32 InputIndex = DataflowInputs.Num()-1; InputIndex >= 0; --InputIndex)
+		{
+			if (const FDataflowInput* const DeformerInput = DataflowNode.FindInput(GetConnectionReference(DataflowInputs, InputIndex)))
+			{
+				Pins.Emplace(UE::Dataflow::FPin{ UE::Dataflow::FPin::EDirection::INPUT, DeformerInput->GetType(), DeformerInput->GetName() });		
+			}
+		}
+	}
+
+	template<typename PinType, typename DataflowType>
+	FORCEINLINE void AddOptionPin(TArray<DataflowType>& DataflowInputs, FDataflowNode& DataflowNode, const FName& PinName, TArray<UE::Dataflow::FPin>& Pins)
+	{
+		const uint32 InputIndex = DataflowInputs.AddDefaulted();
+		FDataflowInput& DeformerInput = DataflowNode.RegisterInputArrayConnection(GetConnectionReference(DataflowInputs, InputIndex) );
+
+		DeformerInput.SetName(PinName);
+	
+		DataflowNode.SetInputConcreteType(GetConnectionReference(DataflowInputs, InputIndex), FName(TDataflowPolicyTypeName<PinType>::GetName()));
+		Pins.Emplace(UE::Dataflow::FPin{ UE::Dataflow::FPin::EDirection::INPUT, DeformerInput.GetType(), DeformerInput.GetName() });	
+	}
+
+	template<typename PinType, typename DataflowType, typename StorageType = typename DataflowType::FStorageType>
+	FORCEINLINE PinType GetPinValue(const TArray<DataflowType>& DataflowInputs, const FDataflowNode& DataflowNode, UE::Dataflow::FDataflowSimulationContext& SimulationContext, const int32& InputIndex)
+	{
+		PinType  ResultValue = PinType();
+		const UE::Dataflow::TConnectionReference<DataflowType> InputReference = GetConnectionReference(DataflowInputs, InputIndex);
+		if (DataflowNode.IsConnected(InputReference))
+		{
+			const StorageType PinValue = DataflowNode.GetValue<DataflowType>(SimulationContext, InputReference);
+			FDataflowConverter<StorageType>::To(PinValue, ResultValue);
+		}
+		return ResultValue;
+	}
+
+	FORCEINLINE void CreateDeformerInstance(UOptimusDeformer* DeformerGraph, UMeshDeformerInstance* DeformerInstance, const FGuid InstanceGuid, UDataflow* DataflowObject)
+	{
+		if(DeformerGraph && DeformerInstance)
+		{
+			FFunctionGraphTask::CreateAndDispatchWhenReady([DeformerGraph, DeformerInstance, InstanceGuid, DataflowObject]()
+			{
+				if (!UE::IsSavingPackage(nullptr) && !IsGarbageCollectingAndLockingUObjectHashTables())
+				{
+					if (UOptimusDeformerDynamicInstanceManager* DeformerInstanceManager = Cast<UOptimusDeformerDynamicInstanceManager>(DeformerInstance))
+					{
+						if (!DeformerInstanceManager->GetDeformerInstance(InstanceGuid))
+						{
+							DeformerInstanceManager->AddProducerDeformer(DataflowObject, InstanceGuid, DeformerGraph);
+						}
+					}
+				}
+			}, TStatId(), NULL, ENamedThreads::GameThread);
+		}
+	}
+}
+
+void FAddSolverDeformerDataflowNode::EvaluateSimulation(UE::Dataflow::FDataflowSimulationContext& SimulationContext, const FDataflowOutput* Output) const
+{
+	const TArray<FDataflowSimulationProperty> SolverProperties = GetValue(SimulationContext, &PhysicsSolvers);
+	const float SimulationDeltaTime = GetValue(SimulationContext, &SimulationTime).DeltaTime;
+	
+	if(!SolverProperties.IsEmpty() && MeshDeformer)
+	{
+		for(const FDataflowSimulationProperty& SolverProperty : SolverProperties)
+		{
+			if(SolverProperty.SimulationProxy)
+			{
+				if(FDataflowGroomSolverProxy* GroomProxy = SolverProperty.SimulationProxy->AsType<FDataflowGroomSolverProxy>())
+				{
+					FGuid& DeformerInstanceGuid = GroomProxy->DeformerInstanceGuids.FindOrAdd(GetGuid());
+					if (!DeformerInstanceGuid.IsValid())
+					{
+						// Build the deformer instance given that GUID
+						if(UDataflow* DataflowObject = Cast<UDataflow>(SimulationContext.Owner))
+						{
+							if(MeshDeformer && GroomProxy->DeformerInstance)
+							{
+								// Create a new GUID for the new deformer instance
+								DeformerInstanceGuid = FGuid::NewGuid();
+								
+								// Build a deformer instance given a guid
+								UE::Groom::Private::CreateDeformerInstance(MeshDeformer, GroomProxy->DeformerInstance, DeformerInstanceGuid, DataflowObject);
+							}
+						}
+					}
+					else if (UOptimusDeformerDynamicInstanceManager* DeformerInstanceManager = Cast<UOptimusDeformerDynamicInstanceManager>(GroomProxy->DeformerInstance))
+					{
+						if (UOptimusDeformerInstance* DeformerInstance = DeformerInstanceManager->GetDeformerInstance(DeformerInstanceGuid))
+						{
+							// Enqueue the execution of the deformer instance
+							DeformerInstanceManager->EnqueueProducerDeformer(DeformerInstanceGuid, EOptimusDeformerExecutionPhase::OverrideDefaultDeformer, 1);
+							
+							// Set the value of the deformer variables
+							for(int32 InputIndex = 0; InputIndex < DeformerNumericInputs.Num(); ++InputIndex)
+							{
+								if (const FDataflowInput* const DeformerInput = FindInput(
+									UE::Groom::Private::GetConnectionReference(DeformerNumericInputs, InputIndex)))
+								{
+									if(DeformerInput->GetType() == TDataflowPolicyTypeName<int32>::GetName())
+									{
+										DeformerInstance->SetIntVariable(DeformerInput->GetName(), UE::Groom::Private::GetPinValue<int32>(
+											DeformerNumericInputs, *this, SimulationContext, InputIndex));
+									}
+									else if(DeformerInput->GetType() == TDataflowPolicyTypeName<uint32>::GetName())
+									{
+										DeformerInstance->SetIntVariable(DeformerInput->GetName(), UE::Groom::Private::GetPinValue<uint32>(
+											DeformerNumericInputs, *this, SimulationContext, InputIndex));
+									}
+									else if(DeformerInput->GetType() == TDataflowPolicyTypeName<double>::GetName())
+									{
+										DeformerInstance->SetFloatVariable(DeformerInput->GetName(), UE::Groom::Private::GetPinValue<double>(
+											DeformerNumericInputs, *this, SimulationContext, InputIndex));
+									}
+									else if(DeformerInput->GetType() == TDataflowPolicyTypeName<float>::GetName())
+									{
+										DeformerInstance->SetFloatVariable(DeformerInput->GetName(), UE::Groom::Private::GetPinValue<float>(
+											DeformerNumericInputs, *this, SimulationContext, InputIndex));
+									}
+								}
+							}
+							for(int32 InputIndex = 0; InputIndex < DeformerVectorInputs.Num(); ++InputIndex)
+							{
+								if (const FDataflowInput* const DeformerInput = FindInput(
+									UE::Groom::Private::GetConnectionReference(DeformerVectorInputs, InputIndex)))
+								{
+									 if(DeformerInput->GetType() == TDataflowPolicyTypeName<FVector2D>::GetName())
+									{
+										DeformerInstance->SetVector2Variable(DeformerInput->GetName(), UE::Groom::Private::GetPinValue<FVector2D>(
+											DeformerVectorInputs, *this, SimulationContext, InputIndex));
+									}
+									else if(DeformerInput->GetType() == TDataflowPolicyTypeName<FVector>::GetName())
+									{
+										DeformerInstance->SetVectorVariable(DeformerInput->GetName(), UE::Groom::Private::GetPinValue<FVector>(
+											DeformerVectorInputs, *this, SimulationContext, InputIndex));
+									}
+									else if(DeformerInput->GetType() == TDataflowPolicyTypeName<FVector4>::GetName())
+									{
+										DeformerInstance->SetVector4Variable(DeformerInput->GetName(), UE::Groom::Private::GetPinValue<FVector4>(
+											DeformerVectorInputs, *this, SimulationContext, InputIndex));
+									}
+									else if(DeformerInput->GetType() == TDataflowPolicyTypeName<FQuat>::GetName())
+									{
+										DeformerInstance->SetQuatVariable(DeformerInput->GetName(), UE::Groom::Private::GetPinValue<FQuat>(
+											DeformerVectorInputs, *this, SimulationContext, InputIndex));
+									}
+									else if(DeformerInput->GetType() == TDataflowPolicyTypeName<FLinearColor>::GetName())
+									{
+										DeformerInstance->SetLinearColorVariable(DeformerInput->GetName(), UE::Groom::Private::GetPinValue<FLinearColor>(
+											DeformerVectorInputs, *this, SimulationContext, InputIndex));
+									}
+									else if(DeformerInput->GetType() == TDataflowPolicyTypeName<FIntPoint>::GetName())
+									{
+										DeformerInstance->SetInt2Variable(DeformerInput->GetName(), UE::Groom::Private::GetPinValue<FIntPoint>(
+											DeformerVectorInputs, *this, SimulationContext, InputIndex));
+									}
+									else if(DeformerInput->GetType() == TDataflowPolicyTypeName<FIntVector3>::GetName())
+									{
+										DeformerInstance->SetInt3Variable(DeformerInput->GetName(), UE::Groom::Private::GetPinValue<FIntVector3>(
+											DeformerVectorInputs, *this, SimulationContext, InputIndex));
+									}
+									else if(DeformerInput->GetType() == TDataflowPolicyTypeName<FIntVector4>::GetName())
+									{
+										DeformerInstance->SetInt4Variable(DeformerInput->GetName(), UE::Groom::Private::GetPinValue<FIntVector4>(
+											DeformerVectorInputs, *this, SimulationContext, InputIndex));
+									}
+									else if(DeformerInput->GetType() == TDataflowPolicyTypeName<FRotator>::GetName())
+									{
+										DeformerInstance->SetRotatorVariable(DeformerInput->GetName(), UE::Groom::Private::GetPinValue<FRotator>(
+											DeformerVectorInputs, *this, SimulationContext, InputIndex));
+									}
+								}
+							}
+							for(int32 InputIndex = 0; InputIndex < DeformerStringInputs.Num(); ++InputIndex)
+							{
+								if (const FDataflowInput* const DeformerInput = FindInput(
+									UE::Groom::Private::GetConnectionReference(DeformerStringInputs, InputIndex)))
+								{
+									if(DeformerInput->GetType() == TDataflowPolicyTypeName<FName>::GetName())
+									{
+										DeformerInstance->SetNameVariable(DeformerInput->GetName(), UE::Groom::Private::GetPinValue<FName>(
+											DeformerStringInputs, *this, SimulationContext, InputIndex));
+									}
+								}
+							}
+							for(int32 InputIndex = 0; InputIndex < DeformerBoolInputs.Num(); ++InputIndex)
+							{
+								if (const FDataflowInput* const DeformerInput = FindInput(
+									UE::Groom::Private::GetConnectionReference(DeformerBoolInputs, InputIndex)))
+								{
+									if(DeformerInput->GetType() == TDataflowPolicyTypeName<bool>::GetName())
+									{
+										DeformerInstance->SetBoolVariable(DeformerInput->GetName(), UE::Groom::Private::GetPinValue<bool>(
+											DeformerBoolInputs, *this, SimulationContext, InputIndex));
+									}
+								}
+							}
+							for(int32 InputIndex = 0; InputIndex < DeformerTransformInputs.Num(); ++InputIndex)
+							{
+								if (const FDataflowInput* const DeformerInput = FindInput(
+									UE::Groom::Private::GetConnectionReference(DeformerTransformInputs, InputIndex)))
+								{
+									if(DeformerInput->GetType() == TDataflowPolicyTypeName<FTransform>::GetName())
+									{
+										DeformerInstance->SetTransformVariable(DeformerInput->GetName(), UE::Groom::Private::GetPinValue<FTransform>(
+											DeformerTransformInputs, *this, SimulationContext, InputIndex));
+									}
+								}
+							}
+							for(int32 InputIndex = 0; InputIndex < DeformerNumericArrays.Num(); ++InputIndex)
+							{
+								if (const FDataflowInput* const DeformerInput = FindInput(
+									UE::Groom::Private::GetConnectionReference(DeformerNumericArrays, InputIndex)))
+								{
+									if(DeformerInput->GetType() == TDataflowPolicyTypeName<TArray<int32>>::GetName())
+									{
+										DeformerInstance->SetIntArrayVariable(DeformerInput->GetName(), UE::Groom::Private::GetPinValue<TArray<int32>>(
+											DeformerNumericArrays, *this, SimulationContext, InputIndex));
+									}
+									else if(DeformerInput->GetType() == TDataflowPolicyTypeName<TArray<uint32>>::GetName())
+									{
+										DeformerInstance->SetIntArrayVariable(DeformerInput->GetName(), UE::Groom::Private::GetPinValue<TArray<int32>>(
+											DeformerNumericArrays, *this, SimulationContext, InputIndex));
+									}
+									else if(DeformerInput->GetType() == TDataflowPolicyTypeName<TArray<double>>::GetName())
+									{
+										DeformerInstance->SetFloatArrayVariable(DeformerInput->GetName(), UE::Groom::Private::GetPinValue<TArray<double>>(
+											DeformerNumericArrays, *this, SimulationContext, InputIndex));
+									}
+									else if(DeformerInput->GetType() == TDataflowPolicyTypeName<TArray<float>>::GetName())
+									{
+										DeformerInstance->SetFloatArrayVariable(DeformerInput->GetName(), UE::Groom::Private::GetPinValue<TArray<double>>(
+											DeformerNumericArrays, *this, SimulationContext, InputIndex));
+									}
+								}
+							}
+							for(int32 InputIndex = 0; InputIndex < DeformerVectorArrays.Num(); ++InputIndex)
+							{
+								if (const FDataflowInput* const DeformerInput = FindInput(
+									UE::Groom::Private::GetConnectionReference(DeformerVectorArrays, InputIndex)))
+								{
+									 if(DeformerInput->GetType() == TDataflowPolicyTypeName<TArray<FVector2D>>::GetName())
+									{
+										DeformerInstance->SetVector2ArrayVariable(DeformerInput->GetName(), UE::Groom::Private::GetPinValue<TArray<FVector2D>>(
+											DeformerVectorArrays, *this, SimulationContext, InputIndex));
+									}
+									else if(DeformerInput->GetType() == TDataflowPolicyTypeName<TArray<FVector>>::GetName())
+									{
+										DeformerInstance->SetVectorArrayVariable(DeformerInput->GetName(), UE::Groom::Private::GetPinValue<TArray<FVector>>(
+											DeformerVectorArrays, *this, SimulationContext, InputIndex));
+									}
+									else if(DeformerInput->GetType() == TDataflowPolicyTypeName<TArray<FVector4>>::GetName())
+									{
+										DeformerInstance->SetVector4ArrayVariable(DeformerInput->GetName(), UE::Groom::Private::GetPinValue<TArray<FVector4>>(
+											DeformerVectorArrays, *this, SimulationContext, InputIndex));
+									}
+									else if(DeformerInput->GetType() == TDataflowPolicyTypeName<TArray<FQuat>>::GetName())
+									{
+										DeformerInstance->SetQuatArrayVariable(DeformerInput->GetName(), UE::Groom::Private::GetPinValue<TArray<FQuat>>(
+											DeformerVectorArrays, *this, SimulationContext, InputIndex));
+									}
+									else if(DeformerInput->GetType() == TDataflowPolicyTypeName<TArray<FLinearColor>>::GetName())
+									{
+										DeformerInstance->SetLinearColorArrayVariable(DeformerInput->GetName(), UE::Groom::Private::GetPinValue<TArray<FLinearColor>>(
+											DeformerVectorArrays, *this, SimulationContext, InputIndex));
+									}
+									else if(DeformerInput->GetType() == TDataflowPolicyTypeName<TArray<FIntPoint>>::GetName())
+									{
+										DeformerInstance->SetInt2ArrayVariable(DeformerInput->GetName(), UE::Groom::Private::GetPinValue<TArray<FIntPoint>>(
+											DeformerVectorArrays, *this, SimulationContext, InputIndex));
+									}
+									else if(DeformerInput->GetType() == TDataflowPolicyTypeName<TArray<FIntVector3>>::GetName())
+									{
+										DeformerInstance->SetInt3ArrayVariable(DeformerInput->GetName(), UE::Groom::Private::GetPinValue<TArray<FIntVector3>>(
+											DeformerVectorArrays, *this, SimulationContext, InputIndex));
+									}
+									else if(DeformerInput->GetType() == TDataflowPolicyTypeName<TArray<FIntVector4>>::GetName())
+									{
+										DeformerInstance->SetInt4ArrayVariable(DeformerInput->GetName(), UE::Groom::Private::GetPinValue<TArray<FIntVector4>>(
+											DeformerVectorArrays, *this, SimulationContext, InputIndex));
+									}
+									else if(DeformerInput->GetType() == TDataflowPolicyTypeName<TArray<FRotator>>::GetName())
+									{
+										DeformerInstance->SetRotatorArrayVariable(DeformerInput->GetName(), UE::Groom::Private::GetPinValue<TArray<FRotator>>(
+											DeformerVectorArrays, *this, SimulationContext, InputIndex));
+									}
+								}
+							}
+							for(int32 InputIndex = 0; InputIndex < DeformerBoolArrays.Num(); ++InputIndex)
+							{
+								if (const FDataflowInput* const DeformerInput = FindInput(
+									UE::Groom::Private::GetConnectionReference(DeformerBoolArrays, InputIndex)))
+								{
+									if(DeformerInput->GetType() == TDataflowPolicyTypeName<TArray<FTransform>>::GetName())
+									{
+										DeformerInstance->SetBoolArrayVariable(DeformerInput->GetName(), UE::Groom::Private::GetPinValue<TArray<bool>>(
+											DeformerBoolArrays, *this, SimulationContext, InputIndex));
+									}
+								}
+							}
+							for(int32 InputIndex = 0; InputIndex < DeformerTransformArrays.Num(); ++InputIndex)
+							{
+								if (const FDataflowInput* const DeformerInput = FindInput(
+									UE::Groom::Private::GetConnectionReference(DeformerTransformArrays, InputIndex)))
+								{
+									if(DeformerInput->GetType() == TDataflowPolicyTypeName<TArray<FTransform>>::GetName())
+									{
+										DeformerInstance->SetTransformArrayVariable(DeformerInput->GetName(), UE::Groom::Private::GetPinValue<TArray<FTransform>>(
+											DeformerTransformArrays, *this, SimulationContext, InputIndex));
+									}
+								}
+							}
+						}
+					}
+				}
+			}
+		}
+	}
+	SetValue(SimulationContext, SolverProperties, &PhysicsSolvers);
+}
+
+void FAddSolverDeformerDataflowNode::OnInvalidate()
+{}
+
+TArray<UE::Dataflow::FPin> FAddSolverDeformerDataflowNode::AddPins()
+{
+	if(MeshDeformer)
+	{
+		TArray<UE::Dataflow::FPin> Pins;
+		for (UOptimusVariableDescription* Variable : MeshDeformer->GetVariables())
+		{
+			if (Variable->DataType == FOptimusDataTypeRegistry::Get().FindType(*FIntProperty::StaticClass()))
+			{
+				UE::Groom::Private::AddOptionPin<int32>(DeformerNumericInputs,  *this, Variable->VariableName, Pins);
+			}
+			else if (Variable->DataType == FOptimusDataTypeRegistry::Get().FindType(*FUInt32Property::StaticClass()))
+			{
+				UE::Groom::Private::AddOptionPin<uint32>(DeformerNumericInputs,  *this, Variable->VariableName, Pins);
+			}
+			else if (Variable->DataType == FOptimusDataTypeRegistry::Get().FindType(*FDoubleProperty::StaticClass()))
+			{
+				UE::Groom::Private::AddOptionPin<double>(DeformerNumericInputs,  *this, Variable->VariableName, Pins);
+			}
+			else if (Variable->DataType == FOptimusDataTypeRegistry::Get().FindType(*FFloatProperty::StaticClass()))
+			{
+				UE::Groom::Private::AddOptionPin<float>(DeformerNumericInputs,  *this, Variable->VariableName, Pins);
+			}
+			else if (Variable->DataType == FOptimusDataTypeRegistry::Get().FindType(TBaseStructure<FVector2D>::Get()))
+			{
+				UE::Groom::Private::AddOptionPin<FVector2D>(DeformerVectorInputs,  *this, Variable->VariableName, Pins);
+			}
+			else if (Variable->DataType == FOptimusDataTypeRegistry::Get().FindType(TBaseStructure<FVector>::Get()))
+			{
+				UE::Groom::Private::AddOptionPin<FVector>(DeformerVectorInputs,  *this, Variable->VariableName, Pins);
+			}
+			else if (Variable->DataType == FOptimusDataTypeRegistry::Get().FindType(TBaseStructure<FVector4>::Get()))
+			{
+				UE::Groom::Private::AddOptionPin<FVector4>(DeformerVectorInputs,  *this, Variable->VariableName, Pins);
+			}
+			else if (Variable->DataType == FOptimusDataTypeRegistry::Get().FindType(TBaseStructure<FQuat>::Get()))
+			{
+				UE::Groom::Private::AddOptionPin<FQuat>(DeformerVectorInputs, *this, Variable->VariableName, Pins);
+			}
+			else if (Variable->DataType == FOptimusDataTypeRegistry::Get().FindType(TBaseStructure<FLinearColor>::Get()))
+			{
+				UE::Groom::Private::AddOptionPin<FLinearColor>(DeformerVectorInputs, *this, Variable->VariableName, Pins);
+			}
+			else if (Variable->DataType == FOptimusDataTypeRegistry::Get().FindType(TBaseStructure<FIntPoint>::Get()))
+			{
+				UE::Groom::Private::AddOptionPin<FIntPoint>(DeformerVectorInputs, *this, Variable->VariableName, Pins);
+			}
+			else if (Variable->DataType == FOptimusDataTypeRegistry::Get().FindType(TBaseStructure<FIntVector3>::Get()))
+			{
+				UE::Groom::Private::AddOptionPin<FIntVector3>(DeformerVectorInputs, *this, Variable->VariableName, Pins);
+			}
+			else if (Variable->DataType == FOptimusDataTypeRegistry::Get().FindType(TBaseStructure<FIntVector4>::Get()))
+			{
+				UE::Groom::Private::AddOptionPin<FIntVector4>(DeformerVectorInputs, *this, Variable->VariableName, Pins);
+			}
+			else if (Variable->DataType == FOptimusDataTypeRegistry::Get().FindType(TBaseStructure<FRotator>::Get()))
+			{
+				UE::Groom::Private::AddOptionPin<FRotator>(DeformerVectorInputs, *this, Variable->VariableName, Pins);
+			}
+			else if (Variable->DataType == FOptimusDataTypeRegistry::Get().FindType(*FNameProperty::StaticClass()))
+			{
+				UE::Groom::Private::AddOptionPin<FName>(DeformerStringInputs, *this, Variable->VariableName, Pins);
+			}
+			else if (Variable->DataType == FOptimusDataTypeRegistry::Get().FindType(*FBoolProperty::StaticClass()))
+			{
+				UE::Groom::Private::AddOptionPin<bool>(DeformerBoolInputs, *this, Variable->VariableName, Pins);
+			}
+			else if (Variable->DataType == FOptimusDataTypeRegistry::Get().FindType(TBaseStructure<FTransform>::Get()))
+			{
+				UE::Groom::Private::AddOptionPin<FTransform>(DeformerTransformInputs, *this, Variable->VariableName, Pins);
+			}
+			else if (Variable->DataType == FOptimusDataTypeRegistry::Get().FindArrayType(*FIntProperty::StaticClass()))
+			{
+				UE::Groom::Private::AddOptionPin<TArray<int32>>(DeformerNumericArrays, *this, Variable->VariableName, Pins);
+			}
+			else if (Variable->DataType == FOptimusDataTypeRegistry::Get().FindArrayType(*FUInt32Property::StaticClass()))
+			{
+				UE::Groom::Private::AddOptionPin<TArray<uint32>>(DeformerNumericArrays,  *this, Variable->VariableName, Pins);
+			}
+			else if (Variable->DataType == FOptimusDataTypeRegistry::Get().FindArrayType(*FDoubleProperty::StaticClass()))
+			{
+				UE::Groom::Private::AddOptionPin<TArray<double>>(DeformerNumericArrays,  *this, Variable->VariableName, Pins);
+			}
+			else if (Variable->DataType == FOptimusDataTypeRegistry::Get().FindArrayType(*FFloatProperty::StaticClass()))
+			{
+				UE::Groom::Private::AddOptionPin<TArray<float>>(DeformerNumericArrays,  *this, Variable->VariableName, Pins);
+			}
+			else if (Variable->DataType == FOptimusDataTypeRegistry::Get().FindArrayType(TBaseStructure<FVector2D>::Get()))
+			{
+				UE::Groom::Private::AddOptionPin<TArray<FVector2D>>(DeformerVectorArrays,  *this, Variable->VariableName, Pins);
+			}
+			else if (Variable->DataType == FOptimusDataTypeRegistry::Get().FindArrayType(TBaseStructure<FVector>::Get()))
+			{
+				UE::Groom::Private::AddOptionPin<TArray<FVector>>(DeformerVectorArrays,  *this, Variable->VariableName, Pins);
+			}
+			else if (Variable->DataType == FOptimusDataTypeRegistry::Get().FindArrayType(TBaseStructure<FVector4>::Get()))
+			{
+				UE::Groom::Private::AddOptionPin<TArray<FVector4>>(DeformerVectorArrays,  *this, Variable->VariableName, Pins);
+			}
+			else if (Variable->DataType == FOptimusDataTypeRegistry::Get().FindArrayType(TBaseStructure<FQuat>::Get()))
+			{
+				UE::Groom::Private::AddOptionPin<TArray<FQuat>>(DeformerVectorArrays, *this, Variable->VariableName, Pins);
+			}
+			else if (Variable->DataType == FOptimusDataTypeRegistry::Get().FindArrayType(TBaseStructure<FLinearColor>::Get()))
+			{
+				UE::Groom::Private::AddOptionPin<TArray<FLinearColor>>(DeformerVectorArrays, *this, Variable->VariableName, Pins);
+			}
+			else if (Variable->DataType == FOptimusDataTypeRegistry::Get().FindArrayType(TBaseStructure<FIntPoint>::Get()))
+			{
+				UE::Groom::Private::AddOptionPin<TArray<FIntPoint>>(DeformerVectorArrays, *this, Variable->VariableName, Pins);
+			}
+			else if (Variable->DataType == FOptimusDataTypeRegistry::Get().FindArrayType(TBaseStructure<FIntVector3>::Get()))
+			{
+				UE::Groom::Private::AddOptionPin<TArray<FIntVector3>>(DeformerVectorArrays, *this, Variable->VariableName, Pins);
+			}
+			else if (Variable->DataType == FOptimusDataTypeRegistry::Get().FindArrayType(TBaseStructure<FIntVector4>::Get()))
+			{
+				UE::Groom::Private::AddOptionPin<TArray<FIntVector4>>(DeformerVectorArrays, *this, Variable->VariableName, Pins);
+			}
+			else if (Variable->DataType == FOptimusDataTypeRegistry::Get().FindArrayType(TBaseStructure<FRotator>::Get()))
+			{
+				UE::Groom::Private::AddOptionPin<TArray<FRotator>>(DeformerVectorArrays, *this, Variable->VariableName, Pins);
+			}
+			else if (Variable->DataType == FOptimusDataTypeRegistry::Get().FindArrayType(*FNameProperty::StaticClass()))
+			{
+				UE::Groom::Private::AddOptionPin<TArray<FName>>(DeformerStringArrays, *this, Variable->VariableName, Pins);
+			}
+			else if (Variable->DataType == FOptimusDataTypeRegistry::Get().FindArrayType(*FBoolProperty::StaticClass()))
+			{
+				UE::Groom::Private::AddOptionPin<TArray<bool>>(DeformerBoolArrays, *this, Variable->VariableName, Pins);
+			}
+			else if (Variable->DataType == FOptimusDataTypeRegistry::Get().FindArrayType(TBaseStructure<FTransform>::Get()))
+			{
+				UE::Groom::Private::AddOptionPin<TArray<FTransform>>(DeformerTransformArrays, *this, Variable->VariableName, Pins);
+			}
+			else
+			{
+				UE_LOG(LogTemp, Error, TEXT("Unsupported Dataflow variable type"));
+			}			
+		}
+		return Pins;
+	}
+
+	return Super::AddPins();
+}
+
+TArray<UE::Dataflow::FPin> FAddSolverDeformerDataflowNode::GetPinsToRemove() const
+{
+	if(!DeformerNumericInputs.IsEmpty() || !DeformerStringInputs.IsEmpty() || !DeformerVectorInputs.IsEmpty() || !DeformerBoolInputs.IsEmpty() || !DeformerTransformInputs.IsEmpty())
+	{
+		TArray<UE::Dataflow::FPin> Pins;
+		UE::Groom::Private::GatherOptionPins(DeformerNumericInputs, *this, Pins);
+		UE::Groom::Private::GatherOptionPins(DeformerVectorInputs,  *this, Pins);
+		UE::Groom::Private::GatherOptionPins(DeformerStringInputs, *this, Pins);
+		UE::Groom::Private::GatherOptionPins(DeformerBoolInputs, *this, Pins);
+		UE::Groom::Private::GatherOptionPins(DeformerTransformInputs, *this, Pins);
+
+		UE::Groom::Private::GatherOptionPins(DeformerNumericArrays, *this, Pins);
+		UE::Groom::Private::GatherOptionPins(DeformerVectorArrays, *this, Pins);
+		UE::Groom::Private::GatherOptionPins(DeformerStringArrays, *this, Pins);
+		UE::Groom::Private::GatherOptionPins(DeformerBoolArrays, *this, Pins);
+		UE::Groom::Private::GatherOptionPins(DeformerTransformArrays, *this, Pins);
+		return Pins;
+	}
+	return Super::GetPinsToRemove();
+}
+
+void FAddSolverDeformerDataflowNode::OnPinRemoved(const UE::Dataflow::FPin& Pin)
+{
+	if(!DeformerNumericInputs.IsEmpty() || !DeformerStringInputs.IsEmpty() || !DeformerVectorInputs.IsEmpty() || !DeformerBoolInputs.IsEmpty() || !DeformerTransformInputs.IsEmpty())
+	{
+		UE::Groom::Private::RemoveOptionPin(DeformerNumericInputs,  *this, Pin);
+		UE::Groom::Private::RemoveOptionPin(DeformerVectorInputs,  *this, Pin);
+		UE::Groom::Private::RemoveOptionPin(DeformerStringInputs,  *this, Pin);
+		UE::Groom::Private::RemoveOptionPin(DeformerBoolInputs,  *this, Pin);
+		UE::Groom::Private::RemoveOptionPin(DeformerTransformInputs,  *this, Pin);
+
+		UE::Groom::Private::RemoveOptionPin(DeformerNumericArrays,  *this, Pin);
+		UE::Groom::Private::RemoveOptionPin(DeformerVectorArrays,  *this, Pin);
+		UE::Groom::Private::RemoveOptionPin(DeformerStringArrays,  *this, Pin);
+		UE::Groom::Private::RemoveOptionPin(DeformerBoolArrays,  *this, Pin);
+		UE::Groom::Private::RemoveOptionPin(DeformerTransformArrays,  *this, Pin);
+	}
+	return Super::OnPinRemoved(Pin);
+}
+
+void FAddSolverDeformerDataflowNode::PostSerialize(const FArchive& Ar)
+{
+	// because we add pins we need to make sure we restore them when loading
+	// to make sure they can get properly reconnected
+
+	if (Ar.IsLoading())
+	{
+		UE::Groom::Private::RegisterArrayConnection(DeformerNumericInputs, *this);
+		UE::Groom::Private::RegisterArrayConnection(DeformerVectorInputs, *this);
+		UE::Groom::Private::RegisterArrayConnection(DeformerStringInputs, *this);
+		UE::Groom::Private::RegisterArrayConnection(DeformerBoolInputs, *this);
+		UE::Groom::Private::RegisterArrayConnection(DeformerTransformInputs, *this);
+		
+		UE::Groom::Private::RegisterArrayConnection(DeformerNumericArrays, *this);
+		UE::Groom::Private::RegisterArrayConnection(DeformerVectorArrays, *this);
+		UE::Groom::Private::RegisterArrayConnection(DeformerStringArrays, *this);
+		UE::Groom::Private::RegisterArrayConnection(DeformerBoolArrays, *this);
+		UE::Groom::Private::RegisterArrayConnection(DeformerTransformArrays, *this);
+		
+		if (Ar.IsTransacting())
+		{
+			UE::Groom::Private::UnregisterArrayConnection(DeformerNumericInputs, *this);
+			UE::Groom::Private::UnregisterArrayConnection(DeformerVectorInputs, *this);
+			UE::Groom::Private::UnregisterArrayConnection(DeformerStringInputs, *this);
+			UE::Groom::Private::UnregisterArrayConnection(DeformerBoolInputs, *this);
+			UE::Groom::Private::UnregisterArrayConnection(DeformerTransformInputs, *this);
+			
+			UE::Groom::Private::UnregisterArrayConnection(DeformerNumericArrays, *this);
+			UE::Groom::Private::UnregisterArrayConnection(DeformerVectorArrays, *this);
+			UE::Groom::Private::UnregisterArrayConnection(DeformerStringArrays, *this);
+			UE::Groom::Private::UnregisterArrayConnection(DeformerBoolArrays, *this);
+			UE::Groom::Private::UnregisterArrayConnection(DeformerTransformArrays, *this);
+		}
+		else
+		{
+			ensureAlways((DeformerNumericInputs.Num() + DeformerVectorInputs.Num() + DeformerStringInputs.Num() +
+			DeformerBoolInputs.Num() + DeformerTransformInputs.Num() + DeformerNumericArrays.Num() + DeformerVectorArrays.Num() + DeformerStringArrays.Num() +
+			DeformerBoolArrays.Num() + DeformerTransformArrays.Num() + 2)== GetNumInputs());
+		}
+	}
+}
\ No newline at end of file
diff --git a/Engine/Plugins/Runtime/HairStrands/Source/HairStrandsSolver/Private/GroomCacheAdapter.cpp b/Engine/Plugins/Runtime/HairStrands/Source/HairStrandsSolver/Private/GroomCacheAdapter.cpp
new file mode 100644
--- /dev/null
+++ b/Engine/Plugins/Runtime/HairStrands/Source/HairStrandsSolver/Private/GroomCacheAdapter.cpp
@@ -0,0 +1,294 @@
+﻿// Copyright Epic Games, Inc. All Rights Reserved.
+
+#include "GroomCacheAdapter.h"
+#include "GroomBuilder.h"
+#include "GroomCache.h"
+#include "Chaos/ChaosCache.h"
+#include "Chaos/CacheManagerActor.h"
+#include "GroomComponent.h"
+#include "GroomSolverComponent.h"
+#include "UObject/Package.h"
+#include "UObject/UObjectGlobals.h"
+#include "HairStrandsCore.h"
+
+DEFINE_LOG_CATEGORY(LogGroomCache)
+
+namespace UE::Groom
+{
+FGroomCacheAdapter::~FGroomCacheAdapter()
+{}
+
+Chaos::FComponentCacheAdapter::SupportType FGroomCacheAdapter::SupportsComponentClass(UClass* InComponentClass) const
+{
+	const UClass* Desired = GetDesiredClass();
+	if(InComponentClass == Desired)
+	{
+		return Chaos::FComponentCacheAdapter::SupportType::Direct;
+	}
+	else if(InComponentClass->IsChildOf(Desired))
+	{
+		return Chaos::FComponentCacheAdapter::SupportType::Derived;
+	}
+
+	return Chaos::FComponentCacheAdapter::SupportType::None;
+}
+
+UClass* FGroomCacheAdapter::GetDesiredClass() const
+{
+	return UGroomComponent::StaticClass();
+}
+
+uint8 FGroomCacheAdapter::GetPriority() const
+{
+	return EngineAdapterPriorityBegin;
+}
+
+FGuid FGroomCacheAdapter::GetGuid() const
+{
+	FGuid NewGuid;
+	checkSlow(FGuid::Parse(TEXT("FC61D2A13092410CBCF2F767C8490986"), NewGuid));
+	return NewGuid;
+}
+
+Chaos::FPhysicsSolver* FGroomCacheAdapter::GetComponentSolver(UPrimitiveComponent* InComponent) const
+{
+	return nullptr;
+}
+
+Chaos::FPhysicsSolverEvents* FGroomCacheAdapter::BuildEventsSolver(UPrimitiveComponent* InComponent) const
+{
+	if(UGroomComponent* GroomComponent = CastChecked<UGroomComponent>(InComponent))
+	{
+		if(GroomComponent->GetGroomSolver())
+		{
+			if(UGroomSolverComponent* GroomSolver = Cast<UGroomSolverComponent>(GroomComponent->GetGroomSolver()))
+			{
+				return static_cast<FDataflowGroomSolverProxy*>(GroomSolver->GetSimulationProxy());
+			}
+		}
+	}
+	return nullptr;
+}
+
+bool FGroomCacheAdapter::ValidForPlayback(UPrimitiveComponent* InComponent, UChaosCache* InCache) const
+{
+	const UGroomComponent* GroomComponent = CastChecked<UGroomComponent>(InComponent);
+	return GroomComponent != nullptr;
+}
+
+void FGroomCacheAdapter::InitializeForLoad(UPrimitiveComponent* InComponent, FObservedComponent& InObserved)
+{
+	bIsLoading = true;
+	if(UGroomComponent* GroomComponent = CastChecked<UGroomComponent>(InComponent))
+	{
+		if(UChaosCache* ChaosCache = InObserved.GetChaosCache())
+		{
+			if(ChaosCache->GetCacheData() && !GroomComponent->GetGroomCache())
+			{
+				GroomComponent->SetGroomCache(Cast<UGroomCache>(ChaosCache->GetCacheData()));
+				GroomComponent->SetManualTick(true);
+
+				if(GroomComponent->GroomAsset && !GroomComponent->GroomAsset->GetEnableSimulationCache())
+				{
+					GroomComponent->GroomAsset->ValidateSimulationCache();
+				}
+			}
+		}
+	}
+}
+
+void FGroomCacheAdapter::SetRestState(UPrimitiveComponent* InComponent, UChaosCache* InCache, const FTransform& InRootTransform, Chaos::FReal InTime) const
+{
+	if(bIsLoading)
+	{
+		if(UGroomComponent* GroomComponent = CastChecked<UGroomComponent>(InComponent))
+		{
+			if(GroomComponent->GetGroomCache())
+			{
+				if(InTime < GroomComponent->GetGroomCache()->GetDuration())
+				{
+					GroomComponent->TickAtThisTime(InTime, true, false, true);
+				}
+			}
+		}
+	}
+}
+
+bool FGroomCacheAdapter::InitializeForRecord(UPrimitiveComponent* InComponent, FObservedComponent& InObserved)
+{
+	bIsLoading = false;
+	if(UGroomComponent* GroomComponent = CastChecked<UGroomComponent>(InComponent))
+	{
+		FGroomCacheData& GroomCache = GroomCaches.FindOrAdd(InComponent);
+		
+		GroomCache.AnimInfo.Attributes = EGroomCacheAttributes::Position;
+		GroomCache.AnimInfo.StartFrame = 0;
+		GroomCache.AnimInfo.StartTime = 0.0f;
+		GroomCache.AnimInfo.EndFrame = 0;
+		GroomCache.AnimInfo.EndTime = 0.0f;
+
+		GroomCache.PositionsBuffer.Reset();
+		GroomCache.CacheTimes.Reset();
+		GroomCache.CacheName = InObserved.CacheName.ToString();
+		GroomCache.CacheProcessor = FGroomCacheProcessor(EGroomCacheType::Guides, GroomCache.AnimInfo.Attributes);
+
+		GroomComponent->SetGroomCache(nullptr);
+		GroomComponent->SetManualTick(false);
+
+		if(UChaosCache* ChaosCache = InObserved.GetChaosCache())
+		{
+#if WITH_EDITOR
+			if(!ChaosCache->GetCacheData())
+			{
+				FString PackageNameCache, NameCache;
+                FHairStrandsCore::AssetHelper().CreateFilename(GroomComponent->GroomAsset->GetOutermost()->GetName(), TEXT("_GroomCache"), PackageNameCache, NameCache);
+                UPackage* PackageCache = CreatePackage(*PackageNameCache);
+
+				UGroomCache* CacheData = NewObject<UGroomCache>(PackageCache, *NameCache, RF_Public | RF_Standalone | RF_Transactional);
+				FHairStrandsCore::AssetHelper().RegisterAsset(CacheData);
+				
+				ChaosCache->SetCacheData(CacheData);
+			}
+			GroomCache.CacheAsset = Cast<UGroomCache>(ChaosCache->GetCacheData());
+#endif
+		}
+	}
+	return true;
+}
+
+void FGroomCacheAdapter::Record_PostSolve(UPrimitiveComponent* InComponent, const FTransform& InRootTransform, FPendingFrameWrite& OutFrame, Chaos::FReal InTime) const
+{
+	if(!bIsLoading)
+	{
+		if(const UGroomComponent* GroomComponent = CastChecked<UGroomComponent>(InComponent))
+		{
+			if(FGroomCacheData* GroomCache = GroomCaches.Find(InComponent))
+			{
+				TSharedPtr<FStrandsPositionOutput> PositionOutput = MakeShared<FStrandsPositionOutput>();
+				GroomCache->PositionsBuffer.Add(PositionOutput);
+				GroomCache->CacheTimes.Add(InTime);
+
+#if WITH_EDITOR
+				ReadbackTasks.Add(FFunctionGraphTask::CreateAndDispatchWhenReady([GroomComponent, PositionOutput](
+					ENamedThreads::Type CurrentThread, const FGraphEventRef& MyCompletionGraphEvent)
+				{
+					RequestStrandsPosition(GroomComponent, PositionOutput, true);
+				},TStatId(), nullptr, ENamedThreads::Type::RHIThread ));
+#endif
+				GroomCache->AnimInfo.EndFrame++;
+				GroomCache->AnimInfo.EndTime = InTime;
+			}
+		}
+	}
+}
+
+bool FGroomCacheAdapter::InitializeForPlayback(UPrimitiveComponent* InComponent, FObservedComponent& InObserved, float InTime)
+{
+	bIsLoading = false;
+	return false;
+}
+
+FORCEINLINE void FillGuidesPositions(const FStrandsPositionOutput& PositionsBuffer, TArray<FGroomCacheInputData>& GroupsData)
+{
+	if(PositionsBuffer.IsValid() && PositionsBuffer.Groups.Num() == GroupsData.Num())
+	{
+		for(uint32 GroupIndex = 0, NumGroups = GroupsData.Num(); GroupIndex < NumGroups; ++GroupIndex)
+		{
+			const FStrandsPositionOutput::FGroup& PositionsGroup = PositionsBuffer.Groups[GroupIndex];
+			TArray<FVector3f>& GuidesPositions = GroupsData[GroupIndex].Guides.StrandsPoints.PointsPosition;
+			
+			for(uint32 GuideIndex = 0, PointIndex = 0, NumGuides = PositionsGroup.Num(); GuideIndex < NumGuides; ++GuideIndex)
+			{
+				for(uint32 VertexIndex = 0, NumVertices = PositionsGroup[GuideIndex].Num()-1; VertexIndex < NumVertices; ++VertexIndex, ++PointIndex)
+				{
+					GuidesPositions[PointIndex] = PositionsGroup[GuideIndex][VertexIndex];
+				}
+			}
+		}
+	}
+}
+
+void FGroomCacheAdapter::WaitForSolverTasks(UPrimitiveComponent* InComponent) const
+{
+	if (ReadbackTasks.Num())
+	{
+		FTaskGraphInterface::Get().WaitUntilTasksComplete(ReadbackTasks, ENamedThreads::GameThread);
+	}
+}
+
+#if WITH_EDITORONLY_DATA
+
+void FGroomCacheAdapter::FillCacheProcessor(const UGroomComponent* GroomComponent, FGroomCacheData& GroomCache, const int32 NumFrames, float& MaxTime, int32& MaxFrame)
+{
+	if(UGroomAsset* GroomAsset = GroomComponent->GroomAsset)
+	{
+		TArray<FGroomCacheInputData> GroupsData;
+		if( UE::Groom::BuildGroupsData(GroomAsset->GetHairDescription(), GroomAsset->GetHairGroupsPlatformData(),
+			GroomAsset->GetHairGroupsInfo(), GroomAsset->GetHairGroupsInterpolation(), GroupsData))
+		{
+			if(GroomCache.PositionsBuffer.Num() == NumFrames)
+			{
+				for(int32 FrameIndex = 0; FrameIndex < NumFrames; ++FrameIndex)
+				{
+					if(GroomCache.PositionsBuffer[FrameIndex]->IsValid())
+					{
+						MaxTime = FMath::Max(MaxTime, GroomCache.CacheTimes[FrameIndex]);
+						MaxFrame = FMath::Max(MaxFrame, FrameIndex+1);
+										
+						TArray<FGroomCacheInputData> LocalData = GroupsData;
+						FillGuidesPositions(*GroomCache.PositionsBuffer[FrameIndex], LocalData);
+						GroomCache.CacheProcessor.AddGroomSample(MoveTemp(LocalData));
+					}
+				}
+			}
+		}
+	}
+}
+#endif
+
+void FGroomCacheAdapter::Finalize()
+{
+	if(!bIsLoading)
+	{
+		for(TPair<UPrimitiveComponent*,FGroomCacheData>& GroomCache : GroomCaches)
+		{
+			if(UGroomComponent* GroomComponent = CastChecked<UGroomComponent>(GroomCache.Key))
+			{
+#if WITH_EDITORONLY_DATA
+				if(!GroomCache.Value.PositionsBuffer.IsEmpty())
+				{
+					float MaxTime = -TNumericLimits<float>::Max();
+                    int32 MaxFrame = -TNumericLimits<int32>::Max();
+					FillCacheProcessor(GroomComponent, GroomCache.Value,
+						GroomCache.Value.AnimInfo.EndFrame - GroomCache.Value.AnimInfo.StartFrame, MaxTime, MaxFrame);
+					
+					GroomCache.Value.AnimInfo.EndTime = MaxTime;
+					GroomCache.Value.AnimInfo.EndFrame = MaxFrame;
+					
+					GroomCache.Value.AnimInfo.NumFrames = GroomCache.Value.AnimInfo.EndFrame - GroomCache.Value.AnimInfo.StartFrame;
+					GroomCache.Value.AnimInfo.Duration = GroomCache.Value.AnimInfo.EndTime - GroomCache.Value.AnimInfo.StartTime;
+					GroomCache.Value.AnimInfo.SecondsPerFrame = GroomCache.Value.AnimInfo.Duration / GroomCache.Value.AnimInfo.NumFrames;
+					
+					GroomCache.Value.PositionsBuffer.Reset();
+					GroomCache.Value.CacheTimes.Reset();
+					if(GroomCache.Value.CacheAsset)
+					{
+						GroomCache.Value.CacheAsset->Initialize(EGroomCacheType::Guides);
+						UE::Groom::BuildGroomCache(GroomCache.Value.CacheProcessor, GroomCache.Value.AnimInfo, GroomCache.Value.CacheAsset);
+
+						if(GroomCache.Value.CacheAsset->MarkPackageDirty())
+						{
+							FHairStrandsCore::SaveAsset(GroomCache.Value.CacheAsset);
+						}
+					}
+				}
+#endif
+				GroomComponent->SetGroomCache(nullptr);
+				GroomComponent->SetManualTick(false);
+			}
+		}
+	}
+	GroomCaches.Reset();
+	ReadbackTasks.Reset();
+}
+}    // namespace UE::Groom
diff --git a/Engine/Plugins/Runtime/HairStrands/Source/HairStrandsSolver/Private/GroomSolverComponent.cpp b/Engine/Plugins/Runtime/HairStrands/Source/HairStrandsSolver/Private/GroomSolverComponent.cpp
new file mode 100644
--- /dev/null
+++ b/Engine/Plugins/Runtime/HairStrands/Source/HairStrandsSolver/Private/GroomSolverComponent.cpp
@@ -0,0 +1,200 @@
+// Copyright Epic Games, Inc. All Rights Reserved.
+
+#include "GroomSolverComponent.h"
+#include "GroomComponent.h"
+#include "Animation/MeshDeformer.h"
+#include "Animation/MeshDeformerInstance.h"
+
+namespace UE::Groom::Private
+{
+	FORCEINLINE void GatherViewLocations(const UWorld* LocalWorld, TArray<FVector>& ViewLocations)
+	{
+		if(LocalWorld)
+		{
+			if(LocalWorld->GetPlayerControllerIterator())
+			{
+				for (FConstPlayerControllerIterator Iterator = LocalWorld->GetPlayerControllerIterator(); Iterator; ++Iterator)
+				{
+					APlayerController* PlayerController = Iterator->Get();
+					if (PlayerController && PlayerController->IsLocalPlayerController())
+					{
+						FVector PlayerLocation; FRotator CameraRotation;
+						PlayerController->GetPlayerViewPoint(PlayerLocation, CameraRotation);
+								
+						ViewLocations.Add(PlayerLocation);
+					}
+				}
+			}
+			else
+			{
+				ViewLocations = LocalWorld->ViewLocationsRenderedLastFrame;
+			}
+		}
+	}
+}
+
+UGroomSolverComponent::UGroomSolverComponent(const FObjectInitializer& ObjectInitializer)
+	: Super(ObjectInitializer)
+{
+	PrimaryComponentTick.bCanEverTick = true;
+	bTickInEditor = true;
+	bAutoActivate = true;
+	bSelectable = true;
+	Mobility = EComponentMobility::Movable;
+	bCanEverAffectNavigation = false;
+
+	// Overlap events are expensive and not needed (at least at the moment) as we don't need to collide against other component.
+	SetGenerateOverlapEvents(false);
+}
+
+#if WITH_EDITOR
+
+void UGroomSolverComponent::PostEditChangeProperty(FPropertyChangedEvent& PropertyChangedEvent)
+{
+	Super::PostEditChangeProperty(PropertyChangedEvent);
+
+	FName PropertyName = (PropertyChangedEvent.Property != NULL) ? PropertyChangedEvent.Property->GetFName() : NAME_None;
+	if (PropertyName == GET_MEMBER_NAME_CHECKED(UGroomSolverComponent, MeshDeformer))
+	{
+		SetDeformerSolver(MeshDeformer);
+	}
+}
+
+#endif
+
+void UGroomSolverComponent::OnRegister()
+{
+	Super::OnRegister();
+
+	DeformerInstance = (MeshDeformer != nullptr) ? MeshDeformer->CreateInstance(this, DeformerSettings) : nullptr;
+	GroomSolverProxy.DeformerInstance = DeformerInstance;
+}
+
+void UGroomSolverComponent::OnUnregister()
+{
+	Super::OnUnregister();
+	
+	DeformerInstance = nullptr;
+	GroomSolverProxy.DeformerInstance = nullptr;
+}
+
+void UGroomSolverComponent::SelectSimulatedCurves(TArray<int32>& NumSimulatedCurves) const
+{
+	TArray<FVector> ViewLocations;
+	UE::Groom::Private::GatherViewLocations(GetWorld(), ViewLocations);
+
+	if (ViewLocations.Num() > 0)
+	{
+		for(const TObjectPtr<UGroomComponent>& GroomComponent : GetGroomComponents())
+		{
+			const FBoxSphereBounds GroomBounds = GroomComponent->CalcBounds(GroomComponent->GetComponentTransform());
+						
+			FVector::FReal MinDistance = FMath::Square(WORLD_MAX);
+			for (const FVector& ViewLocation : ViewLocations)
+			{
+				MinDistance = FMath::Min(MinDistance, (GroomBounds.Origin - ViewLocation).SizeSquared());
+			}
+			MinDistance = FMath::Sqrt(MinDistance);
+
+			const FVector::FReal DistanceRatio = FMath::Clamp(MinDistance / GetSolverSettings().MaxLODDistance, 0.0f, 1.0f);
+
+			const uint32 NumGroups = GroomComponent->GetGroupCount();
+			for (uint32 GroupIndex = 0; GroupIndex < NumGroups; ++GroupIndex)
+			{
+				if (const FHairGroupInstance* Instance = GroomComponent->GetGroupInstance(GroupIndex))
+				{
+					const int32 NumCurves = Instance->Guides.IsValid() ? Instance->Guides.GetData().GetNumCurves() : 0;
+					NumSimulatedCurves.Add((1.0-DistanceRatio) * NumCurves);
+
+					// Keep it for now for debugging since it is WIP
+					//UE_LOG(LogTemp, Log, TEXT("Num Curves = %d Distance Ratio = %f Simulated Curves = %f | Num LODs = %d"), NumCurves, DistanceRatio, (1.0-DistanceRatio) * NumCurves, FMath::CeilLogTwo(NumCurves));
+				}
+			}
+		}
+	}
+}
+
+void UGroomSolverComponent::TickComponent(float DeltaTime, enum ELevelTick TickType, FActorComponentTickFunction* ThisTickFunction)
+{
+	Super::TickComponent(DeltaTime, TickType, ThisTickFunction);
+
+	if(DeformerInstance != nullptr)
+	{
+		MarkRenderDynamicDataDirty();
+	}
+}
+
+void UGroomSolverComponent::SendRenderDynamicData_Concurrent()
+{
+	Super::SendRenderDynamicData_Concurrent();
+	
+	if (DeformerInstance != nullptr )
+	{
+		UMeshDeformerInstance::FEnqueueWorkDesc Desc;
+		Desc.Scene = GetScene();
+		Desc.OwnerName = GetFName();
+		Desc.ExecutionGroup = UMeshDeformerInstance::ExecutionGroup_BeginInitViews;
+		DeformerInstance->EnqueueWork(Desc);
+	}
+}
+
+void UGroomSolverComponent::CreateRenderState_Concurrent(FRegisterComponentContext* Context)
+{
+	if (DeformerInstance)
+	{
+		DeformerInstance->AllocateResources();
+	}
+	
+	Super::CreateRenderState_Concurrent(Context);
+}
+
+void UGroomSolverComponent::DestroyRenderState_Concurrent()
+{
+	Super::DestroyRenderState_Concurrent();
+
+	if (DeformerInstance)
+	{
+		DeformerInstance->ReleaseResources();
+	}
+}
+
+void UGroomSolverComponent::AddGroomComponent(UGroomComponent* GroomPhysicsComponent)
+{
+	GroomComponents.Add(GroomPhysicsComponent);
+	GroomPhysicsComponent->SetGroomSolver(this);
+}
+
+void UGroomSolverComponent::RemoveGroomComponent(UGroomComponent* GroomPhysicsComponent)
+{
+	GroomComponents.Remove(GroomPhysicsComponent);
+	GroomPhysicsComponent->SetGroomSolver(nullptr);
+}
+
+void UGroomSolverComponent::ResetGroomComponents()
+{
+	for(UGroomComponent* GroomPhysicsComponent : GroomComponents)
+	{
+		GroomPhysicsComponent->SetGroomSolver(nullptr);
+	}
+	GroomComponents.Reset();
+}
+
+void UGroomSolverComponent::SetDeformerSolver(UMeshDeformer* DeformerSolver)
+{
+	MeshDeformer = DeformerSolver;
+	if(MeshDeformer && !IsBeingDestroyed())
+	{
+		DeformerSettings = MeshDeformer->CreateSettingsInstance(this);
+        DeformerInstance = MeshDeformer->CreateInstance(this, DeformerSettings);
+		GroomSolverProxy.DeformerInstance = DeformerInstance;
+	}
+	else
+	{
+		DeformerSettings = nullptr;
+		DeformerInstance = nullptr;
+		GroomSolverProxy.DeformerInstance =  nullptr;
+	}
+	MarkRenderDynamicDataDirty();
+}
+
+
diff --git a/Engine/Plugins/Runtime/HairStrands/Source/HairStrandsSolver/Private/HairStrandsSolverModule.cpp b/Engine/Plugins/Runtime/HairStrands/Source/HairStrandsSolver/Private/HairStrandsSolverModule.cpp
new file mode 100644
--- /dev/null
+++ b/Engine/Plugins/Runtime/HairStrands/Source/HairStrandsSolver/Private/HairStrandsSolverModule.cpp
@@ -0,0 +1,25 @@
+// Copyright Epic Games, Inc. All Rights Reserved.
+
+#include "HairStrandsSolverModule.h"
+#include "AddSolverDeformerNode.h"
+#include "Modules/ModuleManager.h"
+
+#define LOCTEXT_NAMESPACE "HairStrandsSolver"
+
+void FHairStrandsSolverModule::StartupModule()
+{
+	DATAFLOW_NODE_REGISTER_CREATION_FACTORY(FAddSolverDeformerDataflowNode);
+
+	GroomCacheAdapter = MakeUnique<UE::Groom::FGroomCacheAdapter>();
+	Chaos::RegisterAdapter(GroomCacheAdapter.Get());
+}
+
+void FHairStrandsSolverModule::ShutdownModule()
+{
+	Chaos::UnregisterAdapter(GroomCacheAdapter.Get());
+	GroomCacheAdapter = nullptr;
+}
+
+IMPLEMENT_MODULE(FHairStrandsSolverModule, HairStrandsSolver)
+
+#undef LOCTEXT_NAMESPACE
diff --git a/Engine/Plugins/Runtime/HairStrands/Source/HairStrandsSolver/Public/AddSolverDeformerNode.h b/Engine/Plugins/Runtime/HairStrands/Source/HairStrandsSolver/Public/AddSolverDeformerNode.h
new file mode 100644
--- /dev/null
+++ b/Engine/Plugins/Runtime/HairStrands/Source/HairStrandsSolver/Public/AddSolverDeformerNode.h
@@ -0,0 +1,92 @@
+// Copyright Epic Games, Inc. All Rights Reserved.
+
+#pragma once
+
+#include "Dataflow/DataflowObject.h"
+#include "Dataflow/DataflowSimulationContext.h"
+#include "Dataflow/DataflowSimulationNodes.h"
+
+#include "AddSolverDeformerNode.generated.h"
+
+class UOptimusDeformer;
+
+/** Add a graph deformer to the groom simulation */
+USTRUCT(meta = (DataflowSimulation))
+struct FAddSolverDeformerDataflowNode : public FDataflowSimulationNode
+{
+	GENERATED_BODY()
+	DATAFLOW_NODE_DEFINE_INTERNAL(FAddSolverDeformerDataflowNode, "AddSolverDeformer", "Physics|Solver", UDataflow::SimulationTag)
+
+public:
+	
+	FAddSolverDeformerDataflowNode(const UE::Dataflow::FNodeParameters& InParam, FGuid InGuid = FGuid::NewGuid())
+	: FDataflowSimulationNode(InParam, InGuid)
+	{
+		RegisterInputConnection(&SimulationTime);
+		RegisterInputConnection(&PhysicsSolvers);
+		RegisterOutputConnection(&PhysicsSolvers, &PhysicsSolvers);
+	}
+
+	/** Physics solvers to advance in time */
+	UPROPERTY(Transient, SkipSerialization, Meta = (DataflowInput, DataflowOutput, DataflowPassthrough = "PhysicsSolvers"))
+	TArray<FDataflowSimulationProperty> PhysicsSolvers;
+	
+    /** Delta time to use to advance the solver*/
+    UPROPERTY(Transient, SkipSerialization, Meta = (DataflowInput))
+    FDataflowSimulationTime SimulationTime = FDataflowSimulationTime(0.0f,0.0f);
+	
+	/** Graph deformer solver the component is using */
+    UPROPERTY(EditAnywhere, Category = "Groom")
+    TObjectPtr<UOptimusDeformer> MeshDeformer = nullptr;
+
+	/** List of deformer numeric inputs that will appear in the option pins */
+	UPROPERTY()
+	TArray<FDataflowNumericTypes> DeformerNumericInputs;
+
+	/** List of deformer vector inputs that will appear in the option pins */
+	UPROPERTY()
+	TArray<FDataflowVectorTypes> DeformerVectorInputs;
+
+	/** List of deformer string inputs that will appear in the option pins */
+	UPROPERTY()
+	TArray<FDataflowStringTypes> DeformerStringInputs;
+
+	/** List of deformer bool inputs that will appear in the option pins */
+	UPROPERTY()
+	TArray<FDataflowBoolTypes> DeformerBoolInputs;
+
+	/** List of deformer transform inputs that will appear in the option pins */
+	UPROPERTY()
+	TArray<FDataflowTransformTypes> DeformerTransformInputs;
+
+	/** List of deformer numeric arrays that will appear in the option pins */
+	UPROPERTY()
+	TArray<FDataflowNumericArrayTypes> DeformerNumericArrays;
+	
+	/** List of deformer vector arrays that will appear in the option pins */
+	UPROPERTY()
+	TArray<FDataflowVectorArrayTypes> DeformerVectorArrays;
+	
+	/** List of deformer string arrays that will appear in the option pins */
+	UPROPERTY()
+	TArray<FDataflowStringArrayTypes> DeformerStringArrays;
+	
+	/** List of deformer bool arrays that will appear in the option pins */
+	UPROPERTY()
+	TArray<FDataflowBoolArrayTypes> DeformerBoolArrays;
+
+	/** List of deformer transform arrays that will appear in the option pins */
+	UPROPERTY()
+	TArray<FDataflowTransformArrayTypes> DeformerTransformArrays;
+	
+	//~ Begin UDataflowNode interface
+	virtual void EvaluateSimulation(UE::Dataflow::FDataflowSimulationContext& SimulationContext, const FDataflowOutput* Output) const override;
+	virtual TArray<UE::Dataflow::FPin> AddPins() override;
+	virtual bool CanAddPin() const override { return DeformerNumericInputs.IsEmpty() && DeformerVectorInputs.IsEmpty() && DeformerStringInputs.IsEmpty(); }
+	virtual TArray<UE::Dataflow::FPin> GetPinsToRemove() const override;
+	virtual bool CanRemovePin() const override { return !CanAddPin(); }
+	virtual void OnPinRemoved(const UE::Dataflow::FPin& Pin) override;
+	virtual void OnInvalidate() override;
+	virtual void PostSerialize(const FArchive& Ar) override;
+	//~ End UDataflowNode interface
+};
diff --git a/Engine/Plugins/Runtime/HairStrands/Source/HairStrandsSolver/Public/GroomCacheAdapter.h b/Engine/Plugins/Runtime/HairStrands/Source/HairStrandsSolver/Public/GroomCacheAdapter.h
new file mode 100644
--- /dev/null
+++ b/Engine/Plugins/Runtime/HairStrands/Source/HairStrandsSolver/Public/GroomCacheAdapter.h
@@ -0,0 +1,82 @@
+﻿// Copyright Epic Games, Inc. All Rights Reserved.
+#pragma once
+
+#include "GroomBuilder.h"
+#include "Chaos/Adapters/CacheAdapter.h"
+#include "GroomCache.h"
+
+DECLARE_LOG_CATEGORY_EXTERN(LogGroomCache, Verbose, All);
+class UGroomComponent;
+class UGroomCache;
+
+namespace UE::Groom
+{
+	/**
+	 * Groom cache adapter to be able to cache groom simulation data through the chaos cache system
+	 */
+	class FGroomCacheAdapter : public Chaos::FComponentCacheAdapter
+	{
+	public:
+
+		virtual ~FGroomCacheAdapter();
+
+		// ~Begin FComponentCacheAdapter interface
+		virtual SupportType	SupportsComponentClass(UClass* InComponentClass) const override;
+		virtual UClass*	GetDesiredClass() const override;
+		virtual uint8 GetPriority() const override;
+		virtual FGuid GetGuid() const override;
+		virtual void SetRestState(UPrimitiveComponent* InComponent, UChaosCache* InCache, const FTransform& InRootTransform, Chaos::FReal InTime) const override;
+		virtual bool InitializeForRecord(UPrimitiveComponent* InComponent, FObservedComponent& InObserved) override;
+		virtual bool InitializeForPlayback(UPrimitiveComponent* InComponent, FObservedComponent& InObserved, float InTime) override;
+		virtual void InitializeForLoad(UPrimitiveComponent* InComponent, FObservedComponent& InObserved) override;
+		virtual bool ValidForPlayback(UPrimitiveComponent* InComponent, UChaosCache* InCache) const override;
+		virtual void Finalize() override;
+		virtual Chaos::FPhysicsSolver* GetComponentSolver(UPrimitiveComponent* InComponent) const override;
+		virtual Chaos::FPhysicsSolverEvents* BuildEventsSolver(UPrimitiveComponent* InComponent) const override;
+		virtual void Record_PostSolve(UPrimitiveComponent* InComp, const FTransform& InRootTransform, FPendingFrameWrite& OutFrame, Chaos::FReal InTime) const override;
+		virtual void Playback_PreSolve(UPrimitiveComponent* InComponent, UChaosCache* InCache, Chaos::FReal InTime, FPlaybackTickRecord& TickRecord, TArray<Chaos::TPBDRigidParticleHandle<Chaos::FReal, 3>*>& OutUpdatedRigids) const override {}
+		virtual void WaitForSolverTasks(UPrimitiveComponent* InComponent) const override;
+		// ~End FComponentCacheAdapter interface
+		
+	private:
+
+		/** Groom cache aata  that will store the readback results */
+		struct FGroomCacheData
+		{
+			FGroomCacheData() : CacheProcessor(EGroomCacheType::Guides, EGroomCacheAttributes::Position)
+			{}
+			
+			/** Cache processor to be used to build the groom cache */
+			FGroomCacheProcessor CacheProcessor;
+
+			/** Cache name to create the groom cache */
+			FString CacheName;
+
+			/** Animation info to build the cache */
+			FGroomAnimationInfo AnimInfo;
+
+			/** Guides positions to be used to build the cache */
+			TArray<TSharedPtr<FStrandsPositionOutput>> PositionsBuffer;
+
+			/** Cache asset used to record the positions */
+			UGroomCache* CacheAsset = nullptr;
+
+			/** Cache time at which the task is enqueued */
+			TArray<float> CacheTimes;
+		};
+		
+#if WITH_EDITORONLY_DATA
+		/** Fill the cache processor from the position buffer */
+		static void FillCacheProcessor(const UGroomComponent* GroomComponent, FGroomCacheData& GroomCache, const int32 NumFrames, float& MaxTime, int32& MaxFrame);
+#endif
+
+		/** List of groom caches that will be used to record/play caches */
+		mutable TMap<UPrimitiveComponent*,FGroomCacheData> GroomCaches;
+
+		/** Boolean to check if we are reading the cache or not */
+		bool bIsLoading = false;
+
+		/** List of readback tasks to process */
+		mutable FGraphEventArray ReadbackTasks;
+	};
+}    // namespace UE::Groom
diff --git a/Engine/Plugins/Runtime/HairStrands/Source/HairStrandsSolver/Public/GroomSolverComponent.h b/Engine/Plugins/Runtime/HairStrands/Source/HairStrandsSolver/Public/GroomSolverComponent.h
new file mode 100644
--- /dev/null
+++ b/Engine/Plugins/Runtime/HairStrands/Source/HairStrandsSolver/Public/GroomSolverComponent.h
@@ -0,0 +1,175 @@
+// Copyright Epic Games, Inc. All Rights Reserved.
+#pragma once
+
+#include "GroomComponent.h"
+#include "Components/MeshComponent.h"
+#include "UObject/ObjectMacros.h"
+#include "Dataflow/Interfaces/DataflowPhysicsSolver.h"
+
+#include "GroomSolverComponent.generated.h"
+
+class FRHIGPUBufferReadback;
+class UMeshDeformer;
+
+/** Solver settings that will be used in dataflow/deformergraph*/
+USTRUCT()
+struct FGroomSolverSettings
+{
+	GENERATED_BODY()
+	
+	/** Number of temporal sub-steps */
+	UPROPERTY(EditAnywhere, Category = "Timing")
+	int32 NumSubSteps = 2;
+
+	/** Number of iterations for the constraint solver */
+	UPROPERTY(EditAnywhere, Category = "Timing")
+	int32 NumSolverIterations = 5;
+
+	/** Frame rate for the simulation to run */
+	UPROPERTY(EditAnywhere, Category = "Timing")
+	float SimulationFrameRate = 30;
+
+	/** Maximum LOD distance (if distance between the component and the views is higher that this threshold, no simulation)*/
+	UPROPERTY(EditAnywhere, Category = "LODs")
+	float MaxLODDistance = 1000;
+};
+
+/** Dataflow groom solver proxy used in dataflow simulation */ 
+USTRUCT(BlueprintType)
+struct FDataflowGroomSolverProxy : public FDataflowPhysicsSolverProxy, public Chaos::FPhysicsSolverEvents
+{
+	GENERATED_USTRUCT_BODY()
+
+	FDataflowGroomSolverProxy(FGroomSolverSettings InProp = FGroomSolverSettings()) :
+		FDataflowPhysicsSolverProxy()
+	{
+	}
+	virtual ~FDataflowGroomSolverProxy() override
+	{
+		EventTeardown.Broadcast();
+	}
+
+	//~ Begin FPhysicsSolverInterface interface
+	virtual void AdvanceSolverDatas(const float DeltaTime) override
+	{
+		EventPreSolve.Broadcast(DeltaTime);
+		EventPreBuffer.Broadcast(DeltaTime);
+		EventPostSolve.Broadcast(DeltaTime);
+	}
+	
+	virtual bool IsValid() const override { return false;}
+	
+	virtual const UScriptStruct* GetScriptStruct() const override
+	{
+		return StaticStruct();
+	}
+	//~ End FPhysicsSolverInterface interface
+
+	/** Deformer instance coming from the component */
+	UMeshDeformerInstance* DeformerInstance = nullptr;
+
+	/** Deformer Instance GUIDs */
+	TMap<FGuid, FGuid> DeformerInstanceGuids;
+};
+
+template <>
+struct TStructOpsTypeTraits<FDataflowGroomSolverProxy> : public TStructOpsTypeTraitsBase2<FDataflowGroomSolverProxy>
+{
+	enum { WithCopy = false };
+};
+
+/** Groom solver component in which groom component could be added to be solver together */
+UCLASS(HideCategories = (Object, Physics, Collision, Activation, Mobility, "Components|Activation"), meta = (BlueprintSpawnableComponent), ClassGroup = Physics)
+class HAIRSTRANDSSOLVER_API UGroomSolverComponent : public UMeshComponent, public IDataflowPhysicsSolverInterface
+{
+	GENERATED_BODY()
+
+public:
+
+	UGroomSolverComponent(const FObjectInitializer& ObjectInitializer);
+
+	// Begin IDataflowPhysicsSolverInterface overrides
+	virtual FString GetSimulationName() const override {return GetName();};
+	virtual FDataflowSimulationAsset& GetSimulationAsset() override {return SimulationAsset;};
+	virtual const FDataflowSimulationAsset& GetSimulationAsset() const override {return SimulationAsset;};
+	virtual FDataflowSimulationProxy* GetSimulationProxy() override {return &GroomSolverProxy;}
+	virtual const FDataflowSimulationProxy* GetSimulationProxy() const  override {return &GroomSolverProxy;}
+	virtual void BuildSimulationProxy() override {}
+	virtual void ResetSimulationProxy() override {}
+	virtual void WriteToSimulation(const float DeltaTime, const bool bAsyncTask) override {}
+	virtual void ReadFromSimulation(const float DeltaTime, const bool bAsyncTask) override {}
+	// End IDataflowPhysicsSolverInterface overrides
+
+	//~ Begin UPrimitiveComponent Interface.
+#if WITH_EDITOR
+	virtual void PostEditChangeProperty(FPropertyChangedEvent& PropertyChangedEvent) override;
+#endif
+	//~ End UPrimitiveComponent Interface.
+	
+	//~ Begin USceneComponent interface 
+	virtual void TickComponent(float DeltaTime, enum ELevelTick TickType, FActorComponentTickFunction* ThisTickFunction) override;
+	virtual void OnRegister() override;
+	virtual void OnUnregister() override;
+	virtual void CreateRenderState_Concurrent(FRegisterComponentContext* Context) override;
+	virtual void SendRenderDynamicData_Concurrent() override;
+	virtual void DestroyRenderState_Concurrent() override;
+	//~ End USceneComponent interface
+
+	/** Add a groom component to the solver */
+	UFUNCTION(BlueprintCallable, Category = "Groom")
+	void AddGroomComponent(UGroomComponent* GroomComponent);
+	
+	/** Add a groom component to the solver */
+	UFUNCTION(BlueprintCallable, Category = "Groom")
+	void RemoveGroomComponent(UGroomComponent* GroomComponent);
+
+	/** Reset the groom physics components */
+    UFUNCTION(BlueprintCallable, Category = "Groom")
+    void ResetGroomComponents();
+
+	/** Reset the groom physics components */
+	const TSet< TObjectPtr<UGroomComponent> >& GetGroomComponents() const {return GroomComponents;}
+
+	/* Change the MeshDeformer solver that is used for this Component. */
+	UFUNCTION(BlueprintCallable, Category = "Groom")
+	void SetDeformerSolver(UMeshDeformer* DeformerSolver);
+
+	/** Get the groom solver settings */
+	const FGroomSolverSettings& GetSolverSettings() const {return SolverSettings;};
+
+	/** Return the mesh deformer instance */
+	UMeshDeformerInstance* GetMeshDeformerInstance() const {return DeformerInstance;}
+
+private :
+
+	/** Select the number of simulated curves based on the distance component-views */
+	void SelectSimulatedCurves(TArray<int32>& NumSimulatedCurves) const;
+
+	/** Solver settings used to control the simulation */
+	UPROPERTY(EditAnywhere, Category = "Settings")
+	FGroomSolverSettings SolverSettings;
+
+	/* Solver dataflow asset used to advance in time */
+	UPROPERTY(EditAnywhere, Category = "Dataflow", meta=(EditConditionHides), AdvancedDisplay)
+	FDataflowSimulationAsset SimulationAsset;
+
+	/** List of physics objects registered to the solver */
+	UPROPERTY(VisibleAnywhere, Category = "Groom")
+	TSet< TObjectPtr<UGroomComponent> > GroomComponents;
+
+	/** Graph deformer solver the component is using */
+	UPROPERTY(EditAnywhere, Category = "Groom", meta = (DisplayPriority = 2))
+	TObjectPtr<UMeshDeformer> MeshDeformer = nullptr;
+	
+	/** Object containing state for the MeshDeformer. */
+	UPROPERTY(Transient)
+	TObjectPtr<UMeshDeformerInstance> DeformerInstance = nullptr;
+
+	/** Object containing instance settings for the MeshDeformer. */
+	UPROPERTY(VisibleAnywhere, Instanced, Category = "Groom", meta = (DisplayName = "Deformer Settings", EditCondition = "DeformerSettings!=nullptr", ShowOnlyInnerProperties))
+	TObjectPtr<UMeshDeformerInstanceSettings> DeformerSettings = nullptr;
+
+	/** Groom solver proxy to be used in dataflow */
+	FDataflowGroomSolverProxy GroomSolverProxy;
+};
+
diff --git a/Engine/Plugins/Runtime/HairStrands/Source/HairStrandsSolver/Public/HairStrandsSolverModule.h b/Engine/Plugins/Runtime/HairStrands/Source/HairStrandsSolver/Public/HairStrandsSolverModule.h
new file mode 100644
--- /dev/null
+++ b/Engine/Plugins/Runtime/HairStrands/Source/HairStrandsSolver/Public/HairStrandsSolverModule.h
@@ -0,0 +1,22 @@
+// Copyright Epic Games, Inc. All Rights Reserved.
+
+#pragma once
+
+#include "CoreMinimal.h"
+#include "GroomCacheAdapter.h"
+#include "Modules/ModuleInterface.h"
+
+/** Hair dataflow construction module public interface */
+class FHairStrandsSolverModule : public IModuleInterface
+{
+	public :
+	
+	//~ IModuleInterface begin interface
+	virtual void StartupModule() override;
+	virtual void ShutdownModule() override;
+	//~ IModuleInterface end interface
+
+private:
+	/** Cache adapter to be used when caching grooms in dataflow */
+	TUniquePtr<UE::Groom::FGroomCacheAdapter> GroomCacheAdapter;
+};
