Fix automatic UsedWithHairStrands usage setup for material which was missing when assigned to a groom asset/component.

diff --git a/Engine/Plugins/Runtime/HairStrands/Source/HairStrandsCore/Private/GroomComponent.cpp b/Engine/Plugins/Runtime/HairStrands/Source/HairStrandsCore/Private/GroomComponent.cpp
--- a/Engine/Plugins/Runtime/HairStrands/Source/HairStrandsCore/Private/GroomComponent.cpp
+++ b/Engine/Plugins/Runtime/HairStrands/Source/HairStrandsCore/Private/GroomComponent.cpp
@@ -517,6 +517,7 @@ public:
 			if (IsHairStrandsEnabled(EHairStrandsShaderType::Strands, ShaderPlatform))
 			{
 				const int32 SlotIndex = Component->GroomAsset->GetMaterialIndex(Component->GroomAsset->GetHairGroupsRendering()[GroupIt].MaterialSlotName);
+				Component->CheckHairStrandsUsage(GetMaterialIndexWithFallback(SlotIndex));
 				const UMaterialInterface* Material = Component->GetMaterial(GetMaterialIndexWithFallback(SlotIndex), EHairGeometryType::Strands, true);
 				HairGroupMaterialProxies[GroupIt].Strands = Material ? Material->GetRenderProxy() : nullptr;
 			}
@@ -540,6 +541,7 @@ public:
 								break;
 							}
 						}
+						Component->CheckHairStrandsUsage(GetMaterialIndexWithFallback(SlotIndex));
 						const UMaterialInterface* Material = Component->GetMaterial(GetMaterialIndexWithFallback(SlotIndex), EHairGeometryType::Cards, true);
 						HairGroupMaterialProxies[GroupIt].Cards[CardsLODIndex] = Material ? Material->GetRenderProxy() : nullptr;
 					}
@@ -566,6 +568,7 @@ public:
 								break;
 							}
 						}
+						Component->CheckHairStrandsUsage(GetMaterialIndexWithFallback(SlotIndex));
 						const UMaterialInterface* Material = Component->GetMaterial(GetMaterialIndexWithFallback(SlotIndex), EHairGeometryType::Meshes, true);
 						HairGroupMaterialProxies[GroupIt].Meshes[MeshesLODIndex] = Material ? Material->GetRenderProxy() : nullptr;
 					}
@@ -2147,6 +2150,23 @@ FBoxSphereBounds UGroomComponent::CalcBounds(const FTransform& InLocalToWorld) c
 	return Out;
 }
 
+void UGroomComponent::CheckHairStrandsUsage(int32 ElementIndex) const
+{
+	UMaterialInterface* OverrideMaterial = Super::GetMaterial(ElementIndex);
+	if (!OverrideMaterial && GroomAsset && ElementIndex < GroomAsset->GetHairGroupsMaterials().Num())
+	{
+		if (UMaterialInterface* Material = GroomAsset->GetHairGroupsMaterials()[ElementIndex].Material)
+		{
+			OverrideMaterial = Material;
+		}
+	}
+	const bool bValidUsage = OverrideMaterial && OverrideMaterial->CheckMaterialUsage_Concurrent(MATUSAGE_HairStrands);
+	if(!OverrideMaterial || !bValidUsage)
+	{		
+		UE_CLOG(OverrideMaterial && !bValidUsage, LogHairStrands, Warning, TEXT("Material with missing usage flag was applied to hair strands %s"), *OverrideMaterial->GetPathName());
+	}
+}
+
 /* Return the material slot index corresponding to the material name */
 int32 UGroomComponent::GetMaterialIndex(FName MaterialSlotName) const
 {
diff --git a/Engine/Plugins/Runtime/HairStrands/Source/HairStrandsCore/Public/GroomComponent.h b/Engine/Plugins/Runtime/HairStrands/Source/HairStrandsCore/Public/GroomComponent.h
--- a/Engine/Plugins/Runtime/HairStrands/Source/HairStrandsCore/Public/GroomComponent.h
+++ b/Engine/Plugins/Runtime/HairStrands/Source/HairStrandsCore/Public/GroomComponent.h
@@ -349,6 +349,7 @@ private:
 #endif
 
 private:
+	void CheckHairStrandsUsage(int32 ElementIndex) const;
 	void DeleteDeferredHairGroupInstances();
 	void* InitializedResources;
 	class UMeshComponent* RegisteredMeshComponent;
