* Fix hair strands texture generation parameters being not preserve if the groom didn't have any import settings saved.
* Expose and save list of textures generated by the hair strands texture generation. This allows to regenerate the hair texture into the same texture, or reroute them.

diff --git a/Engine/Plugins/Runtime/HairStrands/Source/HairStrandsCore/Public/GroomCreateStrandsTexturesOptions.h b/Engine/Plugins/Runtime/HairStrands/Source/HairStrandsCore/Public/GroomCreateStrandsTexturesOptions.h
--- a/Engine/Plugins/Runtime/HairStrands/Source/HairStrandsCore/Public/GroomCreateStrandsTexturesOptions.h
+++ b/Engine/Plugins/Runtime/HairStrands/Source/HairStrandsCore/Public/GroomCreateStrandsTexturesOptions.h
@@ -80,4 +80,8 @@ public:
 	/** Number pixels expanded by the post-process dilation (0..64). */
 	UPROPERTY(EditAnywhere, BlueprintReadWrite, meta = (ShowOnlyInnerProperties, ClampMin = 0, ClampMax = 32, UIMin = 0, UIMax = 32), Category = Options)
 	int32 Dilation = 8;
+
+	/** Generated texure. */
+	UPROPERTY(EditAnywhere, BlueprintReadWrite, meta = (ShowOnlyInnerProperties, ClampMin = 0, ClampMax = 32, UIMin = 0, UIMax = 32), Category = Options)
+	TArray<TObjectPtr<UTexture2D>> GeneratedTextures;
 };
diff --git a/Engine/Plugins/Runtime/HairStrands/Source/HairStrandsEditor/Private/GroomActions.cpp b/Engine/Plugins/Runtime/HairStrands/Source/HairStrandsEditor/Private/GroomActions.cpp
--- a/Engine/Plugins/Runtime/HairStrands/Source/HairStrandsEditor/Private/GroomActions.cpp
+++ b/Engine/Plugins/Runtime/HairStrands/Source/HairStrandsEditor/Private/GroomActions.cpp
@@ -386,6 +386,23 @@ void ExecuteCreateStrandsTextures(const FToolMenuContext& InContext)
 					CurrentOptions = DuplicateObject<UGroomCreateStrandsTexturesOptions>(GroomAssetImportData->HairStrandsTexturesOptions, nullptr);
 				}
 			}
+			else
+			{
+				// Create UGroomAssetImportData and copy existing values if any
+				GroomAssetImportData = NewObject<UGroomAssetImportData>();
+				GroomAssetImportData->Rename(nullptr, GroomAsset);
+				GroomAssetImportData->SourceData = GroomAsset->AssetImportData->SourceData;
+				
+				// Create/Initialize import setings
+				const uint32 GroupCount = GroomAsset->GetNumHairGroups();
+				GroomAssetImportData->ImportOptions = NewObject<UGroomImportOptions>();
+				GroomAssetImportData->ImportOptions->Rename(nullptr, GroomAssetImportData);
+				GroomAssetImportData->ImportOptions->InterpolationSettings.SetNum(GroupCount);
+				for (uint32 GroupIndex = 0; GroupIndex < GroupCount; ++GroupIndex)
+				{
+					GroomAssetImportData->ImportOptions->InterpolationSettings[GroupIndex] = GroomAsset->GetHairGroupsInterpolation()[GroupIndex];
+				}
+			}
 
 			if (CurrentOptions == nullptr)
 			{
@@ -451,7 +468,19 @@ void ExecuteCreateStrandsTextures(const FToolMenuContext& InContext)
 						Info.GroupIndices.Add(GroupIndex);
 					}
 				}
-				FStrandsTexturesOutput Output = FGroomTextureBuilder::CreateGroomStrandsTexturesTexture(GroomAsset, Info.Resolution, Info.Layout);
+
+				FStrandsTexturesOutput Output;
+				const uint32 TextureCount = GetHairTextureLayoutTextureCount(Info.Layout);
+				if (CurrentOptions->GeneratedTextures.Num() != TextureCount)
+				{
+					Output = FGroomTextureBuilder::CreateGroomStrandsTexturesTexture(GroomAsset, Info.Resolution, Info.Layout);
+					CurrentOptions->GeneratedTextures = Output.Textures;
+				}
+				else
+				{
+					Output.Textures = CurrentOptions->GeneratedTextures;
+				}
+
 				if (Output.IsValid())
 				{
 					FGroomTextureBuilder::BuildStrandsTextures(Info, Output);
@@ -459,11 +488,12 @@ void ExecuteCreateStrandsTextures(const FToolMenuContext& InContext)
 					// Save settings used for generating these textures
 					if (GroomAssetImportData)
 					{
-						if (GroomAssetImportData->HairStrandsTexturesOptions || GroomAssetImportData->HairStrandsTexturesOptions != CurrentOptions)
+						if (!GroomAssetImportData->HairStrandsTexturesOptions || GroomAssetImportData->HairStrandsTexturesOptions != CurrentOptions)
 						{
 							// Move the transient ImportOptions to the asset package and set it on the GroomAssetImportData for serialization
 							CurrentOptions->Rename(nullptr, GroomAssetImportData);
 							GroomAssetImportData->HairStrandsTexturesOptions = CurrentOptions;
+							GroomAsset->AssetImportData = GroomAssetImportData;
 							GroomAsset->MarkPackageDirty();
 						}
 					}
