Fix for issue where grooms would disappear in some situations when the top mesh LOD is active. This particularly affects MetaHumans, as their hair will disappear when the camera gets close if the user has selected Epic quality settings.

This issue occurs due to USkeletalMesh::GetPlatformMinLODIdx not returning the lowest possible min LOD for the given platform. However, to reduce the scope of this change so close to release, the groom binding code now avoids calling GetPlatformMinLODIdx. GetPlatformMinLODIdx will be fixed for a future release.

Also fixed a related issue where the groom binding build is incorrectly reported to have failed if it was actually redundant, i.e. the binding had already been built successfully and didn't need rebuilding. This could also cause grooms to disappear in some cases.

#rnx

diff --git a/Engine/Plugins/Runtime/HairStrands/Source/HairStrandsCore/Private/GroomBindingAsset.cpp b/Engine/Plugins/Runtime/HairStrands/Source/HairStrandsCore/Private/GroomBindingAsset.cpp
--- a/Engine/Plugins/Runtime/HairStrands/Source/HairStrandsCore/Private/GroomBindingAsset.cpp
+++ b/Engine/Plugins/Runtime/HairStrands/Source/HairStrandsCore/Private/GroomBindingAsset.cpp
@@ -1116,7 +1116,9 @@ static int32 GetClampedSkeletalMeshLOD(const USkeletalMesh* SkeletalMesh, int32
 	// silently build the groom binding using LOD 3 as the source.
 	#if WITH_EDITOR
 	check(TargetPlatform);
-	const int32 BestAvailableLOD = SkeletalMesh->GetPlatformMinLODIdx(TargetPlatform);
+	// NOTE: This is a temporary workaround, as GetPlatformMinLODIdx does not currently return the 
+	// lowest possible LOD for the given platform. It will be fixed in a future release.
+	const int32 BestAvailableLOD = 0;// SkeletalMesh->GetPlatformMinLODIdx(TargetPlatform);
 	#else
 	const bool bForceLowestLODIdx = true;
 	const int32 BestAvailableLOD = SkeletalMesh->GetMinLodIdx(bForceLowestLODIdx);
@@ -1426,11 +1428,14 @@ void UGroomBindingAsset::FinishCacheDerivedDatas(FGroomBindingBuildContext& Cont
 	ReleaseAsyncProperty();
 
 	EGroomBindingAssetBuildResult Result;
-	if (!IsTemplate() && Context.bReloadResource && IsValid())
+	if (!IsTemplate() && IsValid())
 	{
 		Result = EGroomBindingAssetBuildResult::Succeeded;
 
-		InitResource();
+		if (Context.bReloadResource)
+		{
+			InitResource();
+		}
 
 		SetSourceMeshUsedLOD(Context.SourceMeshLOD);
 		SetTargetMeshUsedMinLOD(Context.TargetMeshMinLOD);
