Fix hair proxy resource not being recreate correctly when changing property on groom binding.

The component was calling MarkRenderStateDirty instead of MarkRenderInstancesDirty, which would misregister the recreation of the resources.

diff --git a/Engine/Plugins/Runtime/HairStrands/Source/HairStrandsCore/Private/GroomComponent.cpp b/Engine/Plugins/Runtime/HairStrands/Source/HairStrandsCore/Private/GroomComponent.cpp
--- a/Engine/Plugins/Runtime/HairStrands/Source/HairStrandsCore/Private/GroomComponent.cpp
+++ b/Engine/Plugins/Runtime/HairStrands/Source/HairStrandsCore/Private/GroomComponent.cpp
@@ -3353,7 +3353,7 @@ void UGroomComponent::InvalidateAndRecreate()
 {
 	InitResources(true);
 	UpdateHairSimulation();
-	MarkRenderStateDirty();
+	MarkRenderInstancesDirty();
 }
 #endif
 
