Fix crash when importing a MH with a different quality level while the same MH is present into a level.

diff --git a/Engine/Plugins/Runtime/HairStrands/Source/HairStrandsCore/Private/GroomComponent.cpp b/Engine/Plugins/Runtime/HairStrands/Source/HairStrandsCore/Private/GroomComponent.cpp
--- a/Engine/Plugins/Runtime/HairStrands/Source/HairStrandsCore/Private/GroomComponent.cpp
+++ b/Engine/Plugins/Runtime/HairStrands/Source/HairStrandsCore/Private/GroomComponent.cpp
@@ -51,6 +51,8 @@
 #include UE_INLINE_GENERATED_CPP_BY_NAME(GroomComponent)
 LLM_DECLARE_TAG(Groom);
 
+#define GROOMBINDING_SKELMESH_VALIDATION 0
+
 static int32 GHairEnableAdaptiveSubsteps = 0;  
 static FAutoConsoleVariableRef CVarHairEnableAdaptiveSubsteps(TEXT("r.HairStrands.EnableAdaptiveSubsteps"), GHairEnableAdaptiveSubsteps, TEXT("Enable adaptive solver substeps"));
 bool IsHairAdaptiveSubstepsEnabled() { return (GHairEnableAdaptiveSubsteps == 1); }
@@ -2929,10 +2931,12 @@ void UGroomComponent::InitResources(bool bIsBindingReloading)
 				check(RegisteredMeshComponent);
 				check(LocalBindingAsset);
 				check(GroupIt < LocalBindingAsset->GetHairGroupResources().Num());
+				#if GROOMBINDING_SKELMESH_VALIDATION
 				if (USkeletalMeshComponent* SkeletalMeshComponent = Cast<USkeletalMeshComponent>(RegisteredMeshComponent))
 				{
 					check(SkeletalMeshComponent->GetSkeletalMeshAsset() ? SkeletalMeshComponent->GetSkeletalMeshAsset()->GetLODNum() == LocalBindingAsset->GetHairGroupResources()[GroupIt].SimRootResources->GetLODCount() : false);
 				}
+				#endif
 
 				HairGroupInstance->Guides.RestRootResource = LocalBindingAsset->GetHairGroupResources()[GroupIt].SimRootResources;
 				HairGroupInstance->Guides.DeformedRootResource = new FHairStrandsDeformedRootResource(HairGroupInstance->Guides.RestRootResource, EHairStrandsResourcesType::Guides, ResourceName, OwnerName);
@@ -3069,10 +3073,12 @@ void UGroomComponent::InitResources(bool bIsBindingReloading)
 				check(LocalBindingAsset);
 				check(RegisteredMeshComponent);
 				check(GroupIt < LocalBindingAsset->GetHairGroupResources().Num());
+				#if GROOMBINDING_SKELMESH_VALIDATION
 				if (USkeletalMeshComponent* SkeletalMeshComponent = Cast<USkeletalMeshComponent>(RegisteredMeshComponent))
 				{
 					check(SkeletalMeshComponent->GetSkeletalMeshAsset() ? SkeletalMeshComponent->GetSkeletalMeshAsset()->GetLODNum() == LocalBindingAsset->GetHairGroupResources()[GroupIt].RenRootResources->GetLODCount() : false);
 				}
+				#endif
 
 				HairGroupInstance->Strands.RestRootResource = LocalBindingAsset->GetHairGroupResources()[GroupIt].RenRootResources;
 				HairGroupInstance->Strands.DeformedRootResource = new FHairStrandsDeformedRootResource(HairGroupInstance->Strands.RestRootResource, EHairStrandsResourcesType::Strands, ResourceName, OwnerName);
