Fix crash when importing a MH with a different quality level while the same MH is present into a level.

Similar fix than CL 38379688, but adds a missing case for hair cards.

diff --git a/Engine/Plugins/Runtime/HairStrands/Source/HairStrandsCore/Private/GroomComponent.cpp b/Engine/Plugins/Runtime/HairStrands/Source/HairStrandsCore/Private/GroomComponent.cpp
--- a/Engine/Plugins/Runtime/HairStrands/Source/HairStrandsCore/Private/GroomComponent.cpp
+++ b/Engine/Plugins/Runtime/HairStrands/Source/HairStrandsCore/Private/GroomComponent.cpp
@@ -3178,11 +3178,12 @@ void UGroomComponent::InitResources(bool bIsBindingReloading)
 					{
 						check(LocalBindingAsset);
 						check(GroupIt < LocalBindingAsset->GetHairGroupResources().Num());
+						#if GROOMBINDING_SKELMESH_VALIDATION
 						if (USkeletalMeshComponent* SkeletalMeshComponent = Cast<USkeletalMeshComponent>(RegisteredMeshComponent))
 						{
 							check(SkeletalMeshComponent->GetSkeletalMeshAsset() ? SkeletalMeshComponent->GetSkeletalMeshAsset()->GetLODNum() == LocalBindingAsset->GetHairGroupResources()[GroupIt].CardsRootResources[CardsLODIndex]->GetLODCount() : false);
 						}
-
+						#endif
 						InstanceLOD.Guides.RestRootResource = LocalBindingAsset->GetHairGroupResources()[GroupIt].CardsRootResources[CardsLODIndex];
 						InstanceLOD.Guides.DeformedRootResource = new FHairStrandsDeformedRootResource(InstanceLOD.Guides.RestRootResource, EHairStrandsResourcesType::Cards, LODResourceName, LODOwnerName);
 					}
