Fix crash when using groom binding with Nanite Skeletal mesh.

The proxy casting was incorrect causing incorrect data fetching.

diff --git a/Engine/Plugins/Runtime/HairStrands/Source/HairStrandsCore/Private/GroomManager.cpp b/Engine/Plugins/Runtime/HairStrands/Source/HairStrandsCore/Private/GroomManager.cpp
--- a/Engine/Plugins/Runtime/HairStrands/Source/HairStrandsCore/Private/GroomManager.cpp
+++ b/Engine/Plugins/Runtime/HairStrands/Source/HairStrandsCore/Private/GroomManager.cpp
@@ -597,10 +597,13 @@ static int32 GetMeshLODIndex(
 	{
 		if (Instance->Debug.GroomBindingType == EGroomBindingMeshType::SkeletalMesh)
 		{
-			if (const FSkeletalMeshSceneProxy* SceneProxy = static_cast<const FSkeletalMeshSceneProxy*>(PrimitiveSceneInfo->Proxy))
+			// For now, binding on Nanite mesh are not supported
+			if (!PrimitiveSceneInfo->Proxy->IsNaniteMesh())
 			{
-				OutLODIndex = SceneProxy->GetMeshObject() ? SceneProxy->GetMeshObject()->GetLOD() : -1;
-				OutMeshLODLocalToWorld = FTransform(SceneProxy->GetLocalToWorld());
+				if (const FSkeletalMeshSceneProxy* SceneProxy = static_cast<const FSkeletalMeshSceneProxy*>(PrimitiveSceneInfo->Proxy))
+				{
+					OutLODIndex = SceneProxy->GetMeshObject() ? SceneProxy->GetMeshObject()->GetLOD() : -1;
+				}
 			}
 		}
 		else if (Instance->Debug.GroomBindingType == EGroomBindingMeshType::GeometryCache)
@@ -608,9 +611,12 @@ static int32 GetMeshLODIndex(
 			if (const FGeometryCacheSceneProxy* SceneProxy = static_cast<const FGeometryCacheSceneProxy*>(PrimitiveSceneInfo->Proxy))
 			{
 				OutLODIndex = 0;
-				OutMeshLODLocalToWorld = FTransform(SceneProxy->GetLocalToWorld());
 			}
 		}
+		if (PrimitiveSceneInfo->Proxy)
+		{
+			OutMeshLODLocalToWorld = FTransform(PrimitiveSceneInfo->Proxy->GetLocalToWorld());
+		}
 	}
 
 	return OutLODIndex;
