Fix incorrect groom position offset for render curve.

CPU position offset was overriding GPU position offset due to incorrect condition, causing quantization issue at distance.

diff --git a/Engine/Plugins/Runtime/HairStrands/Source/HairStrandsCore/Private/GroomManager.cpp b/Engine/Plugins/Runtime/HairStrands/Source/HairStrandsCore/Private/GroomManager.cpp
--- a/Engine/Plugins/Runtime/HairStrands/Source/HairStrandsCore/Private/GroomManager.cpp
+++ b/Engine/Plugins/Runtime/HairStrands/Source/HairStrandsCore/Private/GroomManager.cpp
@@ -1398,7 +1398,7 @@ static void RunHairStrandsInterpolation_Strands(
 		// From CPU AABB
 		for (FInstanceData& InstanceData : InstanceDatas)
 		{
-			if (InstanceData.MeshLODIndex >= 0 && (InstanceData.bSimulationEnable || InstanceData.bDeformationEnable))
+			if (InstanceData.MeshLODIndex >= 0 && InstanceData.BindingType != EHairBindingType::Skinning && (InstanceData.bSimulationEnable || InstanceData.bDeformationEnable))
 			{
 				AddHairStrandUpdatePositionOffsetPass(
 					GraphBuilder,
