Fix CIS static analysis issue.

diff --git a/Engine/Plugins/Runtime/HairStrands/Source/HairStrandsCore/Private/GroomManager.cpp b/Engine/Plugins/Runtime/HairStrands/Source/HairStrandsCore/Private/GroomManager.cpp
--- a/Engine/Plugins/Runtime/HairStrands/Source/HairStrandsCore/Private/GroomManager.cpp
+++ b/Engine/Plugins/Runtime/HairStrands/Source/HairStrandsCore/Private/GroomManager.cpp
@@ -598,7 +598,7 @@ static int32 GetMeshLODIndex(
 		if (Instance->Debug.GroomBindingType == EGroomBindingMeshType::SkeletalMesh)
 		{
 			// For now, binding on Nanite mesh are not supported
-			if (!PrimitiveSceneInfo->Proxy->IsNaniteMesh())
+			if (PrimitiveSceneInfo->Proxy && !PrimitiveSceneInfo->Proxy->IsNaniteMesh())
 			{
 				if (const FSkeletalMeshSceneProxy* SceneProxy = static_cast<const FSkeletalMeshSceneProxy*>(PrimitiveSceneInfo->Proxy))
 				{
