* Add support binding triangle selection/filtering based on vertex attribute.
* Add vertex attribute transfert from LOD to LOD to use the same attriute accross LODs.
* Refactor groom binding code to share triangle closest point/triangle search

diff --git a/Engine/Plugins/Runtime/HairStrands/Source/HairStrandsCore/Private/GroomBindingAsset.cpp b/Engine/Plugins/Runtime/HairStrands/Source/HairStrandsCore/Private/GroomBindingAsset.cpp
--- a/Engine/Plugins/Runtime/HairStrands/Source/HairStrandsCore/Private/GroomBindingAsset.cpp
+++ b/Engine/Plugins/Runtime/HairStrands/Source/HairStrandsCore/Private/GroomBindingAsset.cpp
@@ -993,12 +993,13 @@ static FString BuildDerivedDataKeySuffix(const UGroomBindingAsset& BindingAsset,
 	FString GroomKey  = BindingAsset.GetGroom() ? BindingAsset.GetGroom()->GetDerivedDataKey(true /*bUseCachedKey*/) : FString();
 	FString PointKey  = FString::FromInt(BindingAsset.GetNumInterpolationPoints());
 	FString SectionKey = FString::FromInt(BindingAsset.GetMatchingSection());
+	FString TargetBindingAttributeKey = BindingAsset.GetTargetBindingAttribute().ToString();
 
-	uint32 KeyLength  = BindingType.Len() + SourceKey.Len() + TargetKey.Len() + GroomKey.Len() + PointKey.Len() + SectionKey.Len() + SourceLODKey.Len() + TargetLODKey.Len();
+	uint32 KeyLength  = BindingType.Len() + SourceKey.Len() + TargetKey.Len() + GroomKey.Len() + PointKey.Len() + SectionKey.Len() + SourceLODKey.Len() + TargetLODKey.Len() + TargetBindingAttributeKey.Len();
 
 	FString KeySuffix;
 	KeySuffix.Reserve(KeyLength);
-	KeySuffix = BindingType + SourceKey + TargetKey + GroomKey + PointKey + SectionKey + SourceLODKey + TargetLODKey;
+	KeySuffix = BindingType + SourceKey + TargetKey + GroomKey + PointKey + SectionKey + SourceLODKey + TargetLODKey + TargetBindingAttributeKey;
 	return KeySuffix;
 }
 
@@ -1715,6 +1716,7 @@ DEFINE_GROOMBINDING_MEMBER_ACCESSOR(UGeometryCache, *, TargetGeometryCache, , );
 DEFINE_GROOMBINDING_MEMBER_ACCESSOR(int32, , NumInterpolationPoints, , );
 DEFINE_GROOMBINDING_MEMBER_ACCESSOR(int32, , MatchingSection, , );
 DEFINE_GROOMBINDING_MEMBER_ACCESSOR(TArray<FGoomBindingGroupInfo>, &, GroupInfos, const, const);
+DEFINE_GROOMBINDING_MEMBER_ACCESSOR(FName, , TargetBindingAttribute, , );
 
 #if WITH_EDITORONLY_DATA
 
diff --git a/Engine/Plugins/Runtime/HairStrands/Source/HairStrandsCore/Private/GroomBindingBuilder.cpp b/Engine/Plugins/Runtime/HairStrands/Source/HairStrandsCore/Private/GroomBindingBuilder.cpp
--- a/Engine/Plugins/Runtime/HairStrands/Source/HairStrandsCore/Private/GroomBindingBuilder.cpp
+++ b/Engine/Plugins/Runtime/HairStrands/Source/HairStrandsCore/Private/GroomBindingBuilder.cpp
@@ -9,6 +9,7 @@
 #include "Logging/StructuredLog.h"
 #include "Rendering/SkeletalMeshLODRenderData.h"
 #include "Rendering/SkeletalMeshRenderData.h"
+#include "Rendering/SkeletalMeshAttributeVertexBuffer.h"
 #include "HairStrandsMeshProjection.h"
 #include "Async/ParallelFor.h"
 #include "GlobalShader.h"
@@ -61,6 +62,7 @@ FGroomBindingBuilder::FInput::FInput(UGroomBindingAsset* BindingAsset, const ITa
 : BindingType(BindingAsset->GetGroomBindingType())
 , NumInterpolationPoints(BindingAsset->GetNumInterpolationPoints())
 , MatchingSection(BindingAsset->GetMatchingSection())
+, TargetBindingAttribute(BindingAsset->GetTargetBindingAttribute())
 , SourceMeshLOD(InSourceMeshLOD)
 , TargetMeshMinLOD(InTargetMeshMinLOD)
 , bHasValidTarget(BindingAsset->HasValidTarget())
@@ -76,7 +78,7 @@ FGroomBindingBuilder::FInput::FInput(UGroomBindingAsset* BindingAsset, const ITa
 FString FGroomBindingBuilder::GetVersion()
 {
 	// Important to update the version when groom building changes
-	return TEXT("4f");
+	return TEXT("5");
 }
 
 ///////////////////////////////////////////////////////////////////////////////////////////////////
@@ -220,6 +222,10 @@ public:
 	virtual FVector2f GetVertexUV(uint32 VertexIndex, uint32 ChannelIndex) const = 0;
 	virtual FVector3f GetVertexNormal(uint32 VertexIndex) const = 0;
 	virtual int32 GetSectionFromVertexIndex(uint32 InVertIndex) const = 0;
+	virtual bool HasVertexAttribute(FName In) const = 0;
+	virtual void AllocateVertexBindingData(FName In) = 0;
+	virtual void SetVertexBinding(uint32 VertexIt, float InterpolatedAttribute) = 0;
+	virtual float GetVertexBinding(uint32 VertexIt) const = 0;
 	virtual ~IMeshLODData() {}
 };
 
@@ -231,6 +237,7 @@ public:
 	virtual int32 GetNumLODs() const = 0;
 	virtual bool IsLODValid(int32 LODIndex) const = 0;
 	virtual const IMeshLODData& GetMeshLODData(int32 LODIndex) const = 0;
+	virtual IMeshLODData& GetMeshLODData(int32 LODIndex) = 0;
 	virtual ~IMeshData() {}
 };
 
@@ -376,12 +383,76 @@ public:
 			&& MeshData->LODRenderData[LODIndex].StaticVertexBuffers.PositionVertexBuffer.GetVertexData() != nullptr;
 	}
 
+	bool HasVertexAttribute(FName In) const override
+	{
+		check(MeshData);
+		check(MeshData->LODRenderData.IsValidIndex(LODIndex));
+		return MeshData->LODRenderData[LODIndex].VertexAttributeBuffers.GetAttributeBuffer(In) != nullptr;
+	}
+
+	void AllocateVertexBindingData(FName In)
+	{
+		// If the skeletal mesh has the desired attribute, populate the vertex binding data with it
+		if (HasVertexAttribute(In))
+		{
+			check(MeshData);
+			check(MeshData->LODRenderData.IsValidIndex(LODIndex));
+			if (FSkeletalMeshAttributeVertexBuffer* VertexAttribute = MeshData->LODRenderData[LODIndex].VertexAttributeBuffers.GetAttributeBuffer(In))
+			{
+				const FSkeletalMeshAttributeVertexBuffer::FCPUData CPUData = VertexAttribute->GetCPUData();
+				if (CPUData.Data && CPUData.Data->GetNumVertices() > 0)
+				{
+					const uint32 NumVertices = CPUData.Data->GetNumVertices();
+					VertexBindingData.Init(0, NumVertices);
+
+					for (uint32 VertIt = 0; VertIt < NumVertices; ++VertIt)
+					{
+						const uint8* Element = CPUData.Data->GetDataPointer() + VertIt * CPUData.Data->GetStride();
+						switch (CPUData.PixelFormat)
+						{
+							case PF_R8:
+								VertexBindingData[VertIt] = (*Element) * (1.f / 255.f);
+								break;
+							case PF_R16F:
+								VertexBindingData[VertIt] = *(FFloat16*)Element;
+								break;
+							case PF_R32_FLOAT:
+								VertexBindingData[VertIt] = *(float*)Element;
+								break;
+						}
+					}
+				}
+			}
+		}
+
+		// Otherwise allocate dummy vertex binding data
+		const uint32 NumVertices = MeshData->LODRenderData[LODIndex].GetNumVertices();
+		if (VertexBindingData.Num() != NumVertices)
+		{
+			VertexBindingData.Init(1.f, NumVertices);
+		}
+	}
+
+	void SetVertexBinding(uint32 InVertexIndex, float In) override
+	{
+		check(VertexBindingData.IsValidIndex(InVertexIndex));
+		VertexBindingData[InVertexIndex] = In;
+	}
+
+	float GetVertexBinding(uint32 InVertexIndex) const override
+	{
+		check(VertexBindingData.IsValidIndex(InVertexIndex));
+		return VertexBindingData[InVertexIndex];
+	}
+
 private:
 	int32 LODIndex = -1;
 	const FSkeletalMeshRenderData* MeshData = nullptr;
 
 	TArray<uint32> IndexBuffer;
 	TArray<FSkeletalMeshSection> Sections;
+
+	TArray<float> VertexBindingData;
 };
 
 class FSkeletalMeshData : public IMeshData
@@ -437,6 +508,12 @@ public:
 		return MeshesLODData[InLODIndex];
 	}
 
+	virtual IMeshLODData& GetMeshLODData(int32 InLODIndex) override
+	{
+		check(MeshesLODData.IsValidIndex(InLODIndex));
+		return MeshesLODData[InLODIndex];
+	}
+
 private:
 	int32 LODCount = 0;
 	const USkeletalMesh* SkeletalMesh = nullptr;
@@ -543,6 +620,11 @@ public:
 		return LODIndex == 0;
 	}
 
+	virtual IMeshLODData& GetMeshLODData(int32 LODIndex) override
+	{
+		return *this;
+	}
+
 	virtual const IMeshLODData& GetMeshLODData(int32 LODIndex) const override
 	{
 		return *this;
@@ -601,6 +683,29 @@ public:
 		return 0;
 	}
 
+	bool HasVertexAttribute(FName In) const override
+	{
+		return false;
+	}
+
+	void AllocateVertexBindingData(FName In)
+	{
+		const uint32 NumVertices = MeshData.Positions.Num();
+		VertexBindingData.Init(1.f, NumVertices);
+	}
+
+	void SetVertexBinding(uint32 InVertexIndex, float In) override
+	{
+		check(VertexBindingData.IsValidIndex(InVertexIndex));
+		VertexBindingData[InVertexIndex] = In;
+	}
+
+	float GetVertexBinding(uint32 InVertexIndex) const override
+	{
+		check(VertexBindingData.IsValidIndex(InVertexIndex));
+		return VertexBindingData[InVertexIndex];
+	}
+
 private:
 	UGeometryCache* GeometryCache;
 	FGeometryCacheMeshData MeshData;
@@ -619,6 +724,7 @@ private:
 	};
 
 	TArray<FRange> SectionRanges;
+	TArray<float> VertexBindingData;
 };
 
 }
@@ -998,6 +1104,10 @@ namespace GroomBinding_RootProjection
 			FVector2f UV0;
 			FVector2f UV1;
 			FVector2f UV2;
+
+			float Attribute0;
+			float Attribute1;
+			float Attribute2;
 		};
 
 		struct FCell
@@ -1070,7 +1180,7 @@ namespace GroomBinding_RootProjection
 			return CellIndex;
 		}
 
-		FCells ToCells(const FVector3f& P)
+		FCells ToCells(const FVector3f& P) const
 		{
 			FCells Out;
 
@@ -1282,6 +1392,303 @@ namespace GroomBinding_RootProjection
 		return Out;
 	}
 
+	struct FQueryResult
+	{
+		FTriangleGrid::FTriangle Triangle;
+		FTrianglePoint Point;
+		bool bValid = false;
+	};
+
+	FQueryResult QueryClosestPoint(const FTriangleGrid& InGrid, const FVector3f& InQueryP)
+	{
+		FQueryResult Out;
+
+		const FTriangleGrid::FCells Cells = InGrid.ToCells(InQueryP);
+		if (Cells.Num() > 0)
+		{
+			float ClosestDistance = FLT_MAX;
+			FTrianglePoint ClosestPointOnTriangle = {};
+			FTriangleGrid::FTriangle ClosestTriangle = {};
+			for (const FTriangleGrid::FCell* Cell : Cells)
+			{
+				for (const FTriangleGrid::FTriangle& CellTriangle : Cell->Triangles)
+				{
+					const FTrianglePoint PointOnTriangle = ComputeClosestPoint(CellTriangle, InQueryP);
+					const float Distance = FVector3f::Distance(PointOnTriangle.P, InQueryP);
+					if (Distance < ClosestDistance)
+					{
+						ClosestDistance = Distance;
+						ClosestTriangle = CellTriangle;
+						ClosestPointOnTriangle = PointOnTriangle;
+					}
+				}
+			}
+			check(ClosestDistance < FLT_MAX);
+
+			Out.Triangle = ClosestTriangle;
+			Out.Point = ClosestPointOnTriangle;
+			Out.bValid = true;
+		}
+
+		return Out;
+	}
+
+	FBox3f ComputeMeshLODBound(const GroomBinding_Mesh::IMeshLODData& MeshLODData, const TArray<FVector3f>& InTransferredPositions, const FVector3f& InRootsCenter, float& OutClosestTrianglePoint, bool bFilterBindableTriangle)
+	{
+		FBox3f Out;
+		Out.Init();
+
+		const uint32 SectionCount = MeshLODData.GetNumSections();
+		const bool bHasTransferredPosition = InTransferredPositions.Num() > 0;
+
+		const uint32 MaxSectionCount = GetHairStrandsMaxSectionCount();
+		const uint32 MaxTriangleCount = GetHairStrandsMaxTriangleCount();
+		const TArray<uint32>& IndexBuffer = MeshLODData.GetIndexBuffer();
+
+		OutClosestTrianglePoint = FLT_MAX;
+		check(SectionCount > 0);
+		for (uint32 SectionIt = 0; SectionIt < SectionCount; ++SectionIt)
+		{
+			// 2.2.1 Compute the bounding box of the skeletal mesh
+			const GroomBinding_Mesh::IMeshSectionData& Section = MeshLODData.GetSection(SectionIt);
+			const uint32 TriangleCount = Section.GetNumTriangles();
+			const uint32 SectionBaseIndex = Section.GetBaseIndex();
+
+			check(TriangleCount < MaxTriangleCount);
+			check(SectionCount < MaxSectionCount);
+			check(TriangleCount > 0);
+
+			for (uint32 TriangleIt = 0; TriangleIt < TriangleCount; ++TriangleIt)
+			{
+				FTriangleGrid::FTriangle T;
+				T.TriangleIndex = TriangleIt;
+				T.SectionIndex = SectionIt;
+				T.SectionBaseIndex = SectionBaseIndex;
+
+				T.I0 = IndexBuffer[T.SectionBaseIndex + T.TriangleIndex * 3 + 0];
+				T.I1 = IndexBuffer[T.SectionBaseIndex + T.TriangleIndex * 3 + 1];
+				T.I2 = IndexBuffer[T.SectionBaseIndex + T.TriangleIndex * 3 + 2];
+			
+				const bool bBindableTriangle = bFilterBindableTriangle ? (MeshLODData.GetVertexBinding(T.I0) > 0 || MeshLODData.GetVertexBinding(T.I1) > 0 || MeshLODData.GetVertexBinding(T.I2) > 0) : true;
+				if (!bBindableTriangle)
+				{
+					continue;
+				}
+
+				if (bHasTransferredPosition)
+				{
+					T.P0 = InTransferredPositions[T.I0];
+					T.P1 = InTransferredPositions[T.I1];
+					T.P2 = InTransferredPositions[T.I2];
+				}
+				else
+				{
+					T.P0 = MeshLODData.GetVertexPosition(T.I0);
+					T.P1 = MeshLODData.GetVertexPosition(T.I1);
+					T.P2 = MeshLODData.GetVertexPosition(T.I2);
+				}
+
+				Out += T.P0;
+				Out += T.P1;
+				Out += T.P2;
+
+				// Track closest point to the groom bound
+				OutClosestTrianglePoint = FMath::Min(OutClosestTrianglePoint, (T.P0 - InRootsCenter).Length());
+				OutClosestTrianglePoint = FMath::Min(OutClosestTrianglePoint, (T.P1 - InRootsCenter).Length());
+				OutClosestTrianglePoint = FMath::Min(OutClosestTrianglePoint, (T.P2 - InRootsCenter).Length());
+			}
+		}
+
+		return Out;
+	}
+
+	FBox3f ComputeMeshLODBound(const GroomBinding_Mesh::IMeshLODData& MeshLODData, bool bFilterBindableTriangle)
+	{
+		TArray<FVector3f> DummyTransferredPositions;
+		float DummyClosestTrianglePoint;
+		FVector3f DummyRootsCenter = FVector3f::ZeroVector; 		
+		return ComputeMeshLODBound(MeshLODData, DummyTransferredPositions, DummyRootsCenter, DummyClosestTrianglePoint, bFilterBindableTriangle);
+	}
+
+	bool BuildMeshLODTriangleGrid(const GroomBinding_Mesh::IMeshLODData& MeshLODData, const TArray<FVector3f>& InTransferredPositions, FTriangleGrid& OutGrid, bool bFilterBindableTriangle)
+	{
+		const uint32 SectionCount = MeshLODData.GetNumSections();
+		const bool bHasTransferredPosition = InTransferredPositions.Num() > 0;
+		const uint32 ChannelIndex = 0;
+
+		const uint32 MaxSectionCount = GetHairStrandsMaxSectionCount();
+		const uint32 MaxTriangleCount = GetHairStrandsMaxTriangleCount();
+		const TArray<uint32>& IndexBuffer = MeshLODData.GetIndexBuffer();
+
+		bool bIsGridPopulated = false;
+		for (uint32 SectionIt = 0; SectionIt < SectionCount; ++SectionIt)
+		{
+			// 2.2.2 Insert all triangle within the grid
+			const GroomBinding_Mesh::IMeshSectionData& Section = MeshLODData.GetSection(SectionIt);
+			const uint32 TriangleCount = Section.GetNumTriangles();
+			const uint32 SectionBaseIndex = Section.GetBaseIndex();
+
+			check(TriangleCount < MaxTriangleCount);
+			check(SectionCount < MaxSectionCount);
+			check(TriangleCount > 0);
+
+			for (uint32 TriangleIt = 0; TriangleIt < TriangleCount; ++TriangleIt)
+			{
+				FTriangleGrid::FTriangle T;
+				T.TriangleIndex = TriangleIt;
+				T.SectionIndex = SectionIt;
+				T.SectionBaseIndex = SectionBaseIndex;
+
+				T.I0 = IndexBuffer[T.SectionBaseIndex + T.TriangleIndex * 3 + 0];
+				T.I1 = IndexBuffer[T.SectionBaseIndex + T.TriangleIndex * 3 + 1];
+				T.I2 = IndexBuffer[T.SectionBaseIndex + T.TriangleIndex * 3 + 2];
+
+				const bool bBindableTriangle = bFilterBindableTriangle ? (MeshLODData.GetVertexBinding(T.I0) > 0 || MeshLODData.GetVertexBinding(T.I1) > 0 || MeshLODData.GetVertexBinding(T.I2) > 0) : true;
+				if (!bBindableTriangle)
+				{
+					continue;
+				}
+
+				if (bHasTransferredPosition)
+				{
+					T.P0 = InTransferredPositions[T.I0];
+					T.P1 = InTransferredPositions[T.I1];
+					T.P2 = InTransferredPositions[T.I2];
+				}
+				else
+				{
+					T.P0 = MeshLODData.GetVertexPosition(T.I0);
+					T.P1 = MeshLODData.GetVertexPosition(T.I1);
+					T.P2 = MeshLODData.GetVertexPosition(T.I2);
+				}
+
+				// Store non-transfered position as well for RBF remapping
+				T.P0_NonTransfered = MeshLODData.GetVertexPosition(T.I0);
+				T.P1_NonTransfered = MeshLODData.GetVertexPosition(T.I1);
+				T.P2_NonTransfered = MeshLODData.GetVertexPosition(T.I2);
+
+				T.UV0 = MeshLODData.GetVertexUV(T.I0, ChannelIndex);
+				T.UV1 = MeshLODData.GetVertexUV(T.I1, ChannelIndex);
+				T.UV2 = MeshLODData.GetVertexUV(T.I2, ChannelIndex);
+
+				T.N0 = MeshLODData.GetVertexNormal(T.I0);
+				T.N1 = MeshLODData.GetVertexNormal(T.I1);
+				T.N2 = MeshLODData.GetVertexNormal(T.I2);
+
+				T.Attribute0 = MeshLODData.GetVertexBinding(T.I0);
+				T.Attribute1 = MeshLODData.GetVertexBinding(T.I1);
+				T.Attribute2 = MeshLODData.GetVertexBinding(T.I2);
+					
+				bIsGridPopulated = OutGrid.Insert(T) || bIsGridPopulated;
+			}
+		}
+
+		return bIsGridPopulated;
+	}
+
+	bool BuildMeshLODTriangleGrid(const GroomBinding_Mesh::IMeshLODData& MeshLODData, FTriangleGrid& OutGrid, bool bFilterBindableTriangle)
+	{
+		TArray<FVector3f> DummyTransferredPositions;
+		return BuildMeshLODTriangleGrid(MeshLODData, DummyTransferredPositions, OutGrid, bFilterBindableTriangle);
+	}
+
+	void TransfertMeshAttribute(
+		FName InAttributeName,
+		GroomBinding_Mesh::IMeshData* InMeshData)
+	{
+		// Check if the attribute exist on any/all LODs to see if any transfer from LOD to LOD is necessary
+		int32 SourceLODIndex = INDEX_NONE;
+		bool bAny = false;
+		bool bAll = true;
+		const int32 MeshLODCount = InMeshData->GetNumLODs();
+		for (int32 MeshLODIt = 0; MeshLODIt < MeshLODCount; ++MeshLODIt)
+		{
+			const bool bHasAttribute = InMeshData->GetMeshLODData(MeshLODIt).HasVertexAttribute(InAttributeName);
+			bAny = bAny || bHasAttribute;
+			bAll = bAll && bHasAttribute;
+
+			if (SourceLODIndex == INDEX_NONE && bHasAttribute)
+			{
+				SourceLODIndex = MeshLODIt;
+			}
+
+			InMeshData->GetMeshLODData(MeshLODIt).AllocateVertexBindingData(InAttributeName);
+		}
+
+		// If no transfer is needed, early out (i.e., all LODs have the attribute, or no LOD has the attribute)
+		const bool bNeedTransfer = bAny && !bAll;
+		if (!bNeedTransfer || SourceLODIndex == INDEX_NONE)
+		{
+			return;
+		}
+
+		// 1. Build source grid for transfer
+		FTriangleGrid* TriangleGrid = nullptr;
+		{
+			GroomBinding_Mesh::IMeshLODData& MeshLODData = InMeshData->GetMeshLODData(SourceLODIndex);
+
+			if (MeshLODData.GetNumSections() == 0)
+			{
+				UE_LOG(LogHairStrands, Error, TEXT("[Groom] Binding asset could not be built. MeshLODData has 0 sections."));
+				return;
+			}
+
+			// Compute mesh bound
+			// Take the smallest bounding box between the groom and the skeletal mesh
+			const FBox3f MeshBound = ComputeMeshLODBound(MeshLODData, false /*bFIlterBindableTriangle*/);
+			const FVector3f MeshExtent = MeshBound.Max - MeshBound.Min;
+			const FVector3f GridMin = MeshBound.Min - MeshExtent * 0.1f;
+			const FVector3f GridMax = MeshBound.Max + MeshExtent * 0.1f;
+			
+			// Build grid
+			TriangleGrid = new FTriangleGrid(GridMin, GridMax);
+			const bool bIsGridPopulated = BuildMeshLODTriangleGrid(MeshLODData, *TriangleGrid, false /*bFilterBindableTriangle*/);
+
+			if (!bIsGridPopulated)
+			{
+				delete TriangleGrid;
+				return;
+			}
+		}
+
+		// 2. Transfer the attribute to all the LODs
+		for (int32 MeshLODIt = 0; MeshLODIt < MeshLODCount; ++MeshLODIt)
+		{
+			const bool bHasAttribute = InMeshData->GetMeshLODData(MeshLODIt).HasVertexAttribute(InAttributeName);
+			if (bHasAttribute)
+			{
+				continue;
+			}
+
+			GroomBinding_Mesh::IMeshLODData& MeshLODData = InMeshData->GetMeshLODData(MeshLODIt);
+			const uint32 SectionCount = MeshLODData.GetNumSections();
+			for (uint32 SectionIt = 0; SectionIt < SectionCount; ++SectionIt)
+			{
+				// 2.2.2 Insert all triangle within the grid
+				const GroomBinding_Mesh::IMeshSectionData& Section = MeshLODData.GetSection(SectionIt);
+				const uint32 VertexCount = Section.GetNumVertices();
+				const uint32 BaseVertexIndex = Section.GetBaseVertexIndex();
+				
+				GroomBinding_Mesh::IMeshLODData* LocalMeshLODData = &MeshLODData;
+				ParallelFor(VertexCount, [InAttributeName, SectionIt, TriangleGrid, LocalMeshLODData, BaseVertexIndex] (uint32 VertexIt)
+				{
+					const uint32 VertexIndex = BaseVertexIndex + VertexIt;
+					const FVector3f P = LocalMeshLODData->GetVertexPosition(VertexIndex);
+					const FQueryResult Q = QueryClosestPoint(*TriangleGrid, P);
+
+					const float InterpolatedAttribute = 
+						Q.Triangle.Attribute0 * Q.Point.Barycentric.X +
+						Q.Triangle.Attribute1 * Q.Point.Barycentric.Y +
+						Q.Triangle.Attribute2 * Q.Point.Barycentric.Z;
+
+					LocalMeshLODData->SetVertexBinding(VertexIndex,  InterpolatedAttribute);
+				});
+			}
+		}
+
+		delete TriangleGrid;
+	} // Transfer attribute
+
 	static bool Project(
 		const GroomBinding_Mesh::FHairStrandRootPositions& InRootPositions,
 		const GroomBinding_Mesh::IMeshData* InMeshData,
@@ -1319,87 +1726,15 @@ namespace GroomBinding_RootProjection
 			const uint32 MaxSectionCount = GetHairStrandsMaxSectionCount();
 			const uint32 MaxTriangleCount = GetHairStrandsMaxTriangleCount();
 
-			FBox3f MeshBound;
-			MeshBound.Init();
 			const uint32 SectionCount = MeshLODData.GetNumSections();
-			
 			if ( SectionCount == 0 )
 			{
 				UE_LOG(LogHairStrands, Error, TEXT("[Groom] Binding asset could not be built. MeshLODData has 0 sections."));
 				return false;
 			}
 
-			// Only cull section if the 'matching section' exist
-			// For now, never cull section as this can cause binding issue. The 'section' can be shuffled/reindexed 
-			// based on project settings, causing the binding on a wrong surface
-			const bool bCullSection = false; //InMatchingSection > 0 && InMatchingSection < SectionCount;
-
 			float ClosestTrianglePoint = FLT_MAX;
-			check(SectionCount > 0);
-			for (uint32 SectionIt = 0; SectionIt < SectionCount; ++SectionIt)
-			{
-				// Cull out sections which don't match
-				if (bCullSection && SectionIt != InMatchingSection)
-				{
-					continue;
-				}
-
-				// 2.2.1 Compute the bounding box of the skeletal mesh
-				const GroomBinding_Mesh::IMeshSectionData& Section = MeshLODData.GetSection(SectionIt);
-				const uint32 TriangleCount = Section.GetNumTriangles();
-				const uint32 SectionBaseIndex = Section.GetBaseIndex();
-
-				check(TriangleCount < MaxTriangleCount);
-				check(SectionCount < MaxSectionCount);
-				check(TriangleCount > 0);
-
-				for (uint32 TriangleIt = 0; TriangleIt < TriangleCount; ++TriangleIt)
-				{
-					FTriangleGrid::FTriangle T;
-					T.TriangleIndex = TriangleIt;
-					T.SectionIndex = SectionIt;
-					T.SectionBaseIndex = SectionBaseIndex;
-
-					T.I0 = IndexBuffer[T.SectionBaseIndex + T.TriangleIndex * 3 + 0];
-					T.I1 = IndexBuffer[T.SectionBaseIndex + T.TriangleIndex * 3 + 1];
-					T.I2 = IndexBuffer[T.SectionBaseIndex + T.TriangleIndex * 3 + 2];
-
-					if (bHasTransferredPosition)
-					{
-						T.P0 = InTransferredPositions[MeshLODIt][T.I0];
-						T.P1 = InTransferredPositions[MeshLODIt][T.I1];
-						T.P2 = InTransferredPositions[MeshLODIt][T.I2];
-					}
-					else
-					{
-						T.P0 = MeshLODData.GetVertexPosition(T.I0);
-						T.P1 = MeshLODData.GetVertexPosition(T.I1);
-						T.P2 = MeshLODData.GetVertexPosition(T.I2);
-					}
-
-					// Store non-transfered position as well for RBF remapping
-					T.P0_NonTransfered = MeshLODData.GetVertexPosition(T.I0);
-					T.P1_NonTransfered = MeshLODData.GetVertexPosition(T.I1);
-					T.P2_NonTransfered = MeshLODData.GetVertexPosition(T.I2);
-
-					T.UV0 = MeshLODData.GetVertexUV(T.I0, ChannelIndex);
-					T.UV1 = MeshLODData.GetVertexUV(T.I1, ChannelIndex);
-					T.UV2 = MeshLODData.GetVertexUV(T.I2, ChannelIndex);
-
-					T.N0 = MeshLODData.GetVertexNormal(T.I0);
-					T.N1 = MeshLODData.GetVertexNormal(T.I1);
-					T.N2 = MeshLODData.GetVertexNormal(T.I2);
-
-					MeshBound += T.P0;
-					MeshBound += T.P1;
-					MeshBound += T.P2;
-
-					// Track closest point to the groom bound
-					ClosestTrianglePoint = FMath::Min(ClosestTrianglePoint, (T.P0 - InRootPositions.BoundingBox.GetCenter()).Length());
-					ClosestTrianglePoint = FMath::Min(ClosestTrianglePoint, (T.P1 - InRootPositions.BoundingBox.GetCenter()).Length());
-					ClosestTrianglePoint = FMath::Min(ClosestTrianglePoint, (T.P2 - InRootPositions.BoundingBox.GetCenter()).Length());
-				}
-			}
+			const FBox3f MeshBound = ComputeMeshLODBound(MeshLODData, bHasTransferredPosition ? InTransferredPositions[MeshLODIt] : TArray<FVector3f>(), InRootPositions.BoundingBox.GetCenter(), ClosestTrianglePoint, true /*bFIlterBindableTriangle*/);
 
 			// Take the smallest bounding box between the groom and the skeletal mesh
 			const FVector3f MeshExtent = MeshBound.Max - MeshBound.Min;
@@ -1427,65 +1762,7 @@ namespace GroomBinding_RootProjection
 			}
 
 			FTriangleGrid Grid(GridMin, GridMax);
-			bool bIsGridPopulated = false;
-			for (uint32 SectionIt = 0; SectionIt < SectionCount; ++SectionIt)
-			{
-				// 2.2.2 Insert all triangle within the grid
-				const GroomBinding_Mesh::IMeshSectionData& Section = MeshLODData.GetSection(SectionIt);
-				const uint32 TriangleCount = Section.GetNumTriangles();
-				const uint32 SectionBaseIndex = Section.GetBaseIndex();
-
-				check(TriangleCount < MaxTriangleCount);
-				check(SectionCount < MaxSectionCount);
-				check(TriangleCount > 0);
-
-				// Cull out sections which don't match
-				if (bCullSection && SectionIt != InMatchingSection)
-				{
-					continue;
-				}
-
-				for (uint32 TriangleIt = 0; TriangleIt < TriangleCount; ++TriangleIt)
-				{
-					FTriangleGrid::FTriangle T;
-					T.TriangleIndex = TriangleIt;
-					T.SectionIndex = SectionIt;
-					T.SectionBaseIndex = SectionBaseIndex;
-
-					T.I0 = IndexBuffer[T.SectionBaseIndex + T.TriangleIndex * 3 + 0];
-					T.I1 = IndexBuffer[T.SectionBaseIndex + T.TriangleIndex * 3 + 1];
-					T.I2 = IndexBuffer[T.SectionBaseIndex + T.TriangleIndex * 3 + 2];
-
-					if (bHasTransferredPosition)
-					{
-						T.P0 = InTransferredPositions[MeshLODIt][T.I0];
-						T.P1 = InTransferredPositions[MeshLODIt][T.I1];
-						T.P2 = InTransferredPositions[MeshLODIt][T.I2];
-					}
-					else
-					{
-						T.P0 = MeshLODData.GetVertexPosition(T.I0);
-						T.P1 = MeshLODData.GetVertexPosition(T.I1);
-						T.P2 = MeshLODData.GetVertexPosition(T.I2);
-					}
-
-					// Store non-transfered position as well for RBF remapping
-					T.P0_NonTransfered = MeshLODData.GetVertexPosition(T.I0);
-					T.P1_NonTransfered = MeshLODData.GetVertexPosition(T.I1);
-					T.P2_NonTransfered = MeshLODData.GetVertexPosition(T.I2);
-
-					T.UV0 = MeshLODData.GetVertexUV(T.I0, ChannelIndex);
-					T.UV1 = MeshLODData.GetVertexUV(T.I1, ChannelIndex);
-					T.UV2 = MeshLODData.GetVertexUV(T.I2, ChannelIndex);
-
-					T.N0 = MeshLODData.GetVertexNormal(T.I0);
-					T.N1 = MeshLODData.GetVertexNormal(T.I1);
-					T.N2 = MeshLODData.GetVertexNormal(T.I2);
-
-					bIsGridPopulated = Grid.Insert(T) || bIsGridPopulated;
-				}
-			}
-
+			const bool bIsGridPopulated = BuildMeshLODTriangleGrid(MeshLODData, bHasTransferredPosition ? InTransferredPositions[MeshLODIt] : TArray<FVector3f>(), Grid, true /*bFIlterBindableTriangle*/);
 			if (!bIsGridPopulated)
 			{
 				UE_LOG(LogHairStrands, Error, TEXT("[Groom] Binding asset could not be built. The target skeletal mesh could be missing UVs."));
@@ -1525,35 +1802,17 @@ namespace GroomBinding_RootProjection
 		#endif
 			{
 				const FVector3f& RootP = InRootPositions.RootPositions[CurveIndex];
-				const FTriangleGrid::FCells Cells = Grid.ToCells(RootP);
-
-				if (Cells.Num() == 0)
+				const FQueryResult QueryResult = QueryClosestPoint(Grid, RootP);
+				if (!QueryResult.bValid)
 				{
-				#if BINDING_PARALLEL_BUILDING
+					#if BINDING_PARALLEL_BUILDING
 					bIsValid = 0; return;
-				#else
+					#else
 					return false;
-				#endif
+					#endif
 				}
-
-				float ClosestDistance = FLT_MAX;
-				FTriangleGrid::FTriangle ClosestTriangle = {};
-				FVector2f ClosestBarycentrics = FVector2f::ZeroVector;
-				for (const FTriangleGrid::FCell* Cell : Cells)
-				{
-					for (const FTriangleGrid::FTriangle& CellTriangle : Cell->Triangles)
-					{
-						const FTrianglePoint Tri = ComputeClosestPoint(CellTriangle, RootP);
-						const float Distance = FVector3f::Distance(Tri.P, RootP);
-						if (Distance < ClosestDistance)
-						{
-							ClosestDistance = Distance;
-							ClosestTriangle = CellTriangle;
-							ClosestBarycentrics = FVector2f(Tri.Barycentric.X, Tri.Barycentric.Y);
-						}
-					}
-				}
-				check(ClosestDistance < FLT_MAX);
+				const FTriangleGrid::FTriangle& ClosestTriangle = QueryResult.Triangle;
+				const FVector2f ClosestBarycentrics(QueryResult.Point.Barycentric.X, QueryResult.Point.Barycentric.Y);
 
 				// Record closest triangle and the root's barycentrics
 				const uint32 EncodedBarycentrics = FHairStrandsRootUtils::PackBarycentrics(FVector2f(ClosestBarycentrics));	// LWC_TODO: Precision loss
@@ -2658,6 +2917,13 @@ bool BuildHairRootGroupData(const FGroomBindingBuilder::FInput& In, uint32 InGro
 			SlowTask.EnterProgressFrame();
 		}
 
+		// 1.3bis Populate binding attribute
+		{
+			GroomBinding_RootProjection::TransfertMeshAttribute(
+				In.TargetBindingAttribute,
+				TargetMeshData.Get());
+		}
+
 		// 1.4 Build root data for guides/strands/cards
 		{
 			// Guides
diff --git a/Engine/Plugins/Runtime/HairStrands/Source/HairStrandsCore/Private/GroomCreateBindingOptions.cpp b/Engine/Plugins/Runtime/HairStrands/Source/HairStrandsCore/Private/GroomCreateBindingOptions.cpp
--- a/Engine/Plugins/Runtime/HairStrands/Source/HairStrandsCore/Private/GroomCreateBindingOptions.cpp
+++ b/Engine/Plugins/Runtime/HairStrands/Source/HairStrandsCore/Private/GroomCreateBindingOptions.cpp
@@ -14,5 +14,5 @@ UGroomCreateBindingOptions::UGroomCreateBindingOptions(const FObjectInitializer&
 	TargetGeometryCache = nullptr;
 	NumInterpolationPoints = 100;
 	MatchingSection = 0;
-}
-
+	TargetBindingAttribute = NAME_None;
+}
\ No newline at end of file
diff --git a/Engine/Plugins/Runtime/HairStrands/Source/HairStrandsCore/Private/HairStrandsCore.cpp b/Engine/Plugins/Runtime/HairStrands/Source/HairStrandsCore/Private/HairStrandsCore.cpp
--- a/Engine/Plugins/Runtime/HairStrands/Source/HairStrandsCore/Private/HairStrandsCore.cpp
+++ b/Engine/Plugins/Runtime/HairStrands/Source/HairStrandsCore/Private/HairStrandsCore.cpp
@@ -175,10 +175,10 @@ UStaticMesh* FHairStrandsCore::CreateStaticMesh(const FString& InPackageName, co
 
 UGroomBindingAsset* FHairStrandsCore::CreateGroomBindingAsset(const FString& InPackageName, UObject* InParent, UGroomAsset* GroomAsset, USkeletalMesh* SourceSkelMesh, USkeletalMesh* TargetSkelMesh, const int32 NumInterpolationPoints, const int32 MatchingSection)
 {
-	return CreateGroomBindingAsset(EGroomBindingMeshType::SkeletalMesh, InPackageName, InParent, GroomAsset, SourceSkelMesh, TargetSkelMesh, NumInterpolationPoints, MatchingSection);
+	return CreateGroomBindingAsset(EGroomBindingMeshType::SkeletalMesh, InPackageName, InParent, GroomAsset, SourceSkelMesh, TargetSkelMesh, NumInterpolationPoints, MatchingSection, NAME_None);
 }
 
-UGroomBindingAsset* FHairStrandsCore::CreateGroomBindingAsset(EGroomBindingMeshType BindingType, const FString& InPackageName, UObject* InParent, UGroomAsset* GroomAsset, UObject* Source, UObject* Target, const int32 NumInterpolationPoints, const int32 MatchingSection)
+UGroomBindingAsset* FHairStrandsCore::CreateGroomBindingAsset(EGroomBindingMeshType BindingType, const FString& InPackageName, UObject* InParent, UGroomAsset* GroomAsset, UObject* Source, UObject* Target, const int32 NumInterpolationPoints, const int32 MatchingSection, const FName TargetMeshAttribute)
 {
 	if (!GroomAsset ||
 		!((BindingType == EGroomBindingMeshType::SkeletalMesh && Cast<USkeletalMesh>(Target)) || 
@@ -239,6 +239,7 @@ UGroomBindingAsset* FHairStrandsCore::CreateGroomBindingAsset(EGroomBindingMeshT
 		Out->GetHairGroupsPlatformData().Reserve(GroomAsset->GetHairGroupsPlatformData().Num());
 		Out->SetNumInterpolationPoints(NumInterpolationPoints);
 		Out->SetMatchingSection(MatchingSection);
+		Out->SetTargetBindingAttribute(TargetMeshAttribute);
 		Out->MarkPackageDirty();
 		AssetHelper().RegisterAsset(Out);
 		return Out;
@@ -246,14 +247,14 @@ UGroomBindingAsset* FHairStrandsCore::CreateGroomBindingAsset(EGroomBindingMeshT
 	return nullptr;
 }
 
-UGroomBindingAsset* FHairStrandsCore::CreateGroomBindingAsset(EGroomBindingMeshType BindingType, UGroomAsset* GroomAsset, UObject* Source, UObject* Target, const int32 NumInterpolationPoints, const int32 MatchingSection)
+UGroomBindingAsset* FHairStrandsCore::CreateGroomBindingAsset(EGroomBindingMeshType BindingType, UGroomAsset* GroomAsset, UObject* Source, UObject* Target, const int32 NumInterpolationPoints, const int32 MatchingSection, const FName TargetMeshAttribute)
 {
-	return CreateGroomBindingAsset(BindingType, FString(), nullptr, GroomAsset, Source, Target, NumInterpolationPoints, MatchingSection);
+	return CreateGroomBindingAsset(BindingType, FString(), nullptr, GroomAsset, Source, Target, NumInterpolationPoints, MatchingSection, TargetMeshAttribute);
 }
 
 UGroomBindingAsset* FHairStrandsCore::CreateGroomBindingAsset(UGroomAsset* GroomAsset, USkeletalMesh* SourceSkelMesh, USkeletalMesh* TargetSkelMesh, const int32 NumInterpolationPoints, const int32 MatchingSection)
 {
-	return CreateGroomBindingAsset(EGroomBindingMeshType::SkeletalMesh, FString(), nullptr, GroomAsset, SourceSkelMesh, TargetSkelMesh, NumInterpolationPoints, MatchingSection);
+	return CreateGroomBindingAsset(EGroomBindingMeshType::SkeletalMesh, FString(), nullptr, GroomAsset, SourceSkelMesh, TargetSkelMesh, NumInterpolationPoints, MatchingSection, NAME_None);
 }
 
 void FHairStrandsCore::SaveAsset(UObject* Object)
diff --git a/Engine/Plugins/Runtime/HairStrands/Source/HairStrandsCore/Public/GroomBindingAsset.h b/Engine/Plugins/Runtime/HairStrands/Source/HairStrandsCore/Public/GroomBindingAsset.h
--- a/Engine/Plugins/Runtime/HairStrands/Source/HairStrandsCore/Public/GroomBindingAsset.h
+++ b/Engine/Plugins/Runtime/HairStrands/Source/HairStrandsCore/Public/GroomBindingAsset.h
@@ -71,6 +71,7 @@ enum class EGroomBindingAsyncProperties : uint64
 	GroupInfos = 1 << 12,
 	HairGroupResources = 1 << 13,
 	HairGroupPlatformData = 1 << 14,
+	TargetBindingAttribute = 1 << 15,
 	All = MAX_uint64
 };
 
@@ -311,6 +312,17 @@ public:
 	UFUNCTION(BlueprintGetter) int32 GetMatchingSection() const;
 	UFUNCTION(BlueprintSetter) void SetMatchingSection(int32 InMatchingSection);
 
+private:
+	/** Optional binding attribute name on target skeletal mesh, to filter out which triangles are valid to bind groom to*/
+	UE_DEPRECATED(5.3, "Please do not access this member directly; use UGroomBindingAsset::GetTargetBindingAttribute() or UGroomBindingAsset::SetTargetBindingAttribute().")
+	UPROPERTY(EditAnywhere, BlueprintGetter = GetTargetBindingAttribute, BlueprintSetter = SetTargetBindingAttribute, Category = "BuildSettings")
+	FName TargetBindingAttribute = NAME_None;
+
+public:
+	static FName GetTargetBindingAttributeMemberName();
+	UFUNCTION(BlueprintGetter) FName GetTargetBindingAttribute() const;
+	UFUNCTION(BlueprintSetter) void SetTargetBindingAttribute(FName InAttributeName);
+
 private:
 	UE_DEPRECATED(5.3, "Please do not access this member directly; use UGroomBindingAsset::GetGroupInfos() or UGroomBindingAsset::SetGroupInfos().")
 	UPROPERTY(EditAnywhere, EditFixedSize, BlueprintGetter = GetGroupInfos, BlueprintSetter = SetGroupInfos, Category = "HairGroups", meta = (DisplayName = "Group"))
diff --git a/Engine/Plugins/Runtime/HairStrands/Source/HairStrandsCore/Public/GroomBindingBuilder.h b/Engine/Plugins/Runtime/HairStrands/Source/HairStrandsCore/Public/GroomBindingBuilder.h
--- a/Engine/Plugins/Runtime/HairStrands/Source/HairStrandsCore/Public/GroomBindingBuilder.h
+++ b/Engine/Plugins/Runtime/HairStrands/Source/HairStrandsCore/Public/GroomBindingBuilder.h
@@ -24,6 +24,7 @@ struct HAIRSTRANDSCORE_API FGroomBindingBuilder
 		EGroomBindingMeshType BindingType = EGroomBindingMeshType::SkeletalMesh;
 		int32 NumInterpolationPoints = 0;
 		int32 MatchingSection = 0;
+		FName TargetBindingAttribute = NAME_None;
 
 		// These must be initialized to valid LODs.
 		// 
diff --git a/Engine/Plugins/Runtime/HairStrands/Source/HairStrandsCore/Public/GroomCreateBindingOptions.h b/Engine/Plugins/Runtime/HairStrands/Source/HairStrandsCore/Public/GroomCreateBindingOptions.h
--- a/Engine/Plugins/Runtime/HairStrands/Source/HairStrandsCore/Public/GroomCreateBindingOptions.h
+++ b/Engine/Plugins/Runtime/HairStrands/Source/HairStrandsCore/Public/GroomCreateBindingOptions.h
@@ -45,4 +45,8 @@ public:
 	/** Section to pick to transfer the position */
 	UPROPERTY(EditAnywhere, BlueprintReadWrite, AdvancedDisplay, meta = (ShowOnlyInnerProperties), Category = Conversion)
 	int32 MatchingSection;
+
+	/** Optional mesh attribute used to filter triangles to bind groom to */
+	UPROPERTY(EditAnywhere, BlueprintReadWrite, meta = (ShowOnlyInnerProperties), Category = BuildSettings)
+	FName TargetBindingAttribute = NAME_None;
 };
diff --git a/Engine/Plugins/Runtime/HairStrands/Source/HairStrandsCore/Public/HairStrandsCore.h b/Engine/Plugins/Runtime/HairStrands/Source/HairStrandsCore/Public/HairStrandsCore.h
--- a/Engine/Plugins/Runtime/HairStrands/Source/HairStrandsCore/Public/HairStrandsCore.h
+++ b/Engine/Plugins/Runtime/HairStrands/Source/HairStrandsCore/Public/HairStrandsCore.h
@@ -56,8 +56,8 @@ public:
 	HAIRSTRANDSCORE_API static UGroomBindingAsset* CreateGroomBindingAsset(const FString& InPackageName, UObject* InParent, UGroomAsset* GroomAsset, USkeletalMesh* SourceSkelMesh, USkeletalMesh* TargetSkelMesh, const int32 NumInterpolationPoints, const int32 MatchingSection);
 
 	// Create binding asset from groom asset and assets of specified groom binding type
-	HAIRSTRANDSCORE_API static UGroomBindingAsset* CreateGroomBindingAsset(EGroomBindingMeshType BindingType, UGroomAsset* GroomAsset, UObject* Source, UObject* Target, const int32 NumInterpolationPoints, const int32 MatchingSection);
-	HAIRSTRANDSCORE_API static UGroomBindingAsset* CreateGroomBindingAsset(EGroomBindingMeshType BindingType, const FString& InPackageName, UObject* InParent, UGroomAsset* GroomAsset, UObject* Source, UObject* Target, const int32 NumInterpolationPoints, const int32 MatchingSection);
+	HAIRSTRANDSCORE_API static UGroomBindingAsset* CreateGroomBindingAsset(EGroomBindingMeshType BindingType, UGroomAsset* GroomAsset, UObject* Source, UObject* Target, const int32 NumInterpolationPoints, const int32 MatchingSection, const FName TargetMeshAttribute = NAME_None);
+	HAIRSTRANDSCORE_API static UGroomBindingAsset* CreateGroomBindingAsset(EGroomBindingMeshType BindingType, const FString& InPackageName, UObject* InParent, UGroomAsset* GroomAsset, UObject* Source, UObject* Target, const int32 NumInterpolationPoints, const int32 MatchingSection, const FName TargetMeshAttribute = NAME_None);
 
 	HAIRSTRANDSCORE_API static void SaveAsset(UObject* Object);
 
diff --git a/Engine/Plugins/Runtime/HairStrands/Source/HairStrandsEditor/Private/GroomActions.cpp b/Engine/Plugins/Runtime/HairStrands/Source/HairStrandsEditor/Private/GroomActions.cpp
--- a/Engine/Plugins/Runtime/HairStrands/Source/HairStrandsEditor/Private/GroomActions.cpp
+++ b/Engine/Plugins/Runtime/HairStrands/Source/HairStrandsEditor/Private/GroomActions.cpp
@@ -222,7 +222,7 @@ void ExecuteCreateBindingAsset(const FToolMenuContext& InContext)
 					{
 						CurrentOptions->SourceSkeletalMesh->ConditionalPostLoad();
 					}
-					BindingAsset = FHairStrandsCore::CreateGroomBindingAsset(CurrentOptions->GroomBindingType, GroomAsset, CurrentOptions->SourceSkeletalMesh, CurrentOptions->TargetSkeletalMesh, CurrentOptions->NumInterpolationPoints, CurrentOptions->MatchingSection);
+					BindingAsset = FHairStrandsCore::CreateGroomBindingAsset(CurrentOptions->GroomBindingType, GroomAsset, CurrentOptions->SourceSkeletalMesh, CurrentOptions->TargetSkeletalMesh, CurrentOptions->NumInterpolationPoints, CurrentOptions->MatchingSection, CurrentOptions->TargetBindingAttribute);
 				}
 				else
 				{
@@ -231,7 +231,7 @@ void ExecuteCreateBindingAsset(const FToolMenuContext& InContext)
 					{
 						CurrentOptions->SourceGeometryCache->ConditionalPostLoad();
 					}
-					BindingAsset = FHairStrandsCore::CreateGroomBindingAsset(CurrentOptions->GroomBindingType, GroomAsset, CurrentOptions->SourceGeometryCache, CurrentOptions->TargetGeometryCache, CurrentOptions->NumInterpolationPoints, CurrentOptions->MatchingSection);
+					BindingAsset = FHairStrandsCore::CreateGroomBindingAsset(CurrentOptions->GroomBindingType, GroomAsset, CurrentOptions->SourceGeometryCache, CurrentOptions->TargetGeometryCache, CurrentOptions->NumInterpolationPoints, CurrentOptions->MatchingSection, CurrentOptions->TargetBindingAttribute);
 				}
 
 				if (BindingAsset)
diff --git a/Engine/Plugins/Runtime/HairStrands/Source/HairStrandsEditor/Private/GroomBindingDetailsCustomization.cpp b/Engine/Plugins/Runtime/HairStrands/Source/HairStrandsEditor/Private/GroomBindingDetailsCustomization.cpp
--- a/Engine/Plugins/Runtime/HairStrands/Source/HairStrandsEditor/Private/GroomBindingDetailsCustomization.cpp
+++ b/Engine/Plugins/Runtime/HairStrands/Source/HairStrandsEditor/Private/GroomBindingDetailsCustomization.cpp
@@ -4,15 +4,138 @@
 #include "DetailLayoutBuilder.h"
 #include "GroomBindingAsset.h"
 #include "GroomCreateBindingOptions.h"
+#include "DetailWidgetRow.h"
+#include "GroomBindingCompiler.h"
+#include "Widgets/Input/SComboBox.h"
+#include "Widgets/Text/STextBlock.h"
+
+///////////////////////////////////////////////////////////////////////////////////////////////////
+
+template<typename T>
+T* GetCutomizeDetailObject(IDetailLayoutBuilder& LayoutBuilder)
+{
+	T* Out = nullptr;
+
+	const TArray<TWeakObjectPtr<UObject>>& SelectedObjects = LayoutBuilder.GetSelectedObjects();
+	check(SelectedObjects.Num() <= 1);
+	if (SelectedObjects.Num() > 0)
+	{
+		Out = Cast<T>(SelectedObjects[0].Get());
+	}
+	return Out;
+}
+
+static TArray<FName> GetSkelMeshAttributes(USkeletalMesh* SkeletalMesh)
+{
+	TArray<FName> Out;
+	if (SkeletalMesh)
+	{
+		if (FSkeletalMeshLODInfo* LODInfo = SkeletalMesh->GetLODInfo(0/*LODIndex*/))
+		{
+			for (const FSkeletalMeshVertexAttributeInfo& Attribute : LODInfo->VertexAttributes)
+			{
+				if (Attribute.IsEnabledForRender())
+				{
+					Out.Add(Attribute.Name);
+				}
+			}
+		}
+	}
+	return Out;
+}
+
+// Target attribute
+static void AddBindingAttributeSelection(
+	IDetailLayoutBuilder& LayoutBuilder, 
+	TSharedRef<IPropertyHandle> InBindingAttributeProperty, 
+	UObject* InObject, 
+	FGroomBindingAttributeSelection* In, 
+	EGroomBindingMeshType InBindingType, 
+	USkeletalMesh* InTargetSkeletalMesh)
+{
+	if (In && InBindingAttributeProperty->IsValidHandle())
+	{
+		In->SelectedBindingAttribute = 0;
+		In->BindingAttributeNames.Empty();
+		In->BindingAttributeNames.Add(TEXT("No Attribute"));
+		if (InBindingType == EGroomBindingMeshType::SkeletalMesh)
+		{
+			In->BindingAttributeNames.Append(GetSkelMeshAttributes(InTargetSkeletalMesh));
+		}
+
+		IDetailPropertyRow& PropertyRow = LayoutBuilder.AddPropertyToCategory(InBindingAttributeProperty);
+		FDetailWidgetRow& WidgetRow = PropertyRow.CustomWidget();
+
+		TSharedPtr<SWidget> NameWidget;
+		TSharedPtr<SWidget> ValueWidget;
+		PropertyRow.GetDefaultWidgets(NameWidget, ValueWidget, WidgetRow);
+
+		TSharedRef<SComboBox<FName>> AttributeComboBox = SNew(SComboBox<FName>)
+		.OptionsSource(&In->BindingAttributeNames)
+		.OnSelectionChanged_Lambda([InObject, In, InBindingAttributeProperty](FName NewSelection, ESelectInfo::Type SelectInfo)
+		{
+			In->SelectedBindingAttribute = In->BindingAttributeNames.Find(NewSelection);
+			
+			if (UGroomBindingAsset* BindingAsset = Cast<UGroomBindingAsset>(InObject))
+			{
+				if (In->BindingAttributeNames.IsValidIndex(In->SelectedBindingAttribute))
+				{
+					BindingAsset->SetTargetBindingAttribute(In->BindingAttributeNames[In->SelectedBindingAttribute]);
+				}
+			}
+			else if (UGroomCreateBindingOptions* BindingOptions = Cast<UGroomCreateBindingOptions>(InObject))
+			{
+				if (In->BindingAttributeNames.IsValidIndex(In->SelectedBindingAttribute))
+				{
+					BindingOptions->TargetBindingAttribute = In->BindingAttributeNames[In->SelectedBindingAttribute];
+				}
+			}
+			
+			FPropertyChangedEvent PropertyUpdate(InBindingAttributeProperty->GetProperty());
+			InObject->PostEditChangeProperty(PropertyUpdate);
+		})
+		.OnGenerateWidget_Lambda([In](FName Item)
+		{
+			return SNew(STextBlock).Text(FText::FromName(Item));
+		})
+		.Content()
+		[
+			SNew(STextBlock)
+			.Text_Lambda([In] 
+			{ 
+				FName SelectedName = In->BindingAttributeNames.IsValidIndex(In->SelectedBindingAttribute) ? In->BindingAttributeNames[In->SelectedBindingAttribute] : NAME_None;
+				return FText::FromName(SelectedName);
+			})
+		];
+
+		WidgetRow.NameContent()
+		[
+			NameWidget.ToSharedRef()
+		]
+		.ValueContent()
+		[
+			AttributeComboBox
+		];
+	}
+}
+
+static void RefreshLayoutOnValueChanged(IDetailLayoutBuilder& LayoutBuilder, TSharedRef<IPropertyHandle> Property)
+{
+	Property->SetOnPropertyValueChanged(FSimpleDelegate::CreateLambda([&LayoutBuilder]()
+	{
+		LayoutBuilder.ForceRefreshDetails();
+	}));
+}
+
+///////////////////////////////////////////////////////////////////////////////////////////////////
 
 // Customization for UGroomBindingAsset
 void FGroomBindingDetailsCustomization::CustomizeDetails(IDetailLayoutBuilder& LayoutBuilder)
 {
-	TSharedRef<IPropertyHandle> GroomBindingType = LayoutBuilder.GetProperty(UGroomBindingAsset::GetGroomBindingTypeMemberName());
+	UGroomBindingAsset* Asset = GetCutomizeDetailObject<UGroomBindingAsset>(LayoutBuilder);
 	
-	uint8 EnumValue;
-	GroomBindingType->GetValue(EnumValue);
-	if (EnumValue == (uint8) EGroomBindingMeshType::SkeletalMesh)
+	// GroomBindingType show/hide
+	if (Asset->GetGroomBindingType() == EGroomBindingMeshType::SkeletalMesh)
 	{
 		LayoutBuilder.HideProperty(UGroomBindingAsset::GetSourceGeometryCacheMemberName());
 		LayoutBuilder.HideProperty(UGroomBindingAsset::GetTargetGeometryCacheMemberName());
@@ -22,6 +145,18 @@ void FGroomBindingDetailsCustomization::CustomizeDetails(IDetailLayoutBuilder& L
 		LayoutBuilder.HideProperty(UGroomBindingAsset::GetSourceSkeletalMeshMemberName());
 		LayoutBuilder.HideProperty(UGroomBindingAsset::GetTargetSkeletalMeshMemberName());
 	}
+
+	// GroomBindingType
+	TSharedRef<IPropertyHandle> GroomBindingType = LayoutBuilder.GetProperty(UGroomBindingAsset::GetGroomBindingTypeMemberName());
+	RefreshLayoutOnValueChanged(LayoutBuilder, GroomBindingType);
+
+	// Target Skeletal Mesh
+	TSharedRef<IPropertyHandle> TargetSkeletalMesh = LayoutBuilder.GetProperty(UGroomBindingAsset::GetTargetSkeletalMeshMemberName());
+	RefreshLayoutOnValueChanged(LayoutBuilder, TargetSkeletalMesh);
+
+	// Target attribute
+	TSharedRef<IPropertyHandle> TargetBindingAttribute = LayoutBuilder.GetProperty(UGroomBindingAsset::GetTargetBindingAttributeMemberName());
+	AddBindingAttributeSelection(LayoutBuilder, TargetBindingAttribute, Asset, &BindingAttributeSelection, Asset->GetGroomBindingType(), Asset->GetTargetSkeletalMesh());
 }
 
 TSharedRef<IDetailCustomization> FGroomBindingDetailsCustomization::MakeInstance()
@@ -29,14 +164,15 @@ TSharedRef<IDetailCustomization> FGroomBindingDetailsCustomization::MakeInstance
 	return MakeShared<FGroomBindingDetailsCustomization>();
 }
 
+///////////////////////////////////////////////////////////////////////////////////////////////////
+
 // Customization for UGroomCreateBindingOptions
 void FGroomCreateBindingDetailsCustomization::CustomizeDetails(IDetailLayoutBuilder& LayoutBuilder)
 {
-	TSharedRef<IPropertyHandle> GroomBindingType = LayoutBuilder.GetProperty(GET_MEMBER_NAME_CHECKED(UGroomCreateBindingOptions, GroomBindingType));
+	UGroomCreateBindingOptions* Options = GetCutomizeDetailObject<UGroomCreateBindingOptions>(LayoutBuilder);
 
-	uint8 EnumValue;
-	GroomBindingType->GetValue(EnumValue);
-	if (EnumValue == (uint8) EGroomBindingMeshType::SkeletalMesh)
+	// GroomBindingType show/hide
+	if (Options->GroomBindingType == EGroomBindingMeshType::SkeletalMesh)
 	{
 		LayoutBuilder.HideProperty(GET_MEMBER_NAME_CHECKED(UGroomCreateBindingOptions, SourceGeometryCache));
 		LayoutBuilder.HideProperty(GET_MEMBER_NAME_CHECKED(UGroomCreateBindingOptions, TargetGeometryCache));
@@ -47,16 +183,20 @@ void FGroomCreateBindingDetailsCustomization::CustomizeDetails(IDetailLayoutBuil
 		LayoutBuilder.HideProperty(GET_MEMBER_NAME_CHECKED(UGroomCreateBindingOptions, TargetSkeletalMesh));
 	}
 
-	FSimpleDelegate OnImportTypeChangedDelegate = FSimpleDelegate::CreateSP(this, &FGroomCreateBindingDetailsCustomization::OnGroomBindingTypeChanged, &LayoutBuilder);
-	GroomBindingType->SetOnPropertyValueChanged(OnImportTypeChangedDelegate);
+	// GroomBindingType
+	TSharedRef<IPropertyHandle> GroomBindingType = LayoutBuilder.GetProperty(GET_MEMBER_NAME_CHECKED(UGroomCreateBindingOptions, GroomBindingType));
+	RefreshLayoutOnValueChanged(LayoutBuilder, GroomBindingType);
+
+	// Target Skeletal Mesh
+	TSharedRef<IPropertyHandle> TargetSkeletalMesh = LayoutBuilder.GetProperty(GET_MEMBER_NAME_CHECKED(UGroomCreateBindingOptions, TargetSkeletalMesh));
+	RefreshLayoutOnValueChanged(LayoutBuilder, TargetSkeletalMesh);
+
+	// Target attribute
+	TSharedRef<IPropertyHandle> TargetBindingAttribute = LayoutBuilder.GetProperty(GET_MEMBER_NAME_CHECKED(UGroomCreateBindingOptions, TargetBindingAttribute));
+	AddBindingAttributeSelection(LayoutBuilder, TargetBindingAttribute, Options, &BindingAttributeSelection, Options->GroomBindingType, Options->TargetSkeletalMesh);
 }
 
 TSharedRef<IDetailCustomization> FGroomCreateBindingDetailsCustomization::MakeInstance()
 {
 	return MakeShared<FGroomCreateBindingDetailsCustomization>();
 }
-
-void FGroomCreateBindingDetailsCustomization::OnGroomBindingTypeChanged(IDetailLayoutBuilder* LayoutBuilder)
-{
-	LayoutBuilder->ForceRefreshDetails();
-}
diff --git a/Engine/Plugins/Runtime/HairStrands/Source/HairStrandsEditor/Private/GroomBindingDetailsCustomization.h b/Engine/Plugins/Runtime/HairStrands/Source/HairStrandsEditor/Private/GroomBindingDetailsCustomization.h
--- a/Engine/Plugins/Runtime/HairStrands/Source/HairStrandsEditor/Private/GroomBindingDetailsCustomization.h
+++ b/Engine/Plugins/Runtime/HairStrands/Source/HairStrandsEditor/Private/GroomBindingDetailsCustomization.h
@@ -4,14 +4,20 @@
 
 #include "IDetailCustomization.h"
 
+struct FGroomBindingAttributeSelection
+{
+	int32 SelectedBindingAttribute = INDEX_NONE;
+	TArray<FName> BindingAttributeNames;
+};
+
 class FGroomBindingDetailsCustomization : public IDetailCustomization
 {
 public:
 	/** Begin IDetailCustomization interface */
 	virtual void CustomizeDetails(IDetailLayoutBuilder& LayoutBuilder) override;
 	/** End IDetailCustomization interface */
-
 	static TSharedRef<IDetailCustomization> MakeInstance();
+	FGroomBindingAttributeSelection BindingAttributeSelection;
 };
 
 class FGroomCreateBindingDetailsCustomization : public IDetailCustomization
@@ -20,8 +26,6 @@ public:
 	/** Begin IDetailCustomization interface */
 	virtual void CustomizeDetails(IDetailLayoutBuilder& LayoutBuilder) override;
 	/** End IDetailCustomization interface */
-
 	static TSharedRef<IDetailCustomization> MakeInstance();
-protected:
-	void OnGroomBindingTypeChanged(IDetailLayoutBuilder* LayoutBuilder);
+	FGroomBindingAttributeSelection BindingAttributeSelection;
 };
diff --git a/Engine/Source/Runtime/Engine/Public/Rendering/SkeletalMeshAttributeVertexBuffer.h b/Engine/Source/Runtime/Engine/Public/Rendering/SkeletalMeshAttributeVertexBuffer.h
--- a/Engine/Source/Runtime/Engine/Public/Rendering/SkeletalMeshAttributeVertexBuffer.h
+++ b/Engine/Source/Runtime/Engine/Public/Rendering/SkeletalMeshAttributeVertexBuffer.h
@@ -39,6 +39,20 @@ public:
 	}
 
 	void Serialize(FArchive& Ar);
+
+	struct FCPUData
+	{
+		FSkeletalMeshVertexDataInterface* Data = nullptr;
+		int32 PixelFormat = PF_Unknown;
+	};
+
+	FCPUData GetCPUData() const
+	{
+		FCPUData Out;
+		Out.Data = ValueData;
+		Out.PixelFormat = PixelFormat;
+		return Out;
+	}
 	
 private:
 	void Allocate();
