Fix CitySample crash.

diff --git a/Engine/Plugins/Runtime/HairStrands/Source/HairStrandsCore/Private/GroomBindingBuilder.cpp b/Engine/Plugins/Runtime/HairStrands/Source/HairStrandsCore/Private/GroomBindingBuilder.cpp
--- a/Engine/Plugins/Runtime/HairStrands/Source/HairStrandsCore/Private/GroomBindingBuilder.cpp
+++ b/Engine/Plugins/Runtime/HairStrands/Source/HairStrandsCore/Private/GroomBindingBuilder.cpp
@@ -385,9 +385,13 @@ public:
 
 	bool HasVertexAttribute(FName In) const override
 	{
-		check(MeshData);
-		check(MeshData->LODRenderData.IsValidIndex(LODIndex));
-		return MeshData->LODRenderData[LODIndex].VertexAttributeBuffers.GetAttributeBuffer(In) != nullptr;
+		bool bHasAttribute = false;
+		if (MeshData)
+		{
+			check(MeshData->LODRenderData.IsValidIndex(LODIndex));
+			bHasAttribute = MeshData->LODRenderData[LODIndex].VertexAttributeBuffers.GetAttributeBuffer(In) != nullptr;
+		}
+		return bHasAttribute;
 	}
 
 	void AllocateVertexBindingData(FName In)
@@ -426,7 +430,7 @@ public:
 		}
 
 		// Otherwise allocate dummy vertex binding data
-		const uint32 NumVertices = MeshData->LODRenderData[LODIndex].GetNumVertices();
+		const uint32 NumVertices = MeshData ? MeshData->LODRenderData[LODIndex].GetNumVertices() : 0u;
 		if (VertexBindingData.Num() != NumVertices)
 		{
 			VertexBindingData.Init(1.f, NumVertices);
