Fix bone transform niagara groom DI crash

diff --git a/Engine/Plugins/Runtime/HairStrands/Source/HairStrandsCore/Private/GroomComponent.cpp b/Engine/Plugins/Runtime/HairStrands/Source/HairStrandsCore/Private/GroomComponent.cpp
--- a/Engine/Plugins/Runtime/HairStrands/Source/HairStrandsCore/Private/GroomComponent.cpp
+++ b/Engine/Plugins/Runtime/HairStrands/Source/HairStrandsCore/Private/GroomComponent.cpp
@@ -1154,17 +1154,23 @@ private:
 	UMaterialInterface* Strands_DebugMaterial = nullptr;
 };
 
-FHairGroupInstance* GetHairGroupInstance(UGroomComponent* In, int32 InGroupIndex)
+// Return a ref counted hair group instance. Will be removed/modified when we will switch to use PrimitiveId instead of components
+TRefCountPtr<FHairGroupInstance> GetRefCountedHairGroupInstance(const UGroomComponent* GroomComponent, const int32 GroupIndex)
 {
-	check(In);
-	FHairStrandsSceneProxy* Proxy = (FHairStrandsSceneProxy*)In->GetSceneProxy();
-	if (Proxy && Proxy->HairGroupInstances.IsValidIndex(InGroupIndex))
+	const FHairStrandsSceneProxy* GroomProxy = GroomComponent ? static_cast<FHairStrandsSceneProxy*>(GroomComponent->GetSceneProxy()) : nullptr;
+	if (GroomProxy && GroomProxy->HairGroupInstances.IsValidIndex(GroupIndex))
 	{
-		return Proxy->HairGroupInstances[InGroupIndex];
+		return GroomProxy->HairGroupInstances[GroupIndex];
 	}
 	return nullptr;
 }
 
+// @todo: deprecate this function in 5.6 to use GetRefCountedHairGroupInstance
+FHairGroupInstance* GetHairGroupInstance(UGroomComponent* GroomComponent, int32 GroupIndex)
+{
+	return GetRefCountedHairGroupInstance(GroomComponent, GroupIndex).GetReference();
+}
+
 /** GroomCacheBuffers implementation that hold copies of the GroomCacheAnimationData needed for playback */
 class FGroomCacheBuffers : public IGroomCacheBuffers
 {
diff --git a/Engine/Plugins/Runtime/HairStrands/Source/HairStrandsCore/Private/Niagara/NiagaraDataInterfaceHairStrands.cpp b/Engine/Plugins/Runtime/HairStrands/Source/HairStrandsCore/Private/Niagara/NiagaraDataInterfaceHairStrands.cpp
--- a/Engine/Plugins/Runtime/HairStrands/Source/HairStrandsCore/Private/Niagara/NiagaraDataInterfaceHairStrands.cpp
+++ b/Engine/Plugins/Runtime/HairStrands/Source/HairStrandsCore/Private/Niagara/NiagaraDataInterfaceHairStrands.cpp
@@ -249,6 +249,15 @@ static FAutoConsoleVariableRef CVarHairSimulationMaxDelay(TEXT("r.HairStrands.Si
 static int32 GHairSimulationRestUpdate = false;
 static FAutoConsoleVariableRef CVarHairSimulationRestUpdate(TEXT("r.HairStrands.SimulationRestUpdate"), GHairSimulationRestUpdate, TEXT("Update the simulation rest pose"));
 
+// Temporary cvar just in case using the latest hair group instance is changing the behavior 
+static bool GEnableProxyInstanceTransform = false;
+static FAutoConsoleVariableRef CVarEnableProxyInstanceTransform(TEXT("r.HairStrands.Debug.EnableProxyInstanceTransform"), GEnableProxyInstanceTransform, TEXT("Enable the use of the niagara proxy instance to compute the bone/world transform. For debug only"));
+
+//------------------------------------------------------------------------------------------------------------
+
+// Forward declaration of the GroomComponent function that is retrieving a ref counted group instance ptr given an index
+TRefCountPtr<FHairGroupInstance> GetRefCountedHairGroupInstance(const UGroomComponent* GroomComponent, const int32 GroupIndex);
+
 //------------------------------------------------------------------------------------------------------------
 // Helper struct/function to extracting used resource prior to dispatching
 
@@ -287,10 +296,10 @@ struct FHairGroupInstanceRDG
 	bool IsRootValid() const	{ return bIsRootValid && IsValid(); }
 };
 
-static FHairGroupInstanceRDG Convert(FRDGBuilder& GraphBuilder, const FHairGroupInstance* In)
+static FHairGroupInstanceRDG Convert(FRDGBuilder& GraphBuilder, const TRefCountPtr<FHairGroupInstance>& In)
 {
 	FHairGroupInstanceRDG Out;
-	if (In && In->Guides.IsValid())
+	if (In.IsValid() && In->Guides.IsValid())
 	{
 		const int32 MeshLODIndex = In->HairGroupPublicData ? In->HairGroupPublicData->GetMeshLODIndex() : -1;
 
@@ -943,31 +952,44 @@ IMPLEMENT_SHADER_TYPE(, FInterpolateGroomGuidesCS, TEXT("/Plugin/Runtime/HairStr
 
 //------------------------------------------------------------------------------------------------------------
 
-FMatrix44f ComputeBoneTransform(const FNDIHairStrandsData* InstanceData)
+FMatrix44f ComputeBoneTransform(const FNDIHairStrandsData* InstanceData, const TRefCountPtr<FHairGroupInstance>& GroupInstance)
 {
-	const FMatrix44f RigidTransformFloat = FMatrix44f(InstanceData->HairGroupInstance ? InstanceData->HairGroupInstance->Debug.RigidCurrentLocalToWorld.ToMatrixWithScale() : InstanceData->WorldTransform.ToMatrixWithScale());
-	return FMatrix44f(InstanceData->BoneTransform.ToMatrixWithScale()) * RigidTransformFloat;
+	if(InstanceData)
+	{
+		const FTransform RigidTransform = GroupInstance.IsValid() ? GroupInstance->Debug.RigidCurrentLocalToWorld : InstanceData->WorldTransform;
+		if(RigidTransform.IsValid() && InstanceData->BoneTransform.IsValid())
+		{
+			return  FMatrix44f((InstanceData->BoneTransform * RigidTransform).ToMatrixWithScale());
+		}
+	}
+	return FMatrix44f::Identity;
 }
 
-FMatrix44f ComputeWorldTransform(const FNDIHairStrandsData* InstanceData)
+FMatrix44f ComputeWorldTransform(const FNDIHairStrandsData* InstanceData, const TRefCountPtr<FHairGroupInstance>& GroupInstance)
 {
-	const FMatrix44f BoneTransformFloat = ComputeBoneTransform(InstanceData);
-	
-	FMatrix44f WorldTransformFloat = FMatrix44f(InstanceData->HairGroupInstance ? InstanceData->HairGroupInstance->GetCurrentLocalToWorld().ToMatrixWithScale() : InstanceData->WorldTransform.ToMatrixWithScale());
-	if (InstanceData->LocalSimulation)
+	if(InstanceData)
 	{
-		const FMatrix44d WorldTransformDouble(WorldTransformFloat);
-		const FMatrix44d BoneTransformDouble(BoneTransformFloat);
+		const FMatrix44f BoneTransformFloat = ComputeBoneTransform(InstanceData, GroupInstance);
 
-		// Due to large world coordinate we store the relative world transform in double precision 
-		WorldTransformFloat = FMatrix44f(WorldTransformDouble * BoneTransformDouble.Inverse());
-	}
+		const FTransform WorldTransform = GroupInstance.IsValid() ? GroupInstance->GetCurrentLocalToWorld() : InstanceData->WorldTransform;
+		FMatrix44f WorldTransformFloat = WorldTransform.IsValid() ? FMatrix44f(WorldTransform.ToMatrixWithScale()) : FMatrix44f::Identity;
+	
+		if (InstanceData->LocalSimulation)
+		{
+			const FMatrix44d WorldTransformDouble(WorldTransformFloat);
+			const FMatrix44d BoneTransformDouble(BoneTransformFloat);
+
+			// Due to large world coordinate we compute the relative world transform in double precision 
+			WorldTransformFloat = FMatrix44f(WorldTransformDouble * BoneTransformDouble.Inverse());
+		}
 		
-	if (WorldTransformFloat.ContainsNaN())
-	{
-		WorldTransformFloat = FMatrix44f(InstanceData->WorldTransform.ToMatrixWithScale());
+		if (WorldTransformFloat.ContainsNaN() && InstanceData->WorldTransform.IsValid())
+		{
+			WorldTransformFloat = FMatrix44f(InstanceData->WorldTransform.ToMatrixWithScale());
+		}
+		return WorldTransformFloat;
 	}
-	return WorldTransformFloat;
+	return FMatrix44f::Identity;
 }
 
 static void InterpolateGroomGuides(
@@ -976,7 +998,7 @@ static void InterpolateGroomGuides(
 	const uint32 NodePositionComponent,  
 	const uint32 RestPositionComponent, 
 	FNDIHairStrandsBuffer* HairStrandsBuffer,
-	FHairGroupInstance* HairGroupInstance,
+	const TRefCountPtr<FHairGroupInstance>& HairGroupInstance,
 	const uint32 StrandsSize, 
 	const FMatrix44f& WorldToLocal)
 {
@@ -1061,10 +1083,20 @@ void UNiagaraDataInterfaceHairStrands::SimCachePostReadFrame(void* OptionalPerIn
 						FMemMark MemMark(FMemStack::Get());
 						FRDGBuilder GraphBuilder(RHICmdList);
 
-						const FMatrix44f WorldToLocal = ComputeWorldTransform(&ProxyData).Inverse();
-
-						InterpolateGroomGuides(GraphBuilder, ParticlesBuffer, NodePositionComponent, RestPositionComponent,
-							ProxyData.HairStrandsBuffer, ProxyData.HairGroupInstance, ProxyData.StrandsSize, WorldToLocal);
+						const TRefCountPtr<FHairGroupInstance> HairGroupInstance = GetRefCountedHairGroupInstance(ProxyData.HairGroupInstSource.Get(), ProxyData.HairGroupIndex);
+						
+						if(!GEnableProxyInstanceTransform)
+						{
+							const FMatrix44f WorldToLocal = ComputeWorldTransform(&ProxyData, HairGroupInstance).Inverse();
+							InterpolateGroomGuides(GraphBuilder, ParticlesBuffer, NodePositionComponent, RestPositionComponent,
+								ProxyData.HairStrandsBuffer, HairGroupInstance, ProxyData.StrandsSize, WorldToLocal);
+						}
+						else
+						{
+							const FMatrix44f WorldToLocal = ComputeWorldTransform(&ProxyData, ProxyData.HairGroupInstance).Inverse();
+							InterpolateGroomGuides(GraphBuilder, ParticlesBuffer, NodePositionComponent, RestPositionComponent,
+								ProxyData.HairStrandsBuffer, ProxyData.HairGroupInstance, ProxyData.StrandsSize, WorldToLocal);
+						}
 						GraphBuilder.Execute();
 					}
 				);
@@ -3426,13 +3458,10 @@ void UNiagaraDataInterfaceHairStrands::SetShaderParameters(const FNiagaraDataInt
 	FNDIHairStrandsProxy& DIProxy = Context.GetProxy<FNDIHairStrandsProxy>();
 	FNDIHairStrandsData* ProxyData = DIProxy.SystemInstancesToProxyData.Find(Context.GetSystemInstanceID());
 	
-	TRefCountPtr<FHairGroupInstance> HairGroupInstance = nullptr;
-	if (ProxyData && ProxyData->HairGroupInstSource.Get())
-	{
-		HairGroupInstance = GetHairGroupInstance(ProxyData->HairGroupInstSource.Get(), ProxyData->HairGroupIndex);
-	}
+	const TRefCountPtr<FHairGroupInstance> HairGroupInstance = ProxyData ?
+		GetRefCountedHairGroupInstance(ProxyData->HairGroupInstSource.Get(), ProxyData->HairGroupIndex) : nullptr;
 
-	FHairGroupInstanceRDG InstanceRDG = Convert(GraphBuilder, ProxyData ? HairGroupInstance : nullptr);
+	FHairGroupInstanceRDG InstanceRDG = Convert(GraphBuilder, HairGroupInstance);
 	
 	const int32 MeshLODIndex = InstanceRDG.MeshLODIndex;
 	const bool bIsHairValid = ProxyData != nullptr && ProxyData->HairStrandsBuffer && ProxyData->HairStrandsBuffer->IsInitialized();
@@ -3467,8 +3496,18 @@ void UNiagaraDataInterfaceHairStrands::SetShaderParameters(const FNiagaraDataInt
 
 		// Offsets / Transforms
 
-		const FMatrix44f WorldTransformFloat = ComputeWorldTransform(ProxyData);
-		const FMatrix44f BoneTransformFloat = ComputeBoneTransform(ProxyData);
+		FMatrix44f WorldTransformFloat = FMatrix44f::Identity;
+		FMatrix44f BoneTransformFloat = FMatrix44f::Identity;
+		if(!GEnableProxyInstanceTransform)
+		{
+			WorldTransformFloat =  ComputeWorldTransform(ProxyData, HairGroupInstance);
+			BoneTransformFloat = ComputeBoneTransform(ProxyData, HairGroupInstance);
+		}
+		else
+		{
+			WorldTransformFloat =  ComputeWorldTransform(ProxyData, ProxyData->HairGroupInstance);
+			BoneTransformFloat = ComputeBoneTransform(ProxyData, ProxyData->HairGroupInstance);
+		}
 		
 		if (ProxyData->BoneLinearVelocity.ContainsNaN() || ProxyData->BoneAngularVelocity.ContainsNaN() || ProxyData->BoneLinearAcceleration.ContainsNaN() || ProxyData->BoneAngularAcceleration.ContainsNaN())
 		{
@@ -3607,14 +3646,11 @@ void FNDIHairStrandsProxy::PostSimulate(const FNDIGpuComputePostSimulateContext&
 		return;
 	}
 
-	TRefCountPtr<FHairGroupInstance> HairGroupInstance = nullptr;
-	if (ProxyData->HairGroupInstSource.Get())
-	{
-		HairGroupInstance = GetHairGroupInstance(ProxyData->HairGroupInstSource.Get(), ProxyData->HairGroupIndex);
-	}
+	const TRefCountPtr<FHairGroupInstance> HairGroupInstance = ProxyData ?
+		GetRefCountedHairGroupInstance(ProxyData->HairGroupInstSource.Get(), ProxyData->HairGroupIndex) : nullptr;
 
 	// MGPU DeformedPositionBuffer copy after simulation
-	FHairStrandsDeformedResource* GuideDeformedBuffer = HairGroupInstance ? HairGroupInstance->Guides.DeformedResource : nullptr;
+	FHairStrandsDeformedResource* GuideDeformedBuffer = HairGroupInstance.IsValid() ? HairGroupInstance->Guides.DeformedResource : nullptr;
 	if (GuideDeformedBuffer && GuideDeformedBuffer->IsInitialized())
 	{
 		const FRDGExternalBuffer& DeformedBuffer = GuideDeformedBuffer->GetBuffer(FHairStrandsDeformedResource::EFrameType::Current);
