6-Groom deformers extension

#virtualized

diff --git a/Engine/Plugins/Animation/DeformerGraph/Source/OptimusCore/Private/DataInterfaces/OptimusDataInterfaceRawBuffer.cpp b/Engine/Plugins/Animation/DeformerGraph/Source/OptimusCore/Private/DataInterfaces/OptimusDataInterfaceRawBuffer.cpp
--- a/Engine/Plugins/Animation/DeformerGraph/Source/OptimusCore/Private/DataInterfaces/OptimusDataInterfaceRawBuffer.cpp
+++ b/Engine/Plugins/Animation/DeformerGraph/Source/OptimusCore/Private/DataInterfaces/OptimusDataInterfaceRawBuffer.cpp
@@ -679,6 +679,9 @@ void FOptimusPersistentBufferDataProviderProxy::GatherDispatchData(FDispatchData
 		Parameters.BufferSize = InDispatchData.bUnifiedDispatch ? TotalElementCount : InvocationElementCounts[InvocationIndex];
 		Parameters.BufferUAV = BufferUAVs[InDispatchData.GraphKernelIndex];
 
-		StartOffset += InvocationElementCounts[InvocationIndex];
+		if(!InDispatchData.bUnifiedDispatch)
+		{
+			StartOffset += InvocationElementCounts[InvocationIndex];
+		}
 	}
 }
diff --git a/Engine/Plugins/Experimental/Dataflow/Source/DataflowEditor/Private/Dataflow/DataflowSimulationScene.cpp b/Engine/Plugins/Experimental/Dataflow/Source/DataflowEditor/Private/Dataflow/DataflowSimulationScene.cpp
--- a/Engine/Plugins/Experimental/Dataflow/Source/DataflowEditor/Private/Dataflow/DataflowSimulationScene.cpp
+++ b/Engine/Plugins/Experimental/Dataflow/Source/DataflowEditor/Private/Dataflow/DataflowSimulationScene.cpp
@@ -253,7 +253,7 @@ void FDataflowSimulationScene::UpdateSimulationCache(float& DeltaSeconds)
 			DeltaSeconds = 1.0f / SceneDescription->CacheParams.FrameRate;
 			SimulationTime += DeltaSeconds;
 
-			if(SimulationTime >= TimeRange[1])
+			if(SimulationTime >= TimeRange[1]+UE_SMALL_NUMBER)
 			{
 				bIsRecordingCache = false;
 				CacheManager->EndEvaluate();
diff --git a/Engine/Plugins/Experimental/Dataflow/Source/DataflowNodes/Private/Dataflow/DataflowCollectionAddScalarVertexPropertyNode.cpp b/Engine/Plugins/Experimental/Dataflow/Source/DataflowNodes/Private/Dataflow/DataflowCollectionAddScalarVertexPropertyNode.cpp
--- a/Engine/Plugins/Experimental/Dataflow/Source/DataflowNodes/Private/Dataflow/DataflowCollectionAddScalarVertexPropertyNode.cpp
+++ b/Engine/Plugins/Experimental/Dataflow/Source/DataflowNodes/Private/Dataflow/DataflowCollectionAddScalarVertexPropertyNode.cpp
@@ -57,6 +57,35 @@ TArray<UE::Dataflow::FRenderingParameter> FDataflowAddScalarVertexPropertyCallba
 	return UniqueParameters;
 }
 
+TArray<UE::Dataflow::FRenderingParameter> FDataflowAddScalarVertexPropertyCallbackRegistry::GetRenderingParameters(const FName& TargetGroup) const
+{
+	TArray<UE::Dataflow::FRenderingParameter> UniqueParameters;
+
+	for (const TPair<FName, TUniquePtr<IDataflowAddScalarVertexPropertyCallbacks>>& CallbacksEntry : AllCallbacks)
+	{
+		const TArray<UE::Dataflow::FRenderingParameter> RenderingParameters = CallbacksEntry.Value->GetRenderingParameters();
+		const TArray<FName> TargetGroups = CallbacksEntry.Value->GetTargetGroupNames();
+		if(RenderingParameters.Num() == TargetGroups.Num())
+		{
+			for(int32 TargetIndex = 0, NumTargets = TargetGroups.Num(); TargetIndex < NumTargets; ++TargetIndex)
+			{
+				if(TargetGroups[TargetIndex] == TargetGroup)
+				{
+					UniqueParameters.AddUnique(RenderingParameters[TargetIndex]);
+				}
+			}
+		}
+		else
+		{
+			for (const UE::Dataflow::FRenderingParameter& RenderingParameter : CallbacksEntry.Value->GetRenderingParameters())
+			{
+				UniqueParameters.AddUnique(RenderingParameter);
+			}
+		}
+	}
+	return UniqueParameters;
+}
+
 
 //
 // Deprecated class
@@ -123,7 +152,7 @@ FDataflowCollectionAddScalarVertexPropertyNode::FDataflowCollectionAddScalarVert
 
 TArray<UE::Dataflow::FRenderingParameter> FDataflowCollectionAddScalarVertexPropertyNode::GetRenderParametersImpl() const
 {
-	return FDataflowAddScalarVertexPropertyCallbackRegistry::Get().GetRenderingParameters();
+	return FDataflowAddScalarVertexPropertyCallbackRegistry::Get().GetRenderingParameters(TargetGroup.Name);
 }
 
 
diff --git a/Engine/Plugins/Experimental/Dataflow/Source/DataflowNodes/Public/Dataflow/DataflowCollectionAddScalarVertexPropertyNode.h b/Engine/Plugins/Experimental/Dataflow/Source/DataflowNodes/Public/Dataflow/DataflowCollectionAddScalarVertexPropertyNode.h
--- a/Engine/Plugins/Experimental/Dataflow/Source/DataflowNodes/Public/Dataflow/DataflowCollectionAddScalarVertexPropertyNode.h
+++ b/Engine/Plugins/Experimental/Dataflow/Source/DataflowNodes/Public/Dataflow/DataflowCollectionAddScalarVertexPropertyNode.h
@@ -32,6 +32,8 @@ public:
 	DATAFLOWNODES_API TArray<FName> GetTargetGroupNames() const;
 
 	DATAFLOWNODES_API TArray<UE::Dataflow::FRenderingParameter> GetRenderingParameters() const;
+	
+	DATAFLOWNODES_API TArray<UE::Dataflow::FRenderingParameter> GetRenderingParameters(const FName& TargetGroup) const;
 
 private:
 
diff --git a/Engine/Plugins/Runtime/HairStrands/Config/BaseHairStrands.ini b/Engine/Plugins/Runtime/HairStrands/Config/BaseHairStrands.ini
--- a/Engine/Plugins/Runtime/HairStrands/Config/BaseHairStrands.ini
+++ b/Engine/Plugins/Runtime/HairStrands/Config/BaseHairStrands.ini
@@ -19,4 +19,8 @@ fx.UseShaderStages=1
 +ClassRedirects=(OldName="/Script/HairStrandsNiagara.NiagaraDataInterfacePressureGrid",NewName="/Script/HairStrandsCore.NiagaraDataInterfacePressureGrid")
 +PropertyRedirects=(OldName="GroomAsset.DeformedSkeletalMesh",NewName="GroomAsset.RiggedSkeletalMesh")
 +PropertyRedirects=(OldName="HairDeformationSettings.bEnableDeformation",NewName="HairDeformationSettings.bEnableRigging")
-+PropertyRedirects=(OldName="HairGroupsInterpolation.DeformationSettings",NewName="HairGroupsInterpolation.RiggingSettings")
\ No newline at end of file
++PropertyRedirects=(OldName="HairGroupsInterpolation.DeformationSettings",NewName="HairGroupsInterpolation.RiggingSettings")
++ClassRedirects=(OldName="OptimusGroomWriteDataInterface",NewName="/Script/HairStrandsDeformer.OptimusGroomStrandsWriteDataInterface")
++ClassRedirects=(OldName="OptimusGroomDataInterface",NewName="/Script/HairStrandsDeformer.OptimusGroomStrandsReadDataInterface")
++ClassRedirects=(OldName="OptimusGroomGuideDataInterface",NewName="/Script/HairStrandsDeformer.OptimusGroomGuidesReadDataInterface")
++ClassRedirects=(OldName="OptimusGroomComponentSource",NewName="/Script/HairStrandsDeformer.OptimusGroomAssetComponentSource")
\ No newline at end of file
diff --git a/Engine/Plugins/Runtime/HairStrands/HairStrands.uplugin b/Engine/Plugins/Runtime/HairStrands/HairStrands.uplugin
--- a/Engine/Plugins/Runtime/HairStrands/HairStrands.uplugin
+++ b/Engine/Plugins/Runtime/HairStrands/HairStrands.uplugin
@@ -71,6 +71,10 @@
 		{
 			"Name": "Dataflow",
 			"Enabled": true
+		},
+		{
+			"Name": "ComputeFramework",
+			"Enabled": true
 		}
 	]
 }
diff --git a/Engine/Plugins/Runtime/HairStrands/Shaders/Private/DeformerDataInterfaceGroom.ush b/Engine/Plugins/Runtime/HairStrands/Shaders/Private/DeformerDataInterfaceGroom.ush
deleted file mode 100644
diff --git a/Engine/Plugins/Runtime/HairStrands/Shaders/Private/DeformerDataInterfaceGroomExec.ush b/Engine/Plugins/Runtime/HairStrands/Shaders/Private/DeformerDataInterfaceGroomExec.ush
deleted file mode 100644
diff --git a/Engine/Plugins/Runtime/HairStrands/Shaders/Private/DeformerDataInterfaceGroomGuide.ush b/Engine/Plugins/Runtime/HairStrands/Shaders/Private/DeformerDataInterfaceGroomGuide.ush
deleted file mode 100644
diff --git a/Engine/Plugins/Runtime/HairStrands/Shaders/Private/DeformerDataInterfaceGroomWrite.ush b/Engine/Plugins/Runtime/HairStrands/Shaders/Private/DeformerDataInterfaceGroomWrite.ush
deleted file mode 100644
diff --git a/Engine/Plugins/Runtime/HairStrands/Shaders/Private/Deformers/DeformerGroomAttributeRead.ush b/Engine/Plugins/Runtime/HairStrands/Shaders/Private/Deformers/DeformerGroomAttributeRead.ush
new file mode 100644
--- /dev/null
+++ b/Engine/Plugins/Runtime/HairStrands/Shaders/Private/Deformers/DeformerGroomAttributeRead.ush
@@ -0,0 +1,147 @@
+// Copyright Epic Games, Inc. All Rights Reserved.
+
+#include "/Plugin/ComputeFramework/Private/ComputeKernelCommon.ush"
+
+uint {DataInterfaceName}_NumValues; 
+
+Buffer<float> {DataInterfaceName}_ValueTypedData;
+
+// Interface functions
+
+uint ReadNumValues_{DataInterfaceName}()
+{
+	return {DataInterfaceName}_NumValues;
+}
+
+uint ReadValueTypedUint_{DataInterfaceName}(uint ValueIndex)
+{
+	if (ValueIndex < {DataInterfaceName}_NumValues)
+	{
+		return (uint){DataInterfaceName}_ValueTypedData[ValueIndex];
+	}
+	return 0;
+}
+
+bool ReadValueTypedBool_{DataInterfaceName}(uint ValueIndex)
+{
+	if (ValueIndex < {DataInterfaceName}_NumValues)
+	{
+		return ({DataInterfaceName}_ValueTypedData[ValueIndex] == 1.0f) ? true : false;
+	}
+	return false;
+}
+
+int ReadValueTypedInt_{DataInterfaceName}(uint ValueIndex)
+{
+	if (ValueIndex < {DataInterfaceName}_NumValues)
+	{
+		return (int){DataInterfaceName}_ValueTypedData[ValueIndex];
+	}
+	return 0;
+}
+
+int2 ReadValueTypedInt2_{DataInterfaceName}(uint ValueIndex)
+{
+	if (ValueIndex < {DataInterfaceName}_NumValues)
+	{
+		const int ValueOffset = 2 * ValueIndex;
+		return int2({DataInterfaceName}_ValueTypedData[ValueOffset], {DataInterfaceName}_ValueTypedData[ValueOffset+1]);
+	}
+	return int2(0,0);
+}
+
+int3 ReadValueTypedInt3_{DataInterfaceName}(uint ValueIndex)
+{
+	if (ValueIndex < {DataInterfaceName}_NumValues)
+	{
+		const int ValueOffset = 3 * ValueIndex;
+		return int3({DataInterfaceName}_ValueTypedData[ValueOffset], {DataInterfaceName}_ValueTypedData[ValueOffset+1], {DataInterfaceName}_ValueTypedData[ValueOffset+2]);
+	}
+	return int3(0,0,0);
+}
+
+int4 ReadValueTypedInt4_{DataInterfaceName}(uint ValueIndex)
+{
+	if (ValueIndex < {DataInterfaceName}_NumValues)
+	{
+		const int ValueOffset = 4 * ValueIndex;
+		return int4({DataInterfaceName}_ValueTypedData[ValueOffset], {DataInterfaceName}_ValueTypedData[ValueOffset+1],
+					{DataInterfaceName}_ValueTypedData[ValueOffset+2], {DataInterfaceName}_ValueTypedData[ValueOffset+3]);
+	}
+	return int4(0,0,0,0);
+}
+
+float ReadValueTypedFloat_{DataInterfaceName}(uint ValueIndex)
+{
+	if (ValueIndex < {DataInterfaceName}_NumValues)
+	{
+		return {DataInterfaceName}_ValueTypedData[ValueIndex];
+	}
+	return 0.0f;
+}
+
+float2 ReadValueTypedFloat2_{DataInterfaceName}(uint ValueIndex)
+{
+	if (ValueIndex < {DataInterfaceName}_NumValues)
+	{
+		const int ValueOffset = 2 * ValueIndex;
+		return float2({DataInterfaceName}_ValueTypedData[ValueOffset], {DataInterfaceName}_ValueTypedData[ValueOffset+1]);
+	}
+	return float2(0.0f,0.0f);
+}
+
+float3 ReadValueTypedFloat3_{DataInterfaceName}(uint ValueIndex)
+{
+	if (ValueIndex < {DataInterfaceName}_NumValues)
+	{
+		const int ValueOffset = 3 * ValueIndex;
+		return float3({DataInterfaceName}_ValueTypedData[ValueOffset], {DataInterfaceName}_ValueTypedData[ValueOffset+1], {DataInterfaceName}_ValueTypedData[ValueOffset+2]);
+	}
+	return float3(0.0f,0.0f,0.0f);
+}
+
+float4 ReadValueTypedFloat4_{DataInterfaceName}(uint ValueIndex)
+{
+	if (ValueIndex < {DataInterfaceName}_NumValues)
+	{
+		const int ValueOffset = 4 * ValueIndex;
+		return float4({DataInterfaceName}_ValueTypedData[ValueOffset], {DataInterfaceName}_ValueTypedData[ValueOffset+1],
+					{DataInterfaceName}_ValueTypedData[ValueOffset+2], {DataInterfaceName}_ValueTypedData[ValueOffset+3]);
+	}
+	return float4(0.0f,0.0f,0.0f,0.0f);
+}
+
+float4x4 ReadValueTypedFloat4x4_{DataInterfaceName}(uint ValueIndex)
+{
+	if (ValueIndex < {DataInterfaceName}_NumValues)
+	{
+		const int ValueOffset = 16 * ValueIndex;
+		return float4x4({DataInterfaceName}_ValueTypedData[ValueOffset], {DataInterfaceName}_ValueTypedData[ValueOffset+1],
+						{DataInterfaceName}_ValueTypedData[ValueOffset+2], {DataInterfaceName}_ValueTypedData[ValueOffset+3],
+						{DataInterfaceName}_ValueTypedData[ValueOffset+4], {DataInterfaceName}_ValueTypedData[ValueOffset+5],
+						{DataInterfaceName}_ValueTypedData[ValueOffset+6], {DataInterfaceName}_ValueTypedData[ValueOffset+7],
+						{DataInterfaceName}_ValueTypedData[ValueOffset+8], {DataInterfaceName}_ValueTypedData[ValueOffset+9],
+						{DataInterfaceName}_ValueTypedData[ValueOffset+10], {DataInterfaceName}_ValueTypedData[ValueOffset+11],
+						{DataInterfaceName}_ValueTypedData[ValueOffset+12], {DataInterfaceName}_ValueTypedData[ValueOffset+13],
+						{DataInterfaceName}_ValueTypedData[ValueOffset+14], {DataInterfaceName}_ValueTypedData[ValueOffset+15]);
+	}
+	return float4x4(1.0f,0.0f,0.0f,0.0f,0.0f,1.0f,0.0f,0.0f,0.0f,0.0f,1.0f,0.0f,0.0f,0.0f,0.0f,1.0f);
+}
+
+float3x4 ReadValueTypedFloat4x4_{DataInterfaceName}(uint ValueIndex)
+{
+	if (ValueIndex < {DataInterfaceName}_NumValues)
+	{
+		const int ValueOffset = 12 * ValueIndex;
+		return float4x4({DataInterfaceName}_ValueTypedData[ValueOffset], {DataInterfaceName}_ValueTypedData[ValueOffset+1],
+						{DataInterfaceName}_ValueTypedData[ValueOffset+2], {DataInterfaceName}_ValueTypedData[ValueOffset+3],
+						{DataInterfaceName}_ValueTypedData[ValueOffset+4], {DataInterfaceName}_ValueTypedData[ValueOffset+5],
+						{DataInterfaceName}_ValueTypedData[ValueOffset+6], {DataInterfaceName}_ValueTypedData[ValueOffset+7],
+						{DataInterfaceName}_ValueTypedData[ValueOffset+8], {DataInterfaceName}_ValueTypedData[ValueOffset+9],
+						{DataInterfaceName}_ValueTypedData[ValueOffset+10], {DataInterfaceName}_ValueTypedData[ValueOffset+11]);
+	}
+	return float3x4(1.0f,0.0f,0.0f,0.0f,0.0f,1.0f,0.0f,0.0f,0.0f,0.0f,1.0f,0.0f);
+}
+
+
+
diff --git a/Engine/Plugins/Runtime/HairStrands/Shaders/Private/Deformers/DeformerGroomDomainsExec.ush b/Engine/Plugins/Runtime/HairStrands/Shaders/Private/Deformers/DeformerGroomDomainsExec.ush
new file mode 100644
--- /dev/null
+++ b/Engine/Plugins/Runtime/HairStrands/Shaders/Private/Deformers/DeformerGroomDomainsExec.ush
@@ -0,0 +1,10 @@
+// Copyright Epic Games, Inc. All Rights Reserved.
+
+#include "/Plugin/ComputeFramework/Private/ComputeKernelCommon.ush"
+
+int3 {DataInterfaceName}_NumThreads;
+
+int3 ReadNumThreads_{DataInterfaceName}()
+{
+	return {DataInterfaceName}_NumThreads;
+}
diff --git a/Engine/Plugins/Runtime/HairStrands/Shaders/Private/Deformers/DeformerGroomGuidesRead.ush b/Engine/Plugins/Runtime/HairStrands/Shaders/Private/Deformers/DeformerGroomGuidesRead.ush
new file mode 100644
--- /dev/null
+++ b/Engine/Plugins/Runtime/HairStrands/Shaders/Private/Deformers/DeformerGroomGuidesRead.ush
@@ -0,0 +1,138 @@
+// Copyright Epic Games, Inc. All Rights Reserved.
+
+#include "/Plugin/ComputeFramework/Private/ComputeKernelCommon.ush"
+#include "/Engine/Private/HairStrands/HairStrandsVertexFactoryCommon.ush"
+#include "/Engine/Private/HairStrands/HairStrandsAttributeCommon.ush"
+#include "/Engine/Private/HairStrands/HairStrandsBindingCommon.ush"
+
+uint {DataInterfaceName}_PointCount; 
+uint {DataInterfaceName}_CurveCount;
+uint {DataInterfaceName}_ObjectIndex;
+float3 {DataInterfaceName}_RestPositionOffset;
+
+ByteAddressBuffer {DataInterfaceName}_PointRestPositions;
+ByteAddressBuffer {DataInterfaceName}_CurvePointOffsets;
+ByteAddressBuffer {DataInterfaceName}_PointCurveIndices;
+
+float4x4 {DataInterfaceName}_ObjectRestTransform;
+float4x4 {DataInterfaceName}_ObjectDeformedTransform;
+
+float3x4 BuildRigidTransform(float4x4 RigidTransform)
+{
+	return float3x4(RigidTransform[0], RigidTransform[1], RigidTransform[2]);
+}
+ 
+#if ENABLE_SKINNED_TRANSFORM == 1
+Buffer<float4> {DataInterfaceName}_TriangleRestPositions;
+Buffer<float4> {DataInterfaceName}_TriangleDeformedPositions;
+Buffer<uint> {DataInterfaceName}_CurveBarycentricCoordinates;
+Buffer<uint> {DataInterfaceName}_CurveTriangleIndices;
+
+float3x4 BuildSkinnedTransform(uint CurveIndex, Buffer<float4> TrianglePositions,
+	Buffer<uint> BarycentricCoordinates, Buffer<uint> TriangleIndices)
+{
+	const float3 BarycentricCoord = UnpackBarycentrics(BarycentricCoordinates[CurveIndex]);
+	const uint TriangleIndex = TriangleIndices[CurveIndex];
+
+	const FHairMeshBasis MeshBasis = GetHairMeshBasis(TriangleIndex, TrianglePositions, BarycentricCoord);
+	
+	return float3x4(MeshBasis.BaseT.x, MeshBasis.BaseB.x, MeshBasis.BaseN.x, MeshBasis.BaseO.x,
+					MeshBasis.BaseT.y, MeshBasis.BaseB.y, MeshBasis.BaseN.y, MeshBasis.BaseO.y,
+					MeshBasis.BaseT.z, MeshBasis.BaseB.z, MeshBasis.BaseN.z, MeshBasis.BaseO.z);
+}
+#endif
+
+// Interface functions
+
+uint ReadNumPoints_{DataInterfaceName}()
+{
+	return {DataInterfaceName}_PointCount;
+}
+
+uint ReadNumCurves_{DataInterfaceName}()
+{
+	return {DataInterfaceName}_CurveCount;
+}
+
+uint ReadObjectIndex_{DataInterfaceName}()
+{
+	return {DataInterfaceName}_ObjectIndex;
+}
+
+float3 ReadPointRestPosition_{DataInterfaceName}(uint PointIndex)
+{
+	if (PointIndex < {DataInterfaceName}_PointCount)
+	{
+		return ReadHairControlPoint({DataInterfaceName}_PointRestPositions, PointIndex, {DataInterfaceName}_RestPositionOffset).Position;
+	}
+	return float3(0,0,0);
+}
+
+float ReadCurvePointOffset_{DataInterfaceName}(uint CurveIndex)
+{
+	if (CurveIndex < {DataInterfaceName}_CurveCount)
+	{
+		return ReadHairCurve({DataInterfaceName}_CurvePointOffsets, CurveIndex).PointIndex;
+	}
+	return 0;
+}
+
+float ReadCurveNumPoints_{DataInterfaceName}(uint CurveIndex)
+{
+	if (CurveIndex < {DataInterfaceName}_CurveCount)
+	{
+		return ReadHairCurve({DataInterfaceName}_CurvePointOffsets, CurveIndex).PointCount;
+	}
+	return 0;
+}
+
+uint ReadPointCurveIndex_{DataInterfaceName}(uint PointIndex)
+{
+	if (PointIndex < {DataInterfaceName}_PointCount)
+	{
+		return ReadHairPointToCurveIndex({DataInterfaceName}_PointCurveIndices, PointIndex);
+	}
+	return 0;
+}
+
+float3x4 ReadCurveRestTransform_{DataInterfaceName}(uint CurveIndex)
+{
+	//if (CurveIndex < {DataInterfaceName}_CurveCount)
+	{
+#if ENABLE_SKINNED_TRANSFORM == 1
+		return BuildSkinnedTransform(CurveIndex, {DataInterfaceName}_TriangleRestPositions,
+			{DataInterfaceName}_CurveBarycentricCoordinates, {DataInterfaceName}_CurveTriangleIndices);
+#endif
+
+#if ENABLE_SKINNED_TRANSFORM == 0
+		return float3x4(1, 0, 0, 0, 0, 1, 0, 0, 0, 0, 1, 0);
+#endif
+	}
+	//return float3x4(1, 0, 0, 0, 0, 1, 0, 0, 0, 0, 1, 0);
+}
+
+float3x4 ReadCurveDeformedTransform_{DataInterfaceName}(uint CurveIndex)
+{
+	//if (CurveIndex < {DataInterfaceName}_CurveCount)
+	{
+#if ENABLE_SKINNED_TRANSFORM == 1
+		return BuildSkinnedTransform(CurveIndex, {DataInterfaceName}_TriangleDeformedPositions,
+			{DataInterfaceName}_CurveBarycentricCoordinates, {DataInterfaceName}_CurveTriangleIndices);
+#endif
+
+#if ENABLE_SKINNED_TRANSFORM == 0
+		return float3x4(1, 0, 0, 0, 0, 1, 0, 0, 0, 0, 1, 0);
+#endif
+	}
+	//return float3x4(1, 0, 0, 0, 0, 1, 0, 0, 0, 0, 1, 0);
+}
+
+float3x4 ReadObjectRestTransform_{DataInterfaceName}()
+{
+	return BuildRigidTransform({DataInterfaceName}_ObjectRestTransform);
+}
+
+float3x4 ReadObjectDeformedTransform_{DataInterfaceName}()
+{
+	return BuildRigidTransform({DataInterfaceName}_ObjectDeformedTransform);
+}
\ No newline at end of file
diff --git a/Engine/Plugins/Runtime/HairStrands/Shaders/Private/Deformers/DeformerGroomGuidesWrite.ush b/Engine/Plugins/Runtime/HairStrands/Shaders/Private/Deformers/DeformerGroomGuidesWrite.ush
new file mode 100644
--- /dev/null
+++ b/Engine/Plugins/Runtime/HairStrands/Shaders/Private/Deformers/DeformerGroomGuidesWrite.ush
@@ -0,0 +1,36 @@
+// Copyright Epic Games, Inc. All Rights Reserved.
+
+#include "/Plugin/ComputeFramework/Private/ComputeKernelCommon.ush"
+#include "/Engine/Private/HairStrands/HairStrandsVertexFactoryCommon.ush"
+
+uint {DataInterfaceName}_RegisteredIndex; 
+uint {DataInterfaceName}_PointCount; 
+uint {DataInterfaceName}_OutputStreamStart;
+
+StructuredBuffer<float4> {DataInterfaceName}_DeformedPositionOffset;
+ByteAddressBuffer {DataInterfaceName}_PointRestPositions;
+RWByteAddressBuffer {DataInterfaceName}_PointDeformedPositions;
+
+// Interface helpers
+
+float3 {DataInterfaceName}_GetDeformedPositionOffset() 
+{ 
+	return ReadSimPositionOffset({DataInterfaceName}_DeformedPositionOffset,{DataInterfaceName}_RegisteredIndex); 
+}
+
+// Interface functions
+
+uint ReadNumPoints_{DataInterfaceName}()
+{
+	return {DataInterfaceName}_PointCount;
+}
+
+void WritePointDeformedPosition_{DataInterfaceName}(uint PointIndex, float3 PointPosition)
+{
+	const uint WriteIndex = {DataInterfaceName}_OutputStreamStart + PointIndex;
+	if (WriteIndex < {DataInterfaceName}_PointCount)
+	{
+		const FPackedHairPosition PackedCP = ReadPackedHairPosition({DataInterfaceName}_PointRestPositions, WriteIndex);
+		WritePackedHairControlPointPosition({DataInterfaceName}_PointDeformedPositions, WriteIndex, PackedCP, PointPosition, {DataInterfaceName}_GetDeformedPositionOffset());
+	}
+}
diff --git a/Engine/Plugins/Runtime/HairStrands/Shaders/Private/Deformers/DeformerGroomMeshesRead.ush b/Engine/Plugins/Runtime/HairStrands/Shaders/Private/Deformers/DeformerGroomMeshesRead.ush
new file mode 100644
--- /dev/null
+++ b/Engine/Plugins/Runtime/HairStrands/Shaders/Private/Deformers/DeformerGroomMeshesRead.ush
@@ -0,0 +1,26 @@
+﻿// Copyright Epic Games, Inc. All Rights Reserved.
+
+#include "/Plugin/ComputeFramework/Private/ComputeKernelCommon.ush"
+
+uint {DataInterfaceName}_NumBones; 
+
+Buffer<float4> {DataInterfaceName}_BoneTransformMatrices;
+
+// Interface functions
+
+uint ReadNumBones_{DataInterfaceName}()
+{
+	return {DataInterfaceName}_NumBones;
+}
+
+float3x4 ReadBoneTransformMatrix_{DataInterfaceName}(uint BoneIndex)
+{
+	if (BoneIndex < {DataInterfaceName}_NumBones)
+	{
+		return float3x4({DataInterfaceName}_BoneTransformMatrices[3*BoneIndex], {DataInterfaceName}_BoneTransformMatrices[3*BoneIndex+1], {DataInterfaceName}_BoneTransformMatrices[3*BoneIndex+2]);
+	}
+	return float3x4(1.0f,0.0f,0.0f,0.0f,0.0f,1.0f,0.0f,0.0f,0.0f,0.0f,1.0f,0.0f);
+}
+
+
+
diff --git a/Engine/Plugins/Runtime/HairStrands/Shaders/Private/Deformers/DeformerGroomSolverRead.ush b/Engine/Plugins/Runtime/HairStrands/Shaders/Private/Deformers/DeformerGroomSolverRead.ush
new file mode 100644
--- /dev/null
+++ b/Engine/Plugins/Runtime/HairStrands/Shaders/Private/Deformers/DeformerGroomSolverRead.ush
@@ -0,0 +1,170 @@
+﻿// Copyright Epic Games, Inc. All Rights Reserved.
+
+#include "/Plugin/ComputeFramework/Private/ComputeKernelCommon.ush"
+
+uint {DataInterfaceName}_NumSolverPoints;
+uint {DataInterfaceName}_NumSolverCurves;
+uint {DataInterfaceName}_NumSolverObjects;
+
+uint {DataInterfaceName}_NumDynamicPoints;
+uint {DataInterfaceName}_NumDynamicCurves;
+
+uint {DataInterfaceName}_NumKinematicPoints;
+uint {DataInterfaceName}_NumKinematicCurves;
+
+Buffer<uint> {DataInterfaceName}_ObjectCurveOffsets;
+Buffer<uint> {DataInterfaceName}_ObjectPointOffsets;
+
+Buffer<uint> {DataInterfaceName}_ObjectNumCurves;
+Buffer<uint> {DataInterfaceName}_ObjectNumPoints;
+
+Buffer<int> {DataInterfaceName}_PointObjectIndices;
+Buffer<int> {DataInterfaceName}_CurveObjectIndices;
+
+Buffer<int> {DataInterfaceName}_DynamicPointIndices;
+Buffer<int> {DataInterfaceName}_DynamicCurveIndices;
+
+Buffer<int> {DataInterfaceName}_KinematicPointIndices;
+Buffer<int> {DataInterfaceName}_KinematicCurveIndices;
+
+Buffer<uint> {DataInterfaceName}_ObjectDistanceLods;
+
+// Interface functions
+
+uint ReadNumSolverPoints_{DataInterfaceName}()
+{
+	return {DataInterfaceName}_NumSolverPoints;
+}
+
+uint ReadNumSolverCurves_{DataInterfaceName}()
+{
+	return {DataInterfaceName}_NumSolverCurves;
+}
+
+uint ReadNumSolverObjects_{DataInterfaceName}()
+{
+	return {DataInterfaceName}_NumSolverObjects;
+}
+
+uint ReadNumDynamicPoints_{DataInterfaceName}()
+{
+	return {DataInterfaceName}_NumDynamicPoints;
+}
+
+uint ReadNumDynamicCurves_{DataInterfaceName}()
+{
+	return {DataInterfaceName}_NumDynamicCurves;
+}
+
+uint ReadNumKinematicPoints_{DataInterfaceName}()
+{
+	return {DataInterfaceName}_NumKinematicPoints;
+}
+
+uint ReadNumKinematicCurves_{DataInterfaceName}()
+{
+	return {DataInterfaceName}_NumKinematicCurves;
+}
+
+uint ReadObjectPointOffset_{DataInterfaceName}(uint ObjectIndex)
+{
+	if (ObjectIndex < ({DataInterfaceName}_NumSolverObjects+1))
+	{
+		return {DataInterfaceName}_ObjectPointOffsets[ObjectIndex];
+	}
+	return 0;
+}
+
+uint ReadObjectCurveOffset_{DataInterfaceName}(uint ObjectIndex)
+{
+	if (ObjectIndex < ({DataInterfaceName}_NumSolverObjects+1))
+	{
+		return {DataInterfaceName}_ObjectCurveOffsets[ObjectIndex];
+	}
+	return 0;
+}
+
+uint ReadObjectNumPoints_{DataInterfaceName}(uint ObjectIndex)
+{
+	if (ObjectIndex < {DataInterfaceName}_NumSolverObjects)
+	{
+		return {DataInterfaceName}_ObjectNumPoints[ObjectIndex];
+	}
+	return 0;
+}
+
+uint ReadObjectNumCurves_{DataInterfaceName}(uint ObjectIndex)
+{
+	if (ObjectIndex < {DataInterfaceName}_NumSolverObjects)
+	{
+		return {DataInterfaceName}_ObjectNumCurves[ObjectIndex];
+	}
+	return 0;
+}
+
+uint ReadObjectDistanceLod_{DataInterfaceName}(uint ObjectIndex)
+{
+	if (ObjectIndex < {DataInterfaceName}_NumSolverObjects)
+	{
+		return {DataInterfaceName}_ObjectDistanceLods[ObjectIndex];
+	}
+	return 0;
+}
+
+int ReadPointObjectIndex_{DataInterfaceName}(uint PointIndex)
+{
+	if (PointIndex < {DataInterfaceName}_NumSolverPoints)
+	{
+		return {DataInterfaceName}_PointObjectIndices[PointIndex];
+	}
+	return -1;
+}
+
+int ReadCurveObjectIndex_{DataInterfaceName}(uint CurveIndex)
+{
+	if (CurveIndex < {DataInterfaceName}_NumSolverCurves)
+	{
+		return {DataInterfaceName}_CurveObjectIndices[CurveIndex];
+	}
+	return -1;
+}
+
+int ReadDynamicPointIndex_{DataInterfaceName}(uint PointIndex)
+{
+	if (PointIndex < {DataInterfaceName}_NumDynamicPoints)
+	{
+		return {DataInterfaceName}_DynamicPointIndices[PointIndex];
+	}
+	return -1;
+}
+
+int ReadKinematicPointIndex_{DataInterfaceName}(uint PointIndex)
+{
+	if (PointIndex < {DataInterfaceName}_NumKinematicPoints)
+	{
+		return {DataInterfaceName}_KinematicPointIndices[PointIndex];
+	}
+	return -1;
+}
+
+int ReadDynamicCurveIndex_{DataInterfaceName}(uint CurveIndex)
+{
+	if (PointIndex < {DataInterfaceName}_NumDynamicCurves)
+	{
+		return {DataInterfaceName}_DynamicCurveIndices[PointIndex];
+	}
+	return -1;
+}
+
+int ReadKinematicCurveIndex_{DataInterfaceName}(uint CurveIndex)
+{
+	if (CurveIndex < {DataInterfaceName}_NumKinematicCurves)
+	{
+		return {DataInterfaceName}_KinematicCurveIndices[CurveIndex];
+	}
+	return -1;
+}
+
+
+
+
diff --git a/Engine/Plugins/Runtime/HairStrands/Shaders/Private/Deformers/DeformerGroomStrandsRead.ush b/Engine/Plugins/Runtime/HairStrands/Shaders/Private/Deformers/DeformerGroomStrandsRead.ush
new file mode 100644
--- /dev/null
+++ b/Engine/Plugins/Runtime/HairStrands/Shaders/Private/Deformers/DeformerGroomStrandsRead.ush
@@ -0,0 +1,211 @@
+// Copyright Epic Games, Inc. All Rights Reserved.
+
+#include "/Plugin/ComputeFramework/Private/ComputeKernelCommon.ush"
+#include "/Engine/Private/HairStrands/HairStrandsVertexFactoryCommon.ush"
+#include "/Engine/Private/HairStrands/HairStrandsAttributeCommon.ush"
+
+HAIR_STRANDS_INSTANCE_PARAMETERS({DataInterfaceName})
+
+#define HAIR_STRANDS_ATTRIBUTE_ACCESSORS(Name) {DataInterfaceName}_##Name
+#include "/Engine/Private/HairStrands/HairStrandsAttributeTemplate.ush"
+
+ByteAddressBuffer {DataInterfaceName}_InterpolationBuffer;
+
+#if ENABLE_SKINNED_TRANSFORM == 0
+float4x4 {DataInterfaceName}_RigidRestTransform;
+float4x4 {DataInterfaceName}_RigidDeformedTransform;
+
+float3x4 BuildRigidTransform(float4x4 RigidTransform)
+{
+	return float3x4(RigidTransform[0], RigidTransform[1], RigidTransform[2]);
+}
+#endif
+
+#if ENABLE_SKINNED_TRANSFORM == 1
+Buffer<float4> {DataInterfaceName}_TriangleRestPositions;
+Buffer<float4> {DataInterfaceName}_TriangleDeformedPositions;
+Buffer<uint> {DataInterfaceName}_CurveBarycentricCoordinates;
+Buffer<uint> {DataInterfaceName}_CurveTriangleIndices;
+
+float3x4 BuildSkinnedTransform(uint CurveIndex, Buffer<float4> TrianglePositions,
+	Buffer<uint> CurveBarycentricCoordinates, Buffer<uint> CurveTriangleIndices)
+{
+	const float3 BarycentricCoord = UnpackBarycentrics(CurveBarycentricCoordinates[CurveIndex]);
+	const uint TriangleIndex = CurveTriangleIndices[CurveIndex];
+
+	const FHairMeshBasis MeshBasis = GetHairMeshBasis(TriangleIndex, TrianglePositions, RootBarycentric);
+	
+	return float3x4(MeshBasis.BaseT.x, MeshBasis.BaseB.x, MeshBasis.BaseN.x, MeshBasis.BaseO.x,
+					MeshBasis.BaseT.y, MeshBasis.BaseB.y, MeshBasis.BaseN.y, MeshBasis.BaseO.y,
+					MeshBasis.BaseT.z, MeshBasis.BaseB.z, MeshBasis.BaseN.z, MeshBasis.BaseO.z);
+}
+#endif
+
+
+// Helper functions
+
+FHairControlPoint ReadPointData(uint PointIndex)
+{
+	const uint MaxControlPoints = {DataInterfaceName}_PointCount;
+	if (PointIndex >= MaxControlPoints)
+	{
+		return (FHairControlPoint)0;
+	}
+
+	// Hair deformer work in local space so no need for position offset
+	return ReadHairControlPoint(
+		{DataInterfaceName}_PositionBuffer,
+		PointIndex, 
+		{DataInterfaceName}_GetHairInstancePositionOffset(),
+		{DataInterfaceName}_Radius, 
+		{DataInterfaceName}_RootScale, 
+		{DataInterfaceName}_TipScale);
+}
+
+FHairCurve ReadCurveData(uint CurveIndex)
+{
+	const uint MaxCurve = {DataInterfaceName}_CurveCount;
+	if (CurveIndex >= MaxCurve)
+	{
+		return (FHairCurve)0;
+	}
+
+	return ReadHairCurve({DataInterfaceName}_CurveBuffer, CurveIndex);
+}
+
+// Interface functions
+
+uint ReadNumPoints_{DataInterfaceName}()
+{
+	return {DataInterfaceName}_PointCount;
+}
+
+uint ReadNumCurves_{DataInterfaceName}()
+{
+	return {DataInterfaceName}_CurveCount;
+}
+
+float3 ReadPointRestPosition_{DataInterfaceName}(uint PointIndex)
+{
+	return ReadPointData(PointIndex).Position;
+}
+
+float ReadPointRadius_{DataInterfaceName}(uint PointIndex)
+{
+	return ReadPointData(PointIndex).WorldRadius;
+}
+
+float ReadPointCoordU_{DataInterfaceName}(uint PointIndex)
+{
+	return ReadPointData(PointIndex).UCoord;
+}
+
+float ReadPointLength_{DataInterfaceName}(uint PointIndex)
+{
+	return GetHairStrandsDimensions(PointIndex, 0).x;
+}
+
+float2 ReadPointRootUV_{DataInterfaceName}(uint PointIndex)
+{
+	return GetHairStrandsRootUV(PointIndex);
+}
+
+float ReadPointSeed_{DataInterfaceName}(uint PointIndex)
+{
+	return GetHairStrandsSeed(PointIndex);
+}
+
+uint ReadPointClumpId_{DataInterfaceName}(uint PointIndex)
+{
+	return GetHairStrandsClumpID(PointIndex).x;
+}
+
+float3 ReadPointColor_{DataInterfaceName}(uint PointIndex)
+{
+	return GetHairStrandsColor(PointIndex, 0);
+}
+
+float ReadPointRoughness_{DataInterfaceName}(uint PointIndex)
+{
+	return GetHairStrandsRoughness(PointIndex, 0);
+}
+
+float ReadPointAO_{DataInterfaceName}(uint PointIndex)
+{
+	return GetHairStrandsAO(PointIndex, 0);
+}
+
+uint ReadCurvePointOffset_{DataInterfaceName}(uint CurveIndex)
+{
+	return ReadCurveData(CurveIndex).PointIndex;
+}
+
+uint ReadCurveNumPoints_{DataInterfaceName}(uint CurveIndex)
+{
+	return ReadCurveData(CurveIndex).PointCount;
+}
+
+uint ReadPointCurveIndex_{DataInterfaceName}(uint PointIndex)
+{
+	return ReadHairPointToCurveIndex({DataInterfaceName}_PointToCurveBuffer, PointIndex);;
+}
+
+int2 ReadPointGuideIndices_{DataInterfaceName}(uint PointIndex)
+{
+	// Find the rendering curve index
+	const uint CurveIndex = ReadHairPointToCurveIndex({DataInterfaceName}_PointToCurveBuffer, PointIndex);
+
+	// Find the curve and point interpolation data
+	const FGuideCurveData CurveData = UnpackGuideCurveData({DataInterfaceName}_CurveInterpolationBuffer, CurveIndex, {DataInterfaceName}_bSingleGuide);
+	const FGuidePointData PointData = UnpackGuidePointData({DataInterfaceName}_PointInterpolationBuffer, PointIndex, {DataInterfaceName}_bSingleGuide);
+
+	// Build the sim point index
+	const uint2 GuidePointIndices = CurveData.RootPointIndices + PointData.PointIndices;
+	return int2(GuidePointIndices.x, GuidePointIndices.y);
+}
+
+float2 ReadPointGuideWeights_{DataInterfaceName}(uint PointIndex)
+{
+	// Find the rendering curve index
+	const uint CurveIndex = ReadHairPointToCurveIndex({DataInterfaceName}_PointToCurveBuffer, PointIndex);
+
+	// Find the curve and point interpolation data
+	const FGuideCurveData CurveData = UnpackGuideCurveData({DataInterfaceName}_CurveInterpolationBuffer, CurveIndex, {DataInterfaceName}_bSingleGuide);
+
+	// Build the sim point index
+	return  CurveData.CurveWeights;
+}
+
+float3x4 ReadCurveRestTransform_{DataInterfaceName}(uint CurveIndex)
+{
+	if (CurveIndex < {DataInterfaceName}_NumCurves)
+	{
+#if ENABLE_SKINNED_TRANSFORM == 1
+		return BuildSkinnedTransform(CurveIndex, {DataInterfaceName}_TriangleRestPositions,
+			{DataInterfaceName}_CurveBarycentricCoordinates, {DataInterfaceName}_CurveTriangleIndices);
+#endif
+
+#if ENABLE_SKINNED_TRANSFORM == 0
+		return float3x4(1, 0, 0, 0, 0, 1, 0, 0, 0, 0, 1, 0);
+#endif
+	}
+	return float3x4(1, 0, 0, 0, 0, 1, 0, 0, 0, 0, 1, 0);
+}
+
+float3x4 ReadCurveDeformedTransform_{DataInterfaceName}(uint CurveIndex)
+{
+	if (CurveIndex < {DataInterfaceName}_CurveCount)
+	{
+#if ENABLE_SKINNED_TRANSFORM == 1
+		return BuildSkinnedTransform(CurveIndex, {DataInterfaceName}_TriangleDeformedPositions,
+			{DataInterfaceName}_CurveBarycentricCoordinates, {DataInterfaceName}_CurveTriangleIndices);
+#endif
+
+#if ENABLE_SKINNED_TRANSFORM == 0
+		return float3x4(1, 0, 0, 0, 0, 1, 0, 0, 0, 0, 1, 0);
+#endif
+	}
+	return float3x4(1, 0, 0, 0,
+					0, 1, 0, 0,
+					0, 0, 1, 0);
+}
diff --git a/Engine/Plugins/Runtime/HairStrands/Shaders/Private/Deformers/DeformerGroomStrandsWrite.ush b/Engine/Plugins/Runtime/HairStrands/Shaders/Private/Deformers/DeformerGroomStrandsWrite.ush
new file mode 100644
--- /dev/null
+++ b/Engine/Plugins/Runtime/HairStrands/Shaders/Private/Deformers/DeformerGroomStrandsWrite.ush
@@ -0,0 +1,159 @@
+// Copyright Epic Games, Inc. All Rights Reserved.
+
+#include "/Plugin/ComputeFramework/Private/ComputeKernelCommon.ush"
+#include "/Engine/Private/HairStrands/HairStrandsVertexFactoryCommon.ush"
+
+HAIR_STRANDS_INSTANCE_PARAMETERS({DataInterfaceName})
+
+// Note: resources are explicitly bound here (instance of using HAIR_STRANDS_INSTANCE_PARAMETERS's resources) 
+// because we write out data which would be bound as SRV otherwise (position, attributes, ...)
+uint 				{DataInterfaceName}_OutputStreamStart;
+StructuredBuffer<float4> {DataInterfaceName}_DeformedPositionOffset;
+ByteAddressBuffer		{DataInterfaceName}_PointRestPositions;
+RWByteAddressBuffer		{DataInterfaceName}_PointDeformedPositions;
+RWByteAddressBuffer 	{DataInterfaceName}_PointAttributeBuffer;
+RWByteAddressBuffer 	{DataInterfaceName}_PointAttributeBuffer;
+
+// Helper functions
+
+void InternalSetHairStrandsAttributes(uint Index, uint AttributeOffsetInBytes, uint AttributeStrideInBytes, RWByteAddressBuffer InAttributeBuffer, uint In, uint InIndexToChunkDivAsShift, uint InChunkElementCount, uint InChunkStrideInBytes)
+{
+	// Adapt index and offset
+	const uint ChunkIndex   = Index >> InIndexToChunkDivAsShift;
+	Index                  -= ChunkIndex * InChunkElementCount;
+	AttributeOffsetInBytes += ChunkIndex * InChunkStrideInBytes;
+
+	// Ensure the reading address is 4-byte aligned
+	const uint Address 					= AttributeOffsetInBytes + Index * AttributeStrideInBytes;
+	const uint Address_4BytesAligned 	= Address & (~0x3);
+	const uint ElementIndex 			= Address - Address_4BytesAligned;
+
+	// Create a mask where the 
+	const uint Mask 		  = (1u<<(AttributeStrideInBytes*8u))-1u;
+	const uint AttributeMask  = Mask << (ElementIndex*8);
+	const uint AttributeValue = In   << (ElementIndex*8);
+
+	// Clear previous value, and then write the value
+	InAttributeBuffer.InterlockedAnd(Address_4BytesAligned, ~AttributeMask);
+	InAttributeBuffer.InterlockedOr(Address_4BytesAligned, AttributeValue);
+}
+
+void InternalSetHairStrandsAttributes32(uint Index, uint AttributeOffsetInBytes, uint AttributeStrideInBytes, RWByteAddressBuffer InAttributeBuffer, uint In, uint InIndexToChunkDivAsShift, uint InChunkElementCount, uint InChunkStrideInBytes)
+{
+	// Adapt index and offset
+	const uint ChunkIndex   = Index >> InIndexToChunkDivAsShift;
+	Index                  -= ChunkIndex * InChunkElementCount;
+	AttributeOffsetInBytes += ChunkIndex * InChunkStrideInBytes;
+
+	// Ensure the reading address is 4-byte aligned
+	const uint Address = AttributeOffsetInBytes + Index * AttributeStrideInBytes;
+	const uint Address_4BytesAligned = Address & (~0x3);
+
+	// Clear previous value, and then write the value
+	InAttributeBuffer.Store(Address_4BytesAligned, In);
+}
+
+// Interface functions
+
+float3 {DataInterfaceName}_GetDeformedPositionOffset() 
+{ 
+	return ReadRenPositionOffset({DataInterfaceName}_DeformedPositionOffset,{DataInterfaceName}_RegisteredIndex); 
+}
+
+uint ReadNumPoints_{DataInterfaceName}()
+{
+	return {DataInterfaceName}_PointCount;
+}
+
+uint ReadNumCurves_{DataInterfaceName}()
+{
+	return {DataInterfaceName}_CurveCount;
+}
+
+void WritePointDeformedPosition_{DataInterfaceName}(uint PointIndex, float3 Position)
+{
+	const uint WriteIndex = {DataInterfaceName}_OutputStreamStart + PointIndex;
+	if (WriteIndex < {DataInterfaceName}_PointCount)
+	{
+		const FPackedHairPosition PackedCP = ReadPackedHairPosition({DataInterfaceName}_PointRestPositions, WriteIndex);
+		WritePackedHairControlPointPosition({DataInterfaceName}_PointDeformedPositions, WriteIndex, PackedCP, Position, {DataInterfaceName}_GetDeformedPositionOffset());
+	}
+}
+
+void WritePointRadius_{DataInterfaceName}(uint PointIndex, float Radius)
+{
+	const uint WriteIndex = {DataInterfaceName}_OutputStreamStart + PointIndex;
+	if (WriteIndex < {DataInterfaceName}_PointCount)
+	{
+		const FPackedHairPosition PackedCP = ReadPackedHairPosition({DataInterfaceName}_PointRestPositions, WriteIndex);
+		const FPackedHairPosition NewPackedCP = PackHairControlPointRadius(PackedCP, Radius, {DataInterfaceName}_Radius);
+		WritePackedHairPosition({DataInterfaceName}_PointDeformedPositions, WriteIndex, NewPackedCP);
+	}
+}
+
+void WritePointPositionRadius_{DataInterfaceName}(uint PointIndex, float4 PositionAndRadius)
+{
+	const uint WriteIndex = {DataInterfaceName}_OutputStreamStart + PointIndex;
+	if (WriteIndex < {DataInterfaceName}_PointCount)
+	{
+		const FPackedHairPosition PackedCP = ReadPackedHairPosition({DataInterfaceName}_PointRestPositions, WriteIndex);
+		FPackedHairPosition NewPackedCP = 0;
+		NewPackedCP = PackHairControlPointPosition(PackedCP, PositionAndRadius.xyz, {DataInterfaceName}_GetDeformedPositionOffset());
+		NewPackedCP = PackHairControlPointRadius(NewPackedCP, PositionAndRadius.w, {DataInterfaceName}_Radius);
+		WritePackedHairPosition({DataInterfaceName}_PointDeformedPositions, WriteIndex, NewPackedCP);
+	}
+}
+
+void WriteCurveRootUV_{DataInterfaceName}(uint CurveIndex, float2 RootUV)
+{
+	const uint WriteIndex = {DataInterfaceName}_OutputStreamStart + CurveIndex;
+	if (WriteIndex < {DataInterfaceName}_CurveCount && HasHairStrandsRootUV())
+	{
+		InternalSetHairStrandsAttributes32(CurveIndex, HAIR_CURVE_ATTRIBUTE_OFFSET_ROOTUV({DataInterfaceName}_CurveAttributeOffsets), HAIR_CURVE_ATTRIBUTE_STRIDE_ROOTUV, {DataInterfaceName}_PointAttributeBuffer, PackHairRootUV(RootUV), {DataInterfaceName}_CurveAttributeIndexToChunkDivAsShift, {DataInterfaceName}_CurveAttributeChunkElementCount), {DataInterfaceName}_CurveAttributeChunkStrideInBytes);
+	}
+}
+
+void WriteCurveSeed_{DataInterfaceName}(uint CurveIndex, float Seed)
+{
+	const uint WriteIndex = {DataInterfaceName}_OutputStreamStart + CurveIndex;
+	if (WriteIndex < {DataInterfaceName}_CurveCount)
+	{
+		InternalSetHairStrandsAttributes(CurveIndex, HAIR_CURVE_ATTRIBUTE_OFFSET_SEED({DataInterfaceName}_CurveAttributeOffsets), HAIR_CURVE_ATTRIBUTE_STRIDE_SEED, {DataInterfaceName}_PointAttributeBuffer, PackHairSeed(Seed), {DataInterfaceName}_CurveAttributeIndexToChunkDivAsShift, {DataInterfaceName}_CurveAttributeChunkElementCount), {DataInterfaceName}_CurveAttributeChunkStrideInBytes);
+	}
+}
+
+void WriteCurveClumpId_{DataInterfaceName}(uint CurveIndex, uint ClumpId)
+{
+	const uint WriteIndex = {DataInterfaceName}_OutputStreamStart + CurveIndex;
+	if (WriteIndex < {DataInterfaceName}_CurveCount && HasHairStrandsClumpId())
+	{
+		InternalSetHairStrandsAttributes(CurveIndex, HAIR_CURVE_ATTRIBUTE_OFFSET_CLUMPID({DataInterfaceName}_CurveAttributeOffsets), HAIR_CURVE_ATTRIBUTE_STRIDE_CLUMPID, {DataInterfaceName}_PointAttributeBuffer, PackHairClumpID(ClumpId), {DataInterfaceName}_CurveAttributeIndexToChunkDivAsShift, {DataInterfaceName}_CurveAttributeChunkElementCount), {DataInterfaceName}_CurveAttributeChunkStrideInBytes);
+	}
+}
+
+void WritePointColor_{DataInterfaceName}(uint PointIndex, float3 Color)
+{
+	const uint WriteIndex = {DataInterfaceName}_OutputStreamStart + PointIndex;
+	if (WriteIndex < {DataInterfaceName}_PointCount && HasHairStrandsColor())
+	{
+		InternalSetHairStrandsAttributes32(PointIndex, HAIR_POINT_ATTRIBUTE_OFFSET_COLOR({DataInterfaceName}_PointAttributeOffsets), HAIR_POINT_ATTRIBUTE_STRIDE_COLOR, {DataInterfaceName}_PointAttributeBuffer, PackHairColor(Color), {DataInterfaceName}_PointAttributeIndexToChunkDivAsShift, {DataInterfaceName}_PointAttributeChunkElementCount), {DataInterfaceName}_PointAttributeChunkStrideInBytes);
+	}
+}
+
+void WritePointRoughness_{DataInterfaceName}(uint PointIndex, float Roughness)
+{
+	const uint WriteIndex = {DataInterfaceName}_OutputStreamStart + PointIndex;
+	if (WriteIndex < {DataInterfaceName}_PointCount && HasHairStrandsRoughness())
+	{
+		InternalSetHairStrandsAttributes(PointIndex, HAIR_POINT_ATTRIBUTE_OFFSET_ROUGHNESS({DataInterfaceName}_PointAttributeOffsets), HAIR_POINT_ATTRIBUTE_STRIDE_ROUGHNESS, {DataInterfaceName}_PointAttributeBuffer, PackHairRoughness(Roughness), {DataInterfaceName}_PointAttributeIndexToChunkDivAsShift, {DataInterfaceName}_PointAttributeChunkElementCount), {DataInterfaceName}_PointAttributeChunkStrideInBytes);
+	}
+}
+
+void WritePointAO_{DataInterfaceName}(uint PointIndex, float Ao)
+{
+	const uint WriteIndex = {DataInterfaceName}_OutputStreamStart + PointIndex;
+	if (WriteIndex < {DataInterfaceName}_PointCount && HasHairStrandsAO())
+	{
+		InternalSetHairStrandsAttributes(PointIndex, HAIR_POINT_ATTRIBUTE_OFFSET_AO({DataInterfaceName}_PointAttributeOffsets), HAIR_POINT_ATTRIBUTE_STRIDE_AO, {DataInterfaceName}_PointAttributeBuffer, PackHairAO(Ao), {DataInterfaceName}_PointAttributeIndexToChunkDivAsShift, {DataInterfaceName}_PointAttributeChunkElementCount), {DataInterfaceName}_PointAttributeChunkStrideInBytes);
+	}
+}
\ No newline at end of file
diff --git a/Engine/Plugins/Runtime/HairStrands/Source/HairStrandsCore/Private/GroomBuilder.cpp b/Engine/Plugins/Runtime/HairStrands/Source/HairStrandsCore/Private/GroomBuilder.cpp
--- a/Engine/Plugins/Runtime/HairStrands/Source/HairStrandsCore/Private/GroomBuilder.cpp
+++ b/Engine/Plugins/Runtime/HairStrands/Source/HairStrandsCore/Private/GroomBuilder.cpp
@@ -2153,7 +2153,7 @@ bool FGroomBuilder::BuildHairDescriptionGroups(const FHairDescription& HairDescr
 			continue;
 		}
 
-		const bool bAddEndingControlPoint = NeedsToAddEndingControlPoint(CurveNumVertices) && bAllowAddEndControlPoint;
+		const bool bAddEndingControlPoint = NeedsToAddEndingControlPoint(CurveNumVertices) && bAllowAddEndControlPoint && !bIsGuide;
 		CurrentHairStrandsDatas->StrandsCurves.CurvesCount.Add(FMath::Min(uint32(CurveNumVertices + (bAddEndingControlPoint ? 1u : 0u)), HAIR_MAX_NUM_POINT_PER_CURVE));
 
 		if (bCanUseClosestGuidesAndWeights)
diff --git a/Engine/Plugins/Runtime/HairStrands/Source/HairStrandsCore/Private/GroomComponent.cpp b/Engine/Plugins/Runtime/HairStrands/Source/HairStrandsCore/Private/GroomComponent.cpp
--- a/Engine/Plugins/Runtime/HairStrands/Source/HairStrandsCore/Private/GroomComponent.cpp
+++ b/Engine/Plugins/Runtime/HairStrands/Source/HairStrandsCore/Private/GroomComponent.cpp
@@ -2848,6 +2848,7 @@ void UGroomComponent::InitResources(bool bIsBindingReloading)
 		HairGroupInstance->BindingType = EHairBindingType::NoneBinding;
 
 		HairGroupInstance->Debug.SkinningCurrentLocalToWorld	= RegisteredMeshComponent ? RegisteredMeshComponent->GetComponentTransform() : FTransform::Identity;
+		HairGroupInstance->Debug.SkinningRestLocalToWorld		= HairGroupInstance->Debug.SkinningCurrentLocalToWorld;
 		HairGroupInstance->Debug.SkinningPreviousLocalToWorld	= HairGroupInstance->Debug.SkinningCurrentLocalToWorld;
 		HairGroupInstance->Debug.RigidCurrentLocalToWorld		= GetComponentTransform();
 		HairGroupInstance->Debug.RigidRestLocalToWorld			= HairGroupInstance->Debug.RigidCurrentLocalToWorld;
diff --git a/Engine/Plugins/Runtime/HairStrands/Source/HairStrandsCore/Public/GroomAssetDataflow.h b/Engine/Plugins/Runtime/HairStrands/Source/HairStrandsCore/Public/GroomAssetDataflow.h
--- a/Engine/Plugins/Runtime/HairStrands/Source/HairStrandsCore/Public/GroomAssetDataflow.h
+++ b/Engine/Plugins/Runtime/HairStrands/Source/HairStrandsCore/Public/GroomAssetDataflow.h
@@ -63,7 +63,10 @@ struct HAIRSTRANDSCORE_API FGroomDataflowSettings
 	const FString& GetDataflowTerminal() const { return DataflowTerminal; }
 
 	/** Return the skeletal mesh associated to this groom asset if any. */
-	USkeletalMesh* GetSkeletalMesh() const { return SkeletalMesh; }
+	USkeletalMesh* GetSkeletalMesh(const int32 GroupIndex) const { return !SkeletalMeshes.IsValidIndex(GroupIndex) ? nullptr : SkeletalMeshes[GroupIndex]; }
+
+	/** Return the mesh LOD used to transfer the skinning */
+	int32 GetMeshLOD(const int32 GroupIndex) const { return !MeshLODs.IsValidIndex(GroupIndex) ? INDEX_NONE : MeshLODs[GroupIndex]; }
 	
 	/** Return the groom asset dataflow settings rest collection . */
 	const FManagedArrayCollection* GetRestCollection() const { return RestCollection.Get(); }
@@ -73,9 +76,19 @@ struct HAIRSTRANDSCORE_API FGroomDataflowSettings
 	
 	/** Set the Dataflow terminal associated to this groom asset if any. */
 	void SetDataflowTerminal(const FString& InDataflowTerminal) { DataflowTerminal = InDataflowTerminal; }
+	
+	/** Init the skeletal mesh associated to this groom asset if any. */
+	void InitSkeletalMeshes(const int32 NumGroups) { SkeletalMeshes.Init(nullptr, NumGroups); MeshLODs.Init(INDEX_NONE, NumGroups); }
 
-	/** Set the skeletal mesh associated to this groom asset if any. */
-	void SetSkeletalMesh(USkeletalMesh* InSkeletalMesh) { SkeletalMesh = InSkeletalMesh; }
+	/** Set a group skeletal mesh and its associated LOD to this groom asset */
+	void SetSkeletalMesh(const int32 GroupIndex, USkeletalMesh* SkeletalMesh, const int32 MeshLOD)
+	{
+		if(SkeletalMeshes.IsValidIndex(GroupIndex) && MeshLODs.IsValidIndex(GroupIndex))
+		{
+			SkeletalMeshes[GroupIndex] = SkeletalMesh;
+			MeshLODs[GroupIndex] = MeshLOD;
+		}
+	}
 	
 	/** Set the rest collection onto the groom asset dataflow settings */
 	void SetRestCollection(FManagedArrayCollection* InRestCollection) { RestCollection = TSharedPtr<FManagedArrayCollection, ESPMode::ThreadSafe>(InRestCollection); }
@@ -111,7 +124,11 @@ private :
 
     /** Hair geometry that could describe external cards and meshes. */
     UPROPERTY()
-    TObjectPtr<USkeletalMesh> SkeletalMesh;
+    TArray<TObjectPtr<USkeletalMesh>> SkeletalMeshes;
+
+	/** LOD indices of the hair geometry used to transfer the skinnign weights */
+	UPROPERTY()
+	TArray<int32> MeshLODs;
 
 #if WITH_EDITORONLY_DATA
 	
diff --git a/Engine/Plugins/Runtime/HairStrands/Source/HairStrandsCore/Public/GroomInstance.h b/Engine/Plugins/Runtime/HairStrands/Source/HairStrandsCore/Public/GroomInstance.h
--- a/Engine/Plugins/Runtime/HairStrands/Source/HairStrandsCore/Public/GroomInstance.h
+++ b/Engine/Plugins/Runtime/HairStrands/Source/HairStrandsCore/Public/GroomInstance.h
@@ -264,6 +264,7 @@ struct HAIRSTRANDSCORE_API FHairGroupInstance : public FHairStrandsInstance
 		const UGroomComponent*	GroomComponentForDebug = nullptr; // For debug only, shouldn't be deferred on the rendering thread
 		FTransform				RigidRestLocalToWorld = FTransform::Identity;
 		FTransform				RigidCurrentLocalToWorld = FTransform::Identity;
+		FTransform				SkinningRestLocalToWorld = FTransform::Identity;
 		FTransform				SkinningCurrentLocalToWorld = FTransform::Identity;
 		FTransform				RigidPreviousLocalToWorld = FTransform::Identity;
 		FTransform				SkinningPreviousLocalToWorld = FTransform::Identity;
@@ -298,6 +299,13 @@ struct HAIRSTRANDSCORE_API FHairGroupInstance : public FHairStrandsInstance
 	virtual const FHairGroupPublicData* GetHairData() const { return HairGroupPublicData; }
 	virtual const EHairGeometryType GetHairGeometry() const { return GeometryType; }
 
+	/** Get the rest local to world transform according to the internal binding type */
+	FORCEINLINE const FTransform& GetRestLocalToWorld() const
+	{
+		return (BindingType == EHairBindingType::Skinning) ? Debug.SkinningRestLocalToWorld: 
+															 Debug.RigidRestLocalToWorld;
+	}
+
 	/** Get the current local to world transform according to the internal binding type */
 	FORCEINLINE const FTransform& GetCurrentLocalToWorld() const
 	{
diff --git a/Engine/Plugins/Runtime/HairStrands/Source/HairStrandsDataflow/Private/AttachGuidesRootsNode.cpp b/Engine/Plugins/Runtime/HairStrands/Source/HairStrandsDataflow/Private/AttachGuidesRootsNode.cpp
new file mode 100644
--- /dev/null
+++ b/Engine/Plugins/Runtime/HairStrands/Source/HairStrandsDataflow/Private/AttachGuidesRootsNode.cpp
@@ -0,0 +1,45 @@
+// Copyright Epic Games, Inc. All Rights Reserved.
+
+#include "AttachGuidesRootsNode.h"
+#include "GroomCollectionFacades.h"
+#include "Dataflow/DataflowObjectInterface.h"
+
+#include UE_INLINE_GENERATED_CPP_BY_NAME(AttachGuidesRootsNode)
+
+void FAttachGuidesRootsDataflowNode::Evaluate(UE::Dataflow::FContext& Context, const FDataflowOutput* Out) const
+{
+	if (Out->IsA<FManagedArrayCollection>(&Collection))
+	{
+		FManagedArrayCollection GroomCollection = GetValue<FManagedArrayCollection>(Context, &Collection);
+		
+		UE::Groom::FGroomGuidesFacade GuidesFacade(GroomCollection);
+
+		if(GuidesFacade.IsValid())
+		{
+			TArray<float> KinematicWeights = GuidesFacade.GetVertexKinematicWeights();
+			
+			const int32 LocalIndex = GroupIndex;
+			ParallelFor(GuidesFacade.GetNumPoints(), [&KinematicWeights, &GuidesFacade, LocalIndex](int32 PointIndex)
+			{
+				const int32 CurveIndex = GuidesFacade.GetPointCurveIndices()[PointIndex];
+				const int32 ObjectIndex = GuidesFacade.GetCurveObjectIndices()[CurveIndex];
+
+				if(LocalIndex == INDEX_NONE || ObjectIndex == LocalIndex)
+				{
+					const int32 PointOffset = (CurveIndex == 0) ? 0 : GuidesFacade.GetCurvePointOffsets()[CurveIndex-1];
+					const int32 LocalPoint = PointIndex - PointOffset;
+
+					if(LocalPoint < 2)
+					{
+						KinematicWeights[2*PointIndex] = 1.0f;
+						KinematicWeights[2*PointIndex+1] = 1.0f;
+					}
+				}
+			}, EParallelForFlags::None);
+			GuidesFacade.SetVertexKinematicWeights(KinematicWeights);
+		}
+		
+		SetValue(Context, MoveTemp(GroomCollection), &Collection);
+	}
+}
+
diff --git a/Engine/Plugins/Runtime/HairStrands/Source/HairStrandsDataflow/Private/BuildGuidesLODsNode.cpp b/Engine/Plugins/Runtime/HairStrands/Source/HairStrandsDataflow/Private/BuildGuidesLODsNode.cpp
new file mode 100644
--- /dev/null
+++ b/Engine/Plugins/Runtime/HairStrands/Source/HairStrandsDataflow/Private/BuildGuidesLODsNode.cpp
@@ -0,0 +1,107 @@
+// Copyright Epic Games, Inc. All Rights Reserved.
+
+#include "BuildGuidesLODsNode.h"
+#include "GroomCollectionFacades.h"
+#include "Dataflow/DataflowObjectInterface.h"
+
+#include UE_INLINE_GENERATED_CPP_BY_NAME(BuildGuidesLODsNode)
+
+namespace UE::Groom::Private
+{
+	FORCEINLINE float ComputeGuidesMetric(const FGroomGuidesFacade& GuidesFacade, const uint32 GuideIndexA, const uint32 GuideIndexB, const float& GuideLengthA, const float& GuideLengthB,
+										 const float ShapeWeight, const float ProximityWeight)
+	{
+		const uint32 PointOffsetA = (GuideIndexA > 0) ? GuidesFacade.GetCurvePointOffsets()[GuideIndexA-1] : 0;
+		const uint32 PointOffsetB = (GuideIndexB > 0) ? GuidesFacade.GetCurvePointOffsets()[GuideIndexB-1] : 0;
+		
+		const uint32 NumPointsA = GuidesFacade.GetCurvePointOffsets()[GuideIndexA] - PointOffsetA;
+		const uint32 NumPointsB = GuidesFacade.GetCurvePointOffsets()[GuideIndexB] - PointOffsetB;
+
+		if(NumPointsA == NumPointsB)
+		{
+			float ProximityMetric = 0.0f;
+			float ShapeMetric = 0.0f;
+			for(uint32 PointIndex = 0; PointIndex < NumPointsA; ++PointIndex)
+			{
+				const FVector3f& GuidePositionA = GuidesFacade.GetPointRestPositions()[PointIndex+PointOffsetA];
+				const FVector3f& GuidePositionB = GuidesFacade.GetPointRestPositions()[PointIndex+PointOffsetB];
+				
+				ProximityMetric += (GuidePositionB - GuidePositionA).Size();
+				ShapeMetric += (GuidePositionB - GuidesFacade.GetPointRestPositions()[PointOffsetB] - GuidePositionA + GuidesFacade.GetPointRestPositions()[PointOffsetA]).Size();
+			}
+
+			const float MetricScale = 1.0f / (NumPointsA * 0.5f * (GuideLengthA+GuideLengthB));
+			return FMath::Exp(-ShapeWeight * ShapeMetric * MetricScale) * FMath::Exp(-ProximityWeight * ProximityMetric * MetricScale);
+		}
+		return 0.0f;
+	}
+
+	FORCEINLINE void BuildGuidesLODs(FGroomGuidesFacade& GuidesFacade)
+	{
+		const uint32 NumObjects = GuidesFacade.GetNumObjects();
+		
+		TArray<int32> CurveParentIndices; CurveParentIndices.Init(INDEX_NONE, GuidesFacade.GetNumCurves());
+		TArray<int32> CurveLodIndices; CurveLodIndices.Init(INDEX_NONE, GuidesFacade.GetNumCurves());
+		
+		uint32 CurveOffset = 0, PointOffset = 0;
+		for(uint32 ObjectIndex = 0; ObjectIndex < NumObjects; ++ObjectIndex)
+		{ 
+			TArray<float> GuidesLengths; GuidesLengths.Init(0.0f, GuidesFacade.GetObjectCurveOffsets()[ObjectIndex] - CurveOffset);
+			for(int32 GuideIndex = CurveOffset; GuideIndex < GuidesFacade.GetObjectCurveOffsets()[ObjectIndex]; ++GuideIndex)
+			{
+				for(int32 PointIndex = PointOffset+1; PointIndex < GuidesFacade.GetCurvePointOffsets()[GuideIndex]; ++PointIndex)
+				{
+					GuidesLengths[GuideIndex-CurveOffset] += (GuidesFacade.GetPointRestPositions()[PointIndex]-GuidesFacade.GetPointRestPositions()[PointIndex-1]).Size();
+				}
+				PointOffset = GuidesFacade.GetCurvePointOffsets()[GuideIndex];
+			}
+			const int32 NumObjectLODs = FMath::CeilLogTwo(GuidesLengths.Num());
+			ParallelFor(GuidesLengths.Num(), [CurveOffset, &GuidesLengths, &GuidesFacade, &CurveParentIndices, &CurveLodIndices, &NumObjectLODs](int32 CurveIndex)
+			{
+				const uint32 GuideIndex = CurveIndex + CurveOffset;
+				const uint32 GuideLod = NumObjectLODs - 1 - FMath::FloorLog2(CurveIndex);
+				const uint32 LodOffset = (GuideLod == NumObjectLODs - 1) ? 0 : 1 << FMath::FloorLog2(CurveIndex);
+
+				float MinMetric = FLT_MAX;
+				int32 MinParent = INDEX_NONE;
+				for(uint32 ParentIndex = CurveOffset; ParentIndex < CurveOffset+LodOffset; ++ParentIndex)
+				{
+					const float ParentMetric = 1.0f - ComputeGuidesMetric(GuidesFacade,
+						GuideIndex, ParentIndex, GuidesLengths[CurveIndex], GuidesLengths[ParentIndex-CurveOffset],  1.0f, 1.0f);
+					if (ParentMetric < MinMetric)
+					{
+						MinMetric = ParentMetric;
+						MinParent = ParentIndex;
+					}
+				}
+				CurveParentIndices[GuideIndex] = MinParent;
+				CurveLodIndices[GuideIndex] = GuideLod;
+
+				//UE_LOG(LogTemp, Log, TEXT("Guide Index = %d, Num Guides = %d, Num LODs = %d, Guide LOD = %d, LOD Offset = %d, Parent Index = %d"),
+				//	CurveIndex, GuidesLengths.Num(), NumObjectLODs, GuideLod, LodOffset, MinParent);
+			}, EParallelForFlags::None);
+
+			GuidesFacade.SetCurveParentIndices(CurveParentIndices);
+			GuidesFacade.SetCurveLodIndices(CurveLodIndices);
+			CurveOffset = GuidesFacade.GetObjectCurveOffsets()[ObjectIndex];
+		}
+	}
+}
+
+void FBuildGuidesLODsDataflowNode::Evaluate(UE::Dataflow::FContext& Context, const FDataflowOutput* Out) const
+{
+	if (Out->IsA<FManagedArrayCollection>(&Collection))
+	{
+		FManagedArrayCollection GroomCollection = GetValue<FManagedArrayCollection>(Context, &Collection);
+		
+		UE::Groom::FGroomGuidesFacade GuidesFacade(GroomCollection);
+
+		if(GuidesFacade.IsValid())
+		{
+			UE::Groom::Private::BuildGuidesLODs(GuidesFacade);
+		}
+		
+		SetValue(Context, MoveTemp(GroomCollection), &Collection);
+	}
+}
+
diff --git a/Engine/Plugins/Runtime/HairStrands/Source/HairStrandsDataflow/Private/BuildGuidesSkinningNode.cpp b/Engine/Plugins/Runtime/HairStrands/Source/HairStrandsDataflow/Private/BuildGuidesSkinningNode.cpp
--- a/Engine/Plugins/Runtime/HairStrands/Source/HairStrandsDataflow/Private/BuildGuidesSkinningNode.cpp
+++ b/Engine/Plugins/Runtime/HairStrands/Source/HairStrandsDataflow/Private/BuildGuidesSkinningNode.cpp
@@ -62,32 +62,39 @@ void FBuildGuidesSkinningDataflowNode::Evaluate(UE::Dataflow::FContext& Context,
 	{
 		FManagedArrayCollection GroomCollection = GetValue<FManagedArrayCollection>(Context, &Collection);
 		
-		UE::Groom::FGroomGuidesFacade CurvesFacade(GroomCollection);
+		UE::Groom::FGroomGuidesFacade GuidesFacade(GroomCollection);
 
-		if(CurvesFacade.IsValid())
+		if(GuidesFacade.IsValid())
 		{
-			const USkeletalMesh* LocalSkelmesh = SkeletalMesh;
-			if (!LocalSkelmesh)
+			TManagedArray<TObjectPtr<UObject>>& ObjectSkeletalMeshes =
+				GroomCollection.AddAttribute<TObjectPtr<UObject>>(UE::Groom::FGroomGuidesFacade::ObjectSkeletalMeshesAttribute, UE::Groom::FGroomGuidesFacade::ObjectsGroup);
+
+			TManagedArray<int32>& ObjectMeshLODs =
+				GroomCollection.AddAttribute<int32>(UE::Groom::FGroomGuidesFacade::ObjectMeshLODsAttribute, UE::Groom::FGroomGuidesFacade::ObjectsGroup);
+			
+			if(GroupIndex == INDEX_NONE)
 			{
-				if (const UE::Dataflow::FEngineContext* EngineContext = Context.AsType<UE::Dataflow::FEngineContext>())
+				for(int32 LocalIndex = 0; LocalIndex < ObjectSkeletalMeshes.Num(); ++LocalIndex)
 				{
-					if(UGroomAsset* GroomAsset = Cast<UGroomAsset>(EngineContext->Owner))
-					{
-						LocalSkelmesh = GroomAsset->GetDataflowSettings().GetSkeletalMesh();
-					}
+					ObjectSkeletalMeshes[LocalIndex] = SkeletalMesh;
+					ObjectMeshLODs[LocalIndex] = LODIndex;
 				}
 			}
-			if(LocalSkelmesh && LocalSkelmesh->IsValidLODIndex(LODIndex))
+			else if(ObjectSkeletalMeshes.IsValidIndex(GroupIndex) && ObjectMeshLODs.IsValidIndex(GroupIndex))
+			{
+				ObjectSkeletalMeshes[GroupIndex] = SkeletalMesh;
+				ObjectMeshLODs[GroupIndex] = LODIndex;
+			}
+			if(SkeletalMesh && SkeletalMesh->IsValidLODIndex(LODIndex))
 			{
 				TMap<FName, FBoneIndexType> TargetBoneToIndex;
-				TargetBoneToIndex.Reserve(LocalSkelmesh->GetRefSkeleton().GetRawBoneNum());
-				for (int32 BoneIdx = 0; BoneIdx < LocalSkelmesh->GetRefSkeleton().GetRawBoneNum(); ++BoneIdx)
+				TargetBoneToIndex.Reserve(SkeletalMesh->GetRefSkeleton().GetRawBoneNum());
+				for (int32 BoneIdx = 0; BoneIdx < SkeletalMesh->GetRefSkeleton().GetRawBoneNum(); ++BoneIdx)
 				{
-					TargetBoneToIndex.Add(LocalSkelmesh->GetRefSkeleton().GetRawRefBoneInfo()[BoneIdx].Name, BoneIdx);
+					TargetBoneToIndex.Add(SkeletalMesh->GetRefSkeleton().GetRawRefBoneInfo()[BoneIdx].Name, BoneIdx);
 				}
-				
 				UE::Geometry::FDynamicMesh3 DynamicMesh;
-				if (UE::Groom::Private::SkeletalMeshToDynamicMesh(LocalSkelmesh, LODIndex, DynamicMesh))
+				if (UE::Groom::Private::SkeletalMeshToDynamicMesh(SkeletalMesh, LODIndex, DynamicMesh))
 				{
 					MeshTransforms::ApplyTransform(DynamicMesh, RelativeTransform, true);
 
@@ -98,15 +105,30 @@ void FBuildGuidesSkinningDataflowNode::Evaluate(UE::Dataflow::FContext& Context,
 
 					if (TransferBoneWeights.Validate() == UE::Geometry::EOperationValidationResult::Ok)
 					{
-						ParallelFor(CurvesFacade.GetNumPoints(), [&TransferBoneWeights, &TargetBoneToIndex, &CurvesFacade](int32 PointIndex)
+						const int32 LocalIndex = GroupIndex;
+						ParallelFor(GuidesFacade.GetNumPoints(), [&TransferBoneWeights, &TargetBoneToIndex, &GuidesFacade, LocalIndex](int32 PointIndex)
 						{
-							TArray<int32> BoneIndices;
-							TArray<float> BoneWeights;
-							TransferBoneWeights.TransferWeightsToPoint(BoneIndices, BoneWeights,
-								CurvesFacade.GetPointRestPositions()[PointIndex],&TargetBoneToIndex);
-
-							CurvesFacade.SetPointBoneIndices(PointIndex, BoneIndices);
-							CurvesFacade.SetPointBoneWeights(PointIndex, BoneWeights);
+							const int32 CurveIndex = GuidesFacade.GetPointCurveIndices()[PointIndex];
+							const int32 ObjectIndex = GuidesFacade.GetCurveObjectIndices()[CurveIndex];
+
+							if(LocalIndex == INDEX_NONE || ObjectIndex == LocalIndex)
+							{
+								TArray<int32> BoneIndices;
+								TArray<float> BoneWeights;
+								TransferBoneWeights.TransferWeightsToPoint(BoneIndices, BoneWeights,
+									GuidesFacade.GetPointRestPositions()[PointIndex],&TargetBoneToIndex);
+
+								FIntVector4 PointBoneIndices(INDEX_NONE);
+								FVector4f PointBoneWeights(0.0f);
+
+								for(int32 BoneIdx = 0; BoneIdx < BoneIndices.Num(); ++BoneIdx)
+								{
+									PointBoneIndices[BoneIdx] = BoneIndices[BoneIdx];
+									PointBoneWeights[BoneIdx] = BoneWeights[BoneIdx];
+								}
+								GuidesFacade.SetPointBoneIndices(PointIndex, PointBoneIndices);
+								GuidesFacade.SetPointBoneWeights(PointIndex, PointBoneWeights);
+							}
 
 						}, TransferBoneWeights.bUseParallel  ? EParallelForFlags::None : EParallelForFlags::ForceSingleThread);
 					}
diff --git a/Engine/Plugins/Runtime/HairStrands/Source/HairStrandsDataflow/Private/GroomAssetTerminalNode.cpp b/Engine/Plugins/Runtime/HairStrands/Source/HairStrandsDataflow/Private/GroomAssetTerminalNode.cpp
--- a/Engine/Plugins/Runtime/HairStrands/Source/HairStrandsDataflow/Private/GroomAssetTerminalNode.cpp
+++ b/Engine/Plugins/Runtime/HairStrands/Source/HairStrandsDataflow/Private/GroomAssetTerminalNode.cpp
@@ -2,6 +2,7 @@
 
 #include "GroomAssetTerminalNode.h"
 
+#include "AssetCompilingManager.h"
 #include "GroomCollectionFacades.h"
 #include "GroomEdit.h"
 #include "Dataflow/DataflowObjectInterface.h"
@@ -80,6 +81,14 @@ namespace UE::Groom::Private
 			if(InCollection.GetAttributeType(AttributeName, FacadeType::VerticesGroup) == EManagedArrayType::FFloatType)
 			{
 				const TManagedArray<float>& VerticesAttribute = InCollection.GetAttribute<float>(AttributeName, FacadeType::VerticesGroup);
+				if(OutCollection->NumElements(FacadeType::PointsGroup) != NumPoints)
+				{
+					if(OutCollection->NumElements(FacadeType::PointsGroup) > 0)
+					{
+						OutCollection->EmptyGroup(FacadeType::PointsGroup);
+					}
+					OutCollection->AddElements(NumPoints, FacadeType::PointsGroup);
+				}
 				TManagedArray<float>& PointsAttribute = OutCollection->AddAttribute<float>(AttributeName, FacadeType::PointsGroup);
 				
 				for(int32 PointIndex = 0; PointIndex < NumPoints; ++PointIndex)
@@ -97,22 +106,28 @@ void FGroomAssetTerminalDataflowNode::SetAssetValue(TObjectPtr<UObject> Asset, U
 	{
 		const FManagedArrayCollection& InCollection = GetValue<FManagedArrayCollection>(Context, &Collection);
 		UE::Groom::FGroomGuidesFacade GuidesFacade(InCollection);
+		UE::Groom::FGroomStrandsFacade StrandsFacade(InCollection);
 
-		FManagedArrayCollection* OutCollection = (GuidesFacade.IsValid()) ? new FManagedArrayCollection() : nullptr;
+		FManagedArrayCollection* OutCollection = (GuidesFacade.IsValid() || StrandsFacade.IsValid()) ? new FManagedArrayCollection() : nullptr;
 
+		if(StrandsFacade.IsValid())
+		{
+			// Transfer vertices weight maps onto the points to be stored onto the rest collection
+			UE::Groom::Private::TransferVerticesAttributes<UE::Groom::FGroomStrandsFacade>(InCollection, OutCollection, StrandsFacade.GetNumPoints());
+		}
 		if(GuidesFacade.IsValid())
 		{
 			// Only copy the attributes not stored already in the asset
 			const TArray<TTuple<FName, FName>> AttributesToCopy = {
 				{UE::Groom::FGroomGuidesFacade::PointsGroup, UE::Groom::FGroomGuidesFacade::PointBoneIndicesAttribute},
 				{UE::Groom::FGroomGuidesFacade::PointsGroup, UE::Groom::FGroomGuidesFacade::PointBoneWeightsAttribute},
-				{UE::Groom::FGroomGuidesFacade::ObjectsGroup, UE::Groom::FGroomGuidesFacade::ObjectCurveOffsetsAttribute},
-				{UE::Groom::FGroomGuidesFacade::CurvesGroup, UE::Groom::FGroomGuidesFacade::CurvePointOffsetsAttribute}};
+				{UE::Groom::FGroomGuidesFacade::CurvesGroup, UE::Groom::FGroomGuidesFacade::CurveParentIndicesAttribute},
+				{UE::Groom::FGroomGuidesFacade::CurvesGroup, UE::Groom::FGroomGuidesFacade::CurveLodIndicesAttribute}};
 
 			// Copy attributes from input collection
 			UE::Groom::Private::CopyCollectionAttributes(&InCollection, OutCollection, AttributesToCopy);
 			
-			// Transfer vertices weight maps onto the points to be stored onto the rest collection
+			// Transfer vertices weight maps onto the points to be stored onto the rest collection (VertexKinematicWeights should be included here)
 			UE::Groom::Private::TransferVerticesAttributes<UE::Groom::FGroomGuidesFacade>(InCollection, OutCollection, GuidesFacade.GetNumPoints());
 
 			// Build an editable groom asset for the strands
@@ -124,6 +139,9 @@ void FGroomAssetTerminalDataflowNode::SetAssetValue(TObjectPtr<UObject> Asset, U
 				// Build the editable guides
 				UE::Groom::Private::BuildEditableGuides(GuidesFacade,EditGroom);
 
+				// Ensure compilation dependent assets is done
+				FAssetCompilingManager::Get().FinishCompilationForObjects({ GroomAsset });
+
 				// Convert to groom asset
 				ConvertToGroomAsset(const_cast<UGroomAsset*>(GroomAsset), &EditGroom, EEditableGroomOperations::ControlPoints_Modified);
 			}
@@ -132,6 +150,22 @@ void FGroomAssetTerminalDataflowNode::SetAssetValue(TObjectPtr<UObject> Asset, U
 		if(OutCollection)
 		{
 			GroomAsset->GetDataflowSettings().SetRestCollection(OutCollection);
+			GroomAsset->GetDataflowSettings().InitSkeletalMeshes(GuidesFacade.GetNumObjects());
+
+			if(InCollection.HasAttribute(UE::Groom::FGroomGuidesFacade::ObjectSkeletalMeshesAttribute, UE::Groom::FGroomGuidesFacade::ObjectsGroup))
+			{
+				const TManagedArray<TObjectPtr<UObject>>& ObjectSkeletalMeshes =
+					InCollection.GetAttribute<TObjectPtr<UObject>>(UE::Groom::FGroomGuidesFacade::ObjectSkeletalMeshesAttribute, UE::Groom::FGroomGuidesFacade::ObjectsGroup);
+
+				const TManagedArray<int32>& ObjectMeshLODs =
+					InCollection.GetAttribute<int32>(UE::Groom::FGroomGuidesFacade::ObjectMeshLODsAttribute, UE::Groom::FGroomGuidesFacade::ObjectsGroup);
+				
+				for(int32 GroupIndex = 0; GroupIndex < GuidesFacade.GetNumObjects(); ++GroupIndex)
+				{
+					GroomAsset->GetDataflowSettings().SetSkeletalMesh(GroupIndex,
+						Cast<USkeletalMesh>(ObjectSkeletalMeshes[GroupIndex]), ObjectMeshLODs[GroupIndex]);
+				}
+			}
 		}
 	}
 }
diff --git a/Engine/Plugins/Runtime/HairStrands/Source/HairStrandsDataflow/Private/GroomCollectionFacades.cpp b/Engine/Plugins/Runtime/HairStrands/Source/HairStrandsDataflow/Private/GroomCollectionFacades.cpp
--- a/Engine/Plugins/Runtime/HairStrands/Source/HairStrandsDataflow/Private/GroomCollectionFacades.cpp
+++ b/Engine/Plugins/Runtime/HairStrands/Source/HairStrandsDataflow/Private/GroomCollectionFacades.cpp
@@ -14,6 +14,8 @@ namespace UE::Groom
 	const FName FGroomCollectionFacade<DerivedType>::ObjectsGroup(DerivedType::GroupPrefix.ToString() + FString("Objects"));
 	template<typename DerivedType>
 	const FName FGroomCollectionFacade<DerivedType>::PointsGroup(DerivedType::GroupPrefix.ToString() + FString("Points"));
+	template<typename DerivedType>
+	const FName FGroomCollectionFacade<DerivedType>::VerticesGroup(DerivedType::GroupPrefix.ToString() + FString("Vertices"));
 	
 	// Strands/Guides Attributes
 	template<typename DerivedType>
@@ -28,24 +30,23 @@ namespace UE::Groom
 	const FName FGroomCollectionFacade<DerivedType>::EdgeRestOrientationsAttribute("EdgeRestOrientations");
 	template<typename DerivedType>
 	const FName FGroomCollectionFacade<DerivedType>::PointRestPositionsAttribute("PointRestPositions");
+	template<typename DerivedType>
+	const FName FGroomCollectionFacade<DerivedType>::VertexLinearColorsAttribute("VertexLinearColors");
 
 	// Strands/Guides Prefix
 	const FName FGroomStrandsFacade::GroupPrefix("Strands");
 	const FName FGroomGuidesFacade::GroupPrefix("Guides");
 
-	// Guides Groups
-	const FName FGroomGuidesFacade::VerticesGroup(FGroomGuidesFacade::GroupPrefix.ToString() + FString("Vertices"));
-	//const FName FGroomGuidesFacade::VerticesGroup(FString("GuidesVertices"));
-
 	// Guides Attributes
+	const FName FGroomGuidesFacade::ObjectMeshLODsAttribute("ObjectMeshLODs");
+	const FName FGroomGuidesFacade::ObjectSkeletalMeshesAttribute("ObjectSkeletalMeshes");
 	const FName FGroomGuidesFacade::VertexKinematicWeightsAttribute("VertexKinematicWeights");
 	const FName FGroomGuidesFacade::PointBoneIndicesAttribute("PointBoneIndices");
 	const FName FGroomGuidesFacade::PointBoneWeightsAttribute("PointBoneWeights");
 	const FName FGroomGuidesFacade::ObjectPointSamplesAttribute("ObjectPointSamples");
 	const FName FGroomGuidesFacade::CurveStrandIndicesAttribute("CurveStrandIndices");
-	const FName FGroomGuidesFacade::CurveLevelIndicesAttribute("CurveLevelIndices");
-	const FName FGroomGuidesFacade::PointInterpIndicesAttribute("PointInterpIndices");
-	const FName FGroomGuidesFacade::PointInterpWeightsAttribute("PointInterpWeights");
+	const FName FGroomGuidesFacade::CurveParentIndicesAttribute("CurveParentIndices");
+	const FName FGroomGuidesFacade::CurveLodIndicesAttribute("CurveLodIndices");
 	
 	template<typename DerivedType>
 	FGroomCollectionFacade<DerivedType>::FGroomCollectionFacade(FManagedArrayCollection& InCollection) :
@@ -55,8 +56,10 @@ namespace UE::Groom
 		CurvePointOffsets(InCollection, CurvePointOffsetsAttribute, CurvesGroup),
 		ObjectCurveOffsets(InCollection, ObjectCurveOffsetsAttribute, ObjectsGroup),
 		PointCurveIndices(InCollection, PointCurveIndicesAttribute, PointsGroup),
-		CurveObjectIndices(InCollection, CurveObjectIndicesAttribute, CurvesGroup)
-	{}
+		CurveObjectIndices(InCollection, CurveObjectIndicesAttribute, CurvesGroup),
+		VertexLinearColors(InCollection, VertexLinearColorsAttribute, VerticesGroup)
+	{
+	}
 
 	template<typename DerivedType>
 	FGroomCollectionFacade<DerivedType>::FGroomCollectionFacade(const FManagedArrayCollection& InCollection) :
@@ -66,7 +69,8 @@ namespace UE::Groom
 		CurvePointOffsets(InCollection, CurvePointOffsetsAttribute, CurvesGroup),
 		ObjectCurveOffsets(InCollection, ObjectCurveOffsetsAttribute, ObjectsGroup),
 		PointCurveIndices(InCollection, PointCurveIndicesAttribute, PointsGroup),
-		CurveObjectIndices(InCollection, CurveObjectIndicesAttribute, CurvesGroup)
+		CurveObjectIndices(InCollection, CurveObjectIndicesAttribute, CurvesGroup),
+		VertexLinearColors(InCollection, VertexLinearColorsAttribute, VerticesGroup)
 	{
 	}
 
@@ -75,7 +79,7 @@ namespace UE::Groom
 	{
 		return EdgeRestOrientations.IsValid() && PointRestPositions.IsValid() && 
 			CurvePointOffsets.IsValid() && ObjectCurveOffsets.IsValid() && PointCurveIndices.IsValid() &&
-					CurveObjectIndices.IsValid() && static_cast<const DerivedType*>(this)->IsFacadeValid();
+					CurveObjectIndices.IsValid() && VertexLinearColors.IsValid() && static_cast<const DerivedType*>(this)->IsFacadeValid();
 	}
 
 	template<typename DerivedType>
@@ -88,6 +92,7 @@ namespace UE::Groom
 		ObjectCurveOffsets.Add();
 		PointCurveIndices.Add();
 		CurveObjectIndices.Add();
+		VertexLinearColors.Add();
 		
 		static_cast<DerivedType*>(this)->DefineFacadeSchema();
 	}
@@ -114,6 +119,10 @@ namespace UE::Groom
 			Collection->EmptyGroup(EdgesGroup);
 			Collection->AddElements(InPointRestPositions.Num()-InCurvePointOffsets.Num(), EdgesGroup);
 
+			// Vertices groups
+			Collection->EmptyGroup(VerticesGroup);
+			Collection->AddElements(InPointRestPositions.Num() * 2, VerticesGroup);
+
 			// Fill attributes
 			SetObjectCurveOffsets(InObjectCurveOffsets);
 			SetCurvePointOffsets(InCurvePointOffsets);
@@ -207,10 +216,9 @@ namespace UE::Groom
 		PointBoneIndices(InCollection, PointBoneIndicesAttribute, PointsGroup),
 		PointBoneWeights(InCollection, PointBoneWeightsAttribute, PointsGroup),
 		ObjectPointSamples(InCollection, ObjectPointSamplesAttribute, ObjectsGroup),
-		CurveStrandIndices(InCollection, CurveStrandIndicesAttribute, CurvesGroup)
-		// CurveLevelIndices(InCollection, CurveLevelIndicesAttribute, CurvesGroup),
-		// PointInterpIndices(InCollection, PointInterpIndicesAttribute, PointsGroup),
-		// PointInterpWeights(InCollection, PointInterpWeightsAttribute, PointsGroup)
+		CurveStrandIndices(InCollection, CurveStrandIndicesAttribute, CurvesGroup),
+		CurveParentIndices(InCollection, CurveParentIndicesAttribute, CurvesGroup),
+		CurveLodIndices(InCollection, CurveLodIndicesAttribute, CurvesGroup)
 	{
 		DefineSchema();
 	}
@@ -221,21 +229,16 @@ namespace UE::Groom
 		PointBoneIndices(InCollection, PointBoneIndicesAttribute, PointsGroup),
 		PointBoneWeights(InCollection, PointBoneWeightsAttribute, PointsGroup),
 		ObjectPointSamples(InCollection, ObjectPointSamplesAttribute, ObjectsGroup),
-		CurveStrandIndices(InCollection, CurveStrandIndicesAttribute, CurvesGroup)
-		// CurveLevelIndices(InCollection, CurveLevelIndicesAttribute, CurvesGroup),
-		// PointInterpIndices(InCollection, PointInterpIndicesAttribute, PointsGroup),
-		// PointInterpWeights(InCollection, PointInterpWeightsAttribute, PointsGroup)
+		CurveStrandIndices(InCollection, CurveStrandIndicesAttribute, CurvesGroup),
+		CurveParentIndices(InCollection, CurveParentIndicesAttribute, CurvesGroup),
+		CurveLodIndices(InCollection, CurveLodIndicesAttribute, CurvesGroup)
 	{
 	}
 
 	bool FGroomGuidesFacade::IsFacadeValid() const
 	{
-		return VertexKinematicWeights.IsValid() && PointBoneIndices.IsValid() && PointBoneWeights.IsValid() && ObjectPointSamples.IsValid();
-	}
-
-	bool FGroomGuidesFacade::HasSkinningAttributes() const
-	{
-		return PointBoneIndices.IsValid() && PointBoneWeights.IsValid() && ObjectCurveOffsets.IsValid() && CurvePointOffsets.IsValid();
+		return VertexKinematicWeights.IsValid() && PointBoneIndices.IsValid() && PointBoneWeights.IsValid() && ObjectPointSamples.IsValid() &&
+			CurveStrandIndices.IsValid() && CurveLodIndices.IsValid() && CurveParentIndices.IsValid();
 	}
 
 	void FGroomGuidesFacade::DefineFacadeSchema()
@@ -246,26 +249,20 @@ namespace UE::Groom
 		PointBoneIndices.Add();
 		ObjectPointSamples.Add();
 		CurveStrandIndices.Add();
-		// CurveLevelIndices.Add();
-		// PointInterpIndices.Add();
-		// PointInterpWeights.Add();
+		CurveParentIndices.Add();
+		CurveLodIndices.Add();
 	}
 
 	void FGroomGuidesFacade::InitFacadeCollection()
 	{
-		// Reset groups
-		Collection->EmptyGroup(VerticesGroup);
-		Collection->AddElements(GetNumPoints() * 2, VerticesGroup);
-
 		// Reset attributes
 		VertexKinematicWeights.Modify().Fill(0.0f);
-		PointBoneIndices.Modify().Fill(TArray<int32>());
-		PointBoneWeights.Modify().Fill(TArray<float>());
+		PointBoneIndices.Modify().Fill(FIntVector4::ZeroValue);
+		PointBoneWeights.Modify().Fill(FVector4f::Zero());
 		ObjectPointSamples.Modify().Fill(4);
 		CurveStrandIndices.Modify().Fill(INDEX_NONE);
-		// CurveLevelIndices.Modify().Fill(0);
-		// PointInterpIndices.Modify().Fill(FVector4f(INDEX_NONE));
-		// PointInterpWeights.Modify().Fill(FVector4f(0.0f));
+		CurveParentIndices.Modify().Fill(INDEX_NONE);
+		CurveLodIndices.Modify().Fill(INDEX_NONE);
 	}
 
 	void FGroomGuidesFacade::ResizePointsGroups(const int32 NumPoints) const 
diff --git a/Engine/Plugins/Runtime/HairStrands/Source/HairStrandsDataflow/Private/GroomDataflowRendering.cpp b/Engine/Plugins/Runtime/HairStrands/Source/HairStrandsDataflow/Private/GroomDataflowRendering.cpp
--- a/Engine/Plugins/Runtime/HairStrands/Source/HairStrandsDataflow/Private/GroomDataflowRendering.cpp
+++ b/Engine/Plugins/Runtime/HairStrands/Source/HairStrandsDataflow/Private/GroomDataflowRendering.cpp
@@ -2,6 +2,7 @@
 
 #include "GroomDataflowRendering.h"
 #include "GroomCollectionFacades.h"
+#include "Engine/SkeletalMesh.h"
 #include "GeometryCollection/GeometryCollection.h"
 
 namespace UE::Groom
@@ -10,14 +11,21 @@ namespace UE::Groom
 	static FAutoConsoleVariableRef CVarGroomDataflowDebugDraw(
 		 TEXT("p.Groom.Dataflow.DebugDraw"),
 		 GGroomDataflowDebugDraw,
-		 TEXT("Type of information we want to draw for debug (0-DefaultColor, 1-SkinWeights)")
+		 TEXT("Type of information we want to draw for debug (0-DefaultColor, 1-SkinWeights, 2-GuidesLods)")
 	 );
 
-	int32 GGroomDataflowBoneIndex = 0;
-	static FAutoConsoleVariableRef CVarGroomDataflowBoneIndex(
-		 TEXT("p.Groom.Dataflow.BoneIndex"),
-		 GGroomDataflowBoneIndex,
-		 TEXT("Bone index we want to visualize the skin weights")
+	FString GGroomDataflowBoneName = TEXT("");
+	static FAutoConsoleVariableRef CVarGroomDataflowBoneName(
+		 TEXT("p.Groom.Dataflow.BoneName"),
+		 GGroomDataflowBoneName,
+		 TEXT("Bone name we want to visualize the skin weights")
+	 );
+
+	int32 GGroomDataflowGuidesLod = 0;
+	static FAutoConsoleVariableRef CVarGroomDataflowGuidesLod(
+		 TEXT("p.Groom.Dataflow.GuidesLod"),
+		 GGroomDataflowGuidesLod,
+		 TEXT("Guides Lod we want to display")
 	 );
 
 	float GGroomDataflowRenderingThickness = 0.5f;
@@ -190,7 +198,7 @@ namespace UE::Groom
 				}
 
 				TArray<FLinearColor> Colors;
-				Colors.Init(FLinearColor::Blue, VertexPositions.Num());
+				Colors.Init(FLinearColor(0.2,0.6,1.0), VertexPositions.Num());
 				ColorLambda(GroomFacade, Colors);
  
 				const int32 GeometryIndex = RenderCollection.StartGeometryGroup(State.GetGuid().ToString());
@@ -201,29 +209,95 @@ namespace UE::Groom
 	}
 	
 	Dataflow::FRenderKey FGroomGuidesRenderingCallbacks::RenderKey = { "GuidesRender", FName("FGroomCollection") };
-	
-	void FGroomGuidesRenderingCallbacks::Render(GeometryCollection::Facades::FRenderingFacade& RenderCollection, const Dataflow::FGraphRenderingState& State)
+
+	void RenderBoneWeights(const FGroomGuidesFacade& GuidesFacade, TArray<FLinearColor>& VertexColors)
 	{
-		RenderGroomCollection<FGroomGuidesFacade>(RenderCollection, State,
-			[](const FGroomGuidesFacade& GuidesFacade, TArray<FLinearColor>& VertexColors)
+		if(GuidesFacade.GetManagedArrayCollection().HasAttribute(
+						UE::Groom::FGroomGuidesFacade::ObjectSkeletalMeshesAttribute, UE::Groom::FGroomGuidesFacade::ObjectsGroup))
+		{
+			const TManagedArray<TObjectPtr<UObject>>& ObjectSkeletalMeshes = GuidesFacade.GetManagedArrayCollection().GetAttribute<TObjectPtr<UObject>>(
+				UE::Groom::FGroomGuidesFacade::ObjectSkeletalMeshesAttribute, UE::Groom::FGroomGuidesFacade::ObjectsGroup);
+				
+			for(int32 PointIndex = 0, PointEnd = GuidesFacade.GetNumPoints(); PointIndex < PointEnd; ++PointIndex)
 			{
-				if(GGroomDataflowDebugDraw == 1)
+				const FIntVector4& BoneIndices = GuidesFacade.GetPointBoneIndices(PointIndex);
+				const FVector4f& BoneWeights = GuidesFacade.GetPointBoneWeights(PointIndex);
+
+				const int32 CurveIndex = GuidesFacade.GetPointCurveIndices()[PointIndex];
+				const int32 ObjectIndex = GuidesFacade.GetCurveObjectIndices()[CurveIndex];
+
+				if(USkeletalMesh* SkeletalMesh = Cast<USkeletalMesh>(ObjectSkeletalMeshes[ObjectIndex]))
 				{
-					for(int32 PointIndex = 0, PointEnd = GuidesFacade.GetNumPoints(); PointIndex < PointEnd; ++PointIndex)
+					for(int32 BoneIndex = 0, BoneEnd = BoneIndices.Num(); BoneIndex < BoneEnd; ++BoneIndex)
 					{
-						const TArray<int32>& BoneIndices = GuidesFacade.GetPointBoneIndices(PointIndex);
-						const TArray<float>& BoneWeights = GuidesFacade.GetPointBoneWeights(PointIndex);
-						for(int32 BoneIndex = 0, BoneEnd = BoneIndices.Num(); BoneIndex < BoneEnd; ++BoneIndex)
+						const int32 SkeletonBone = BoneIndices[BoneIndex];
+						if(SkeletonBone != INDEX_NONE)
 						{
-							if(BoneIndices[BoneIndex] == GGroomDataflowBoneIndex)
+							const FName BoneName = SkeletalMesh->GetRefSkeleton().GetRawRefBoneInfo()[SkeletonBone].Name;
+							if(BoneName == GGroomDataflowBoneName)
 							{
 								VertexColors[2*PointIndex] = FLinearColor::LerpUsingHSV(
 									FLinearColor::Black, FLinearColor::Yellow, BoneWeights[BoneIndex]);
 								VertexColors[2*PointIndex+1] = VertexColors[2*PointIndex];
 							}
-						}	
+						}
 					}
 				}
+			}
+		}
+	}
+	
+	void RenderGuidesLods(const FGroomGuidesFacade& GuidesFacade, TArray<FLinearColor>& VertexColors)
+	{
+		const TArray<int32>& ParentIndices = GuidesFacade.GetCurveParentIndices();
+		const TArray<int32>& LodIndices = GuidesFacade.GetCurveLodIndices();
+
+		TArray<FLinearColor> LodColors;
+		LodColors.Init(FLinearColor::Black, GuidesFacade.GetNumCurves());
+		for(int32 CurveIndex = 0, CurveEnd = GuidesFacade.GetNumCurves(); CurveIndex < CurveEnd; ++CurveIndex)
+		{
+			if(LodIndices[CurveIndex] >= GGroomDataflowGuidesLod)
+			{
+				LodColors[CurveIndex] = FLinearColor::MakeRandomColor();
+			}
+		}
+		for(int32 CurveIndex = 0, CurveEnd = GuidesFacade.GetNumCurves(); CurveIndex < CurveEnd; ++CurveIndex)
+		{
+			if((LodIndices[CurveIndex] != INDEX_NONE) && (LodIndices[CurveIndex] < GGroomDataflowGuidesLod))
+			{
+				int32 ParentIndex = CurveIndex;
+				int32 ParentLod = LodIndices[ParentIndex];
+				while((ParentLod < GGroomDataflowGuidesLod) && ParentIndex != INDEX_NONE)
+				{
+					ParentIndex = ParentIndices[ParentIndex];
+					ParentLod = LodIndices[ParentIndex];
+				}
+				LodColors[CurveIndex] = LodColors[ParentIndex];
+			}
+		}
+
+		for(int32 PointIndex = 0, PointEnd = GuidesFacade.GetNumPoints(); PointIndex < PointEnd; ++PointIndex)
+		{
+			const int32 CurveIndex = GuidesFacade.GetPointCurveIndices()[PointIndex];
+
+			VertexColors[2*PointIndex] = LodColors[CurveIndex];
+			VertexColors[2*PointIndex+1] = LodColors[CurveIndex];
+		}
+	}
+	
+	void FGroomGuidesRenderingCallbacks::Render(GeometryCollection::Facades::FRenderingFacade& RenderCollection, const Dataflow::FGraphRenderingState& State)
+	{
+		RenderGroomCollection<FGroomGuidesFacade>(RenderCollection, State,
+			[](const FGroomGuidesFacade& GuidesFacade, TArray<FLinearColor>& VertexColors)
+			{
+				if(GGroomDataflowDebugDraw == 1)
+				{
+					RenderBoneWeights(GuidesFacade, VertexColors);
+				}
+				if(GGroomDataflowDebugDraw == 2)
+				{
+					RenderGuidesLods(GuidesFacade, VertexColors);
+				}
 			});
 	}
 
diff --git a/Engine/Plugins/Runtime/HairStrands/Source/HairStrandsDataflow/Private/HairStrandsDataflowModule.cpp b/Engine/Plugins/Runtime/HairStrands/Source/HairStrandsDataflow/Private/HairStrandsDataflowModule.cpp
--- a/Engine/Plugins/Runtime/HairStrands/Source/HairStrandsDataflow/Private/HairStrandsDataflowModule.cpp
+++ b/Engine/Plugins/Runtime/HairStrands/Source/HairStrandsDataflow/Private/HairStrandsDataflowModule.cpp
@@ -1,6 +1,9 @@
 ﻿// Copyright Epic Games, Inc. All Rights Reserved.
 
 #include "HairStrandsDataflowModule.h"
+
+#include "AttachGuidesRootsNode.h"
+#include "BuildGuidesLODsNode.h"
 #include "GetGroomAssetNode.h"
 #include "GroomCollectionFacades.h"
 #include "GroomAssetTerminalNode.h"
@@ -31,12 +34,13 @@ public:
 
 	virtual TArray<FName> GetTargetGroupNames() const override
 	{
-		return {UE::Groom::FGroomGuidesFacade::VerticesGroup};
+		return {UE::Groom::FGroomGuidesFacade::VerticesGroup, UE::Groom::FGroomStrandsFacade::VerticesGroup};
 	}
 
 	virtual TArray<UE::Dataflow::FRenderingParameter> GetRenderingParameters() const override
 	{
-		return { {TEXT("GuidesRender"), FName("FGroomCollection"), {TEXT("Collection")}}};
+		return { {TEXT("GuidesRender"), FName("FGroomCollection"), {TEXT("Collection")}},
+			{TEXT("StrandsRender"), FName("FGroomCollection"), {TEXT("Collection")}} };
 	}
 };
 
@@ -51,6 +55,8 @@ void FHairStrandsDataflowModule::StartupModule()
 	DATAFLOW_NODE_REGISTER_CREATION_FACTORY(FGenerateGuidesCurvesDataflowNode);
 	DATAFLOW_NODE_REGISTER_CREATION_FACTORY(FGetGuidesKinematicWeightsDataflowNode);
 	DATAFLOW_NODE_REGISTER_CREATION_FACTORY(FSmoothGuidesCurvesDataflowNode);
+	DATAFLOW_NODE_REGISTER_CREATION_FACTORY(FAttachGuidesRootsDataflowNode);
+	DATAFLOW_NODE_REGISTER_CREATION_FACTORY(FBuildGuidesLODsDataflowNode);
 	
 	UE::Groom::RegisterRenderingCallbacks();
 	
diff --git a/Engine/Plugins/Runtime/HairStrands/Source/HairStrandsDataflow/Public/AttachGuidesRootsNode.h b/Engine/Plugins/Runtime/HairStrands/Source/HairStrandsDataflow/Public/AttachGuidesRootsNode.h
new file mode 100644
--- /dev/null
+++ b/Engine/Plugins/Runtime/HairStrands/Source/HairStrandsDataflow/Public/AttachGuidesRootsNode.h
@@ -0,0 +1,42 @@
+// Copyright Epic Games, Inc. All Rights Reserved.
+
+#pragma once 
+
+#include "CoreMinimal.h"
+#include "GroomAsset.h"
+#include "ChaosLog.h"
+#include "Dataflow/DataflowCore.h"
+
+#include "AttachGuidesRootsNode.generated.h"
+
+/** Attach the guides roots by setting their kinematic weights to 1.0f */
+USTRUCT(meta = (DataflowGroom))
+struct FAttachGuidesRootsDataflowNode : public FDataflowNode
+{
+	GENERATED_BODY()
+	
+	DATAFLOW_NODE_DEFINE_INTERNAL(FAttachGuidesRootsDataflowNode, "AttachGuidesRoots", "Groom", "")
+	DATAFLOW_NODE_RENDER_TYPE("GuidesRender", FName("FGroomCollection"), "Collection")
+
+public:
+	
+	FAttachGuidesRootsDataflowNode(const UE::Dataflow::FNodeParameters& InParam, FGuid InGuid = FGuid::NewGuid())
+	: FDataflowNode(InParam, InGuid)
+	{
+		RegisterInputConnection(&Collection);
+		RegisterOutputConnection(&Collection, &Collection);
+	}
+
+	//~ Begin FDataflowNode interface
+	virtual void Evaluate(UE::Dataflow::FContext& Context, const FDataflowOutput* Out) const override;
+	//~ End FDataflowNode interface
+
+	/** Managed array collection to be used to store datas */
+	UPROPERTY(meta = (DataflowInput, DataflowOutput, DisplayName = "Collection", DataflowPassthrough  = "Collection"))
+	FManagedArrayCollection Collection;
+	
+	/** Group index on which the roots will be attached. -1 will attach all the groups */
+	UPROPERTY(EditAnywhere, Category="Groups", meta = (DisplayName = "Group Index"))
+	int32 GroupIndex = INDEX_NONE;
+};
+
diff --git a/Engine/Plugins/Runtime/HairStrands/Source/HairStrandsDataflow/Public/BuildGuidesLODsNode.h b/Engine/Plugins/Runtime/HairStrands/Source/HairStrandsDataflow/Public/BuildGuidesLODsNode.h
new file mode 100644
--- /dev/null
+++ b/Engine/Plugins/Runtime/HairStrands/Source/HairStrandsDataflow/Public/BuildGuidesLODsNode.h
@@ -0,0 +1,40 @@
+// Copyright Epic Games, Inc. All Rights Reserved.
+
+#pragma once 
+
+#include "CoreMinimal.h"
+#include "GroomAsset.h"
+#include "ChaosLog.h"
+#include "Dataflow/DataflowCore.h"
+#include "Dataflow/DataflowEngine.h"
+
+#include "BuildGuidesLODsNode.generated.h"
+
+
+/** Builds the guides LODs */
+USTRUCT(meta = (DataflowGroom))
+struct FBuildGuidesLODsDataflowNode : public FDataflowNode
+{
+	GENERATED_BODY()
+	
+	DATAFLOW_NODE_DEFINE_INTERNAL(FBuildGuidesLODsDataflowNode, "BuildGuidesLODs", "Groom", "")
+	DATAFLOW_NODE_RENDER_TYPE("GuidesRender", FName("FGroomCollection"), "Collection")
+
+public:
+	
+	FBuildGuidesLODsDataflowNode(const UE::Dataflow::FNodeParameters& InParam, FGuid InGuid = FGuid::NewGuid())
+	: FDataflowNode(InParam, InGuid)
+	{
+		RegisterInputConnection(&Collection);
+		RegisterOutputConnection(&Collection, &Collection);
+	}
+
+	//~ Begin FDataflowNode interface
+	virtual void Evaluate(UE::Dataflow::FContext& Context, const FDataflowOutput* Out) const override;
+	//~ End FDataflowNode interface
+
+	/** Managed array collection to be used to store data */
+	UPROPERTY(meta = (DataflowInput, DataflowOutput, DisplayName = "Collection", DataflowPassthrough  = "Collection"))
+	FManagedArrayCollection Collection;
+};
+
diff --git a/Engine/Plugins/Runtime/HairStrands/Source/HairStrandsDataflow/Public/BuildGuidesSkinningNode.h b/Engine/Plugins/Runtime/HairStrands/Source/HairStrandsDataflow/Public/BuildGuidesSkinningNode.h
--- a/Engine/Plugins/Runtime/HairStrands/Source/HairStrandsDataflow/Public/BuildGuidesSkinningNode.h
+++ b/Engine/Plugins/Runtime/HairStrands/Source/HairStrandsDataflow/Public/BuildGuidesSkinningNode.h
@@ -35,13 +35,17 @@ public:
 	UPROPERTY(meta = (DataflowInput, DataflowOutput, DisplayName = "Collection", DataflowPassthrough  = "Collection"))
 	FManagedArrayCollection Collection;
 
-	/** Managed array collection to be used to store datas */
+	/** SkeletalMesh used to transfer the skinning weights. Will be stored onto the groom asset */
 	UPROPERTY(EditAnywhere, Category="Skinning")
 	TObjectPtr<USkeletalMesh> SkeletalMesh = nullptr;
 
-	/** Managed array collection to be used to store datas */
+	/** LOD used to transfer the weights */
 	UPROPERTY(EditAnywhere, Category="Skinning", meta = (DisplayName = "LOD Index"))
 	int32 LODIndex = 0;
+	
+	/** Group index on which the dats will be transfered. -1 will transfer on all the groups */
+	UPROPERTY(EditAnywhere, Category="Groups", meta = (DisplayName = "Group Index"))
+	int32 GroupIndex = INDEX_NONE;
 
 	/** The relative transform between the skeletal mesh and the groom asset. */
 	UPROPERTY(EditAnywhere, Category = "Skinning")
diff --git a/Engine/Plugins/Runtime/HairStrands/Source/HairStrandsDataflow/Public/GroomCollectionFacades.h b/Engine/Plugins/Runtime/HairStrands/Source/HairStrandsDataflow/Public/GroomCollectionFacades.h
--- a/Engine/Plugins/Runtime/HairStrands/Source/HairStrandsDataflow/Public/GroomCollectionFacades.h
+++ b/Engine/Plugins/Runtime/HairStrands/Source/HairStrandsDataflow/Public/GroomCollectionFacades.h
@@ -41,6 +41,7 @@ namespace UE::Groom
 		static const FName EdgesGroup;
 		static const FName ObjectsGroup;
 		static const FName PointsGroup;
+		static const FName VerticesGroup;
 		
 		/** Groom collection attribute names */
 		static const FName CurvePointOffsetsAttribute;
@@ -49,6 +50,7 @@ namespace UE::Groom
 		static const FName PointRestPositionsAttribute;
 		static const FName PointCurveIndicesAttribute;
 		static const FName CurveObjectIndicesAttribute;
+		static const FName VertexLinearColorsAttribute;
 
 		FGroomCollectionFacade(FManagedArrayCollection& InCollection);
 		FGroomCollectionFacade(const FManagedArrayCollection& InCollection);
@@ -73,6 +75,9 @@ namespace UE::Groom
 
 		/** Get the number of edges */
 		int32 GetNumEdges() const { return EdgeRestOrientations.Num(); }
+
+		/** Get the number of vertices */
+		int32 GetNumVertices() const { return VertexLinearColors.Num(); }
 		
 		/** Get the point rest positions */
 		const TArray<FVector3f>& GetPointRestPositions() const { return PointRestPositions.Get().GetConstArray(); }
@@ -91,6 +96,9 @@ namespace UE::Groom
 
 		/** Get the curve object indices */
 		const TArray<int32>& GetCurveObjectIndices() const { return CurveObjectIndices.Get().GetConstArray(); }
+
+		/** Get the curve object indices */
+		const TArray<FLinearColor>& GetVertexLinearColors() const { return VertexLinearColors.Get().GetConstArray(); }
 		
 		/** Set the point rest positions */
 		void SetPointRestPositions(const TArray<FVector3f>& InPointRestPositions) { PointRestPositions.Modify() = InPointRestPositions; UpdateEdgeRestOrientations(); }	
@@ -101,9 +109,15 @@ namespace UE::Groom
 		/** Set the object curve offsets */
 		void SetObjectCurveOffsets(const TArray<int32>& InObjectCurveOffsets) { ObjectCurveOffsets.Modify() = InObjectCurveOffsets; UpdateCurveObjectIndices(); }
 		
+		/** Set the vertex linear colors */
+		void SetVertexLinearColors(const TArray<FLinearColor>& InVertexLinearColors) { VertexLinearColors.Modify() = InVertexLinearColors; }
+		
 		/** Initialize the curve vertices from point positions and objects offsets */
 		void InitGroomCollection(const TArray<FVector3f>& InPointRestPositions,
 			const TArray<int32>& InCurvePointOffsets, const TArray<int32>& InObjectCurveOffsets);
+
+		/** Get the managed array collection */
+		const FManagedArrayCollection& GetManagedArrayCollection() const {return ConstCollection;}
 		
 	protected :
 		
@@ -139,6 +153,9 @@ namespace UE::Groom
 
 		/** Groom curves object index */
 		TManagedArrayAccessor<int32> CurveObjectIndices;
+
+		/** Groom vertices linear color*/
+		TManagedArrayAccessor<FLinearColor> VertexLinearColors;
 	};
 	
 	/**
@@ -184,17 +201,16 @@ namespace UE::Groom
 
 		FGroomGuidesFacade(FManagedArrayCollection& InCollection);
 		FGroomGuidesFacade(const FManagedArrayCollection& InCollection);
-		
-		static const FName VerticesGroup;
 
+		static const FName ObjectMeshLODsAttribute;
+		static const FName ObjectSkeletalMeshesAttribute;
 		static const FName VertexKinematicWeightsAttribute;
 		static const FName PointBoneIndicesAttribute;
 		static const FName PointBoneWeightsAttribute;
 		static const FName ObjectPointSamplesAttribute;
 		static const FName CurveStrandIndicesAttribute;
-		static const FName CurveLevelIndicesAttribute;
-		static const FName PointInterpIndicesAttribute;
-		static const FName PointInterpWeightsAttribute;
+		static const FName CurveParentIndicesAttribute;
+		static const FName CurveLodIndicesAttribute;
 		
 		/** Create the facade attributes. */
 		void DefineFacadeSchema();
@@ -202,15 +218,9 @@ namespace UE::Groom
 		/** Is the Facade defined on the collection? */
 		bool IsFacadeValid() const;
 
-		/** Check if the facade have skinning attributes*/
-		bool HasSkinningAttributes() const;
-
 		/** Init facade collection attributes */
 		void InitFacadeCollection();
 
-		/** Get the number of Vertices */
-		int32 GetNumVertices() const { return VertexKinematicWeights.Num(); }
-
 		/** Get the editable groom types */
 		static const TArray<FEditableType>& GetEditableGroom(const FEditableGroomGroup& GroomGroup)
 		{
@@ -221,10 +231,10 @@ namespace UE::Groom
 		const TArray<float>& GetVertexKinematicWeights() const { return VertexKinematicWeights.Get().GetConstArray(); }
 		
 		/** Get the point bone indices */
-		const TArray<int32>& GetPointBoneIndices(const int32 PointIndex) const { return PointBoneIndices.Get()[PointIndex]; }
+		const FIntVector4& GetPointBoneIndices(const int32 PointIndex) const { return PointBoneIndices.Get()[PointIndex]; }
 
 		/** Get the point bone weights*/
-		const TArray<float>& GetPointBoneWeights(const int32 PointIndex) const { return PointBoneWeights.Get()[PointIndex]; }
+		const FVector4f& GetPointBoneWeights(const int32 PointIndex) const { return PointBoneWeights.Get()[PointIndex]; }
 
 		/** Get the object point samples */
 		const TArray<int32>& GetObjectPointSamples() const { return ObjectPointSamples.Get().GetConstArray(); }
@@ -232,14 +242,20 @@ namespace UE::Groom
 		/** Get the guide strand indices */
 		const TArray<int32>& GetCurveStrandIndices() const { return CurveStrandIndices.Get().GetConstArray(); }
 
+		/** Get the guide parent indices */
+		const TArray<int32>& GetCurveParentIndices() const { return CurveParentIndices.Get().GetConstArray(); }
+		
+		/** Get the guide lod indices */
+		const TArray<int32>& GetCurveLodIndices() const { return CurveLodIndices.Get().GetConstArray(); }
+
 		/** Set the vertex kinematic weights */
 		void SetVertexKinematicWeights(const TArray<float>& InKinematicWeights) { VertexKinematicWeights.Modify() = InKinematicWeights; }
 
 		/** Set the point bone indices */
-		void SetPointBoneIndices(const int32 PointIndex, const TArray<int32>& InBoneIndices) { return PointBoneIndices.ModifyAt(PointIndex,InBoneIndices); }
+		void SetPointBoneIndices(const int32 PointIndex, const FIntVector4& InBoneIndices) { return PointBoneIndices.ModifyAt(PointIndex,InBoneIndices); }
 
 		/** Set the point bone weights*/
-		void SetPointBoneWeights(const int32 PointIndex, const TArray<float>& InBoneWeights) { return PointBoneWeights.ModifyAt(PointIndex, InBoneWeights); }
+		void SetPointBoneWeights(const int32 PointIndex, const FVector4f& InBoneWeights) { return PointBoneWeights.ModifyAt(PointIndex, InBoneWeights); }
 		
 		/** Set the object point samples */
 		void SetObjectPointSamples(const TArray<int32>& NumPointSamples) { ObjectPointSamples.Modify() = NumPointSamples; }
@@ -247,6 +263,12 @@ namespace UE::Groom
 		/** Set the guide strand indices */
         void SetCurveStrandIndices(const TArray<int32>& StrandIndices) { CurveStrandIndices.Modify() = StrandIndices; }
 
+		/** Set the guide parent indices */
+		void SetCurveParentIndices(const TArray<int32>& ParentIndices) { CurveParentIndices.Modify() = ParentIndices; }
+
+		/** Set the guide lod indices */
+		void SetCurveLodIndices(const TArray<int32>& LodIndices) { CurveLodIndices.Modify() = LodIndices; }
+
 		/** Resize the geometry groups (vertex, points, edges) based on a new total number of points */
 		void ResizePointsGroups(const int32 NumPoints) const;
 
@@ -256,10 +278,10 @@ namespace UE::Groom
 		TManagedArrayAccessor<float> VertexKinematicWeights;
 
 		/** Point bone indices */
-		TManagedArrayAccessor<TArray<int32>> PointBoneIndices;
+		TManagedArrayAccessor<FIntVector4> PointBoneIndices;
 
 		/** Point bone weights */
-		TManagedArrayAccessor<TArray<float>> PointBoneWeights;
+		TManagedArrayAccessor<FVector4f> PointBoneWeights;
 
 		/** Object point samples */
 		TManagedArrayAccessor<int32> ObjectPointSamples;
@@ -267,13 +289,10 @@ namespace UE::Groom
 		/** Strand index from which the guide has been generated */
 		TManagedArrayAccessor<int32> CurveStrandIndices;
 
-		// /** Level index for the curves hierarchy */
-		// TManagedArrayAccessor<int32> CurveLevelIndices;
-		//
-		// /** Point (Level L) to point (Level L+1) interpolation indices+coords */
-		// TManagedArrayAccessor<FVector4f> PointInterpIndices;
-		//
-		// /** Point (Level L) to point (Level L+1) interpolation weights */
-		// TManagedArrayAccessor<FVector4f> PointInterpWeights;
+		/** Parent guide indices */
+		TManagedArrayAccessor<int32> CurveParentIndices;
+
+		/** Lod guide indices */
+		TManagedArrayAccessor<int32> CurveLodIndices;
 	};
 }
diff --git a/Engine/Plugins/Runtime/HairStrands/Source/HairStrandsDeformer/HairStrandsDeformer.Build.cs b/Engine/Plugins/Runtime/HairStrands/Source/HairStrandsDeformer/HairStrandsDeformer.Build.cs
--- a/Engine/Plugins/Runtime/HairStrands/Source/HairStrandsDeformer/HairStrandsDeformer.Build.cs
+++ b/Engine/Plugins/Runtime/HairStrands/Source/HairStrandsDeformer/HairStrandsDeformer.Build.cs
@@ -19,7 +19,9 @@ namespace UnrealBuildTool.Rules
 					"RHI",
 					"ComputeFramework",
 					"OptimusCore",
-					"HairStrandsCore"
+					"HairStrandsCore",
+					"HairStrandsSolver", 
+					"Chaos",
 				});
 
 			PrivateIncludePathModuleNames.AddRange(
diff --git a/Engine/Plugins/Runtime/HairStrands/Source/HairStrandsDeformer/Private/DeformerDataInterfaceGroom.cpp b/Engine/Plugins/Runtime/HairStrands/Source/HairStrandsDeformer/Private/DeformerDataInterfaceGroom.cpp
deleted file mode 100644
diff --git a/Engine/Plugins/Runtime/HairStrands/Source/HairStrandsDeformer/Private/DeformerDataInterfaceGroom.h b/Engine/Plugins/Runtime/HairStrands/Source/HairStrandsDeformer/Private/DeformerDataInterfaceGroom.h
deleted file mode 100644
diff --git a/Engine/Plugins/Runtime/HairStrands/Source/HairStrandsDeformer/Private/DeformerDataInterfaceGroomExec.cpp b/Engine/Plugins/Runtime/HairStrands/Source/HairStrandsDeformer/Private/DeformerDataInterfaceGroomExec.cpp
deleted file mode 100644
diff --git a/Engine/Plugins/Runtime/HairStrands/Source/HairStrandsDeformer/Private/DeformerDataInterfaceGroomExec.h b/Engine/Plugins/Runtime/HairStrands/Source/HairStrandsDeformer/Private/DeformerDataInterfaceGroomExec.h
deleted file mode 100644
diff --git a/Engine/Plugins/Runtime/HairStrands/Source/HairStrandsDeformer/Private/DeformerDataInterfaceGroomGuide.cpp b/Engine/Plugins/Runtime/HairStrands/Source/HairStrandsDeformer/Private/DeformerDataInterfaceGroomGuide.cpp
deleted file mode 100644
diff --git a/Engine/Plugins/Runtime/HairStrands/Source/HairStrandsDeformer/Private/DeformerDataInterfaceGroomGuide.h b/Engine/Plugins/Runtime/HairStrands/Source/HairStrandsDeformer/Private/DeformerDataInterfaceGroomGuide.h
deleted file mode 100644
diff --git a/Engine/Plugins/Runtime/HairStrands/Source/HairStrandsDeformer/Private/DeformerDataInterfaceGroomWrite.cpp b/Engine/Plugins/Runtime/HairStrands/Source/HairStrandsDeformer/Private/DeformerDataInterfaceGroomWrite.cpp
deleted file mode 100644
diff --git a/Engine/Plugins/Runtime/HairStrands/Source/HairStrandsDeformer/Private/DeformerDataInterfaceGroomWrite.h b/Engine/Plugins/Runtime/HairStrands/Source/HairStrandsDeformer/Private/DeformerDataInterfaceGroomWrite.h
deleted file mode 100644
diff --git a/Engine/Plugins/Runtime/HairStrands/Source/HairStrandsDeformer/Private/DeformerGroomAttributeRead.cpp b/Engine/Plugins/Runtime/HairStrands/Source/HairStrandsDeformer/Private/DeformerGroomAttributeRead.cpp
new file mode 100644
--- /dev/null
+++ b/Engine/Plugins/Runtime/HairStrands/Source/HairStrandsDeformer/Private/DeformerGroomAttributeRead.cpp
@@ -0,0 +1,651 @@
+// Copyright Epic Games, Inc. All Rights Reserved.
+
+#include "DeformerGroomAttributeRead.h"
+
+#include "ComputeFramework/ShaderParamTypeDefinition.h"
+#include "ShaderCompilerCore.h"
+#include "ShaderParameterMetadataBuilder.h"
+#include "GroomComponent.h"
+#include "GroomInstance.h"
+#include "RenderGraphBuilder.h"
+#include "GlobalRenderResources.h"
+#include "DeformerGroomDomainsSource.h"
+#include "DeformerGroomInterfaceUtils.h"
+#include "OptimusDataTypeRegistry.h"
+#include "GroomAsset.h"
+#include "OptimusNode.h"
+#include "SystemTextures.h"
+
+
+FString UOptimusGroomAttributeReadDataInterface::GetDisplayName() const
+{
+	return TEXT("Groom Attribute");
+}
+
+FName UOptimusGroomAttributeReadDataInterface::GetGroomAttributeName()
+{
+	return GET_MEMBER_NAME_CHECKED(UOptimusGroomAttributeReadDataInterface, GroomAttributeName);
+}
+
+TArray<FOptimusCDIPropertyPinDefinition> UOptimusGroomAttributeReadDataInterface::GetPropertyPinDefinitions() const
+{
+	TArray<FOptimusCDIPropertyPinDefinition> PropertyPinDefinitions;
+
+	const FOptimusDataTypeHandle NameType = FOptimusDataTypeRegistry::Get().FindType(*FNameProperty::StaticClass());
+	
+	PropertyPinDefinitions.Add(
+	{GetGroomAttributeName(), NameType}
+	);
+
+	return PropertyPinDefinitions;
+}
+
+#if WITH_EDITOR
+void UOptimusGroomAttributeReadDataInterface::PostEditChangeChainProperty(FPropertyChangedChainEvent& PropertyChangedEvent)
+{
+	Super::PostEditChangeChainProperty(PropertyChangedEvent);
+	
+	const FName PropertyName = (PropertyChangedEvent.Property ? PropertyChangedEvent.Property->GetFName() : NAME_None);
+	
+	if (PropertyChangedEvent.ChangeType == EPropertyChangeType::ValueSet)
+	{	
+		if (PropertyName == GET_MEMBER_NAME_STRING_CHECKED(UOptimusGroomAttributeReadDataInterface, GroomAttributeGroup) ||
+			PropertyName == GET_MEMBER_NAME_STRING_CHECKED(UOptimusGroomAttributeReadDataInterface, GroomAttributeType))
+		{
+			OnPinDefinitionChangedDelegate.Execute();
+		}
+	}
+}
+#endif
+
+FORCEINLINE void AddGroupDefinitions(const EOptimusGroomAttributeTypes GroomAttributeType, const FName& ContextName, const FString& CountFunctionName, TArray<FOptimusCDIPinDefinition>& Defs)
+{
+	Defs.Add({ "NumAttributeValues", CountFunctionName, false, {"NumAttributeValues"} });
+
+	// List all mapping between the deformer graph pin types and the hlsl ones
+	
+	if(GroomAttributeType == EOptimusGroomAttributeTypes::Bool)
+	{
+		Defs.Add({ "ValueTypedData", "ReadValueTypedBool", ContextName, CountFunctionName, false, "ValueTypedData" });
+	}
+	else if(GroomAttributeType == EOptimusGroomAttributeTypes::Int)
+	{
+		Defs.Add({ "ValueTypedData", "ReadValueTypedInt", ContextName, CountFunctionName, false, "ValueTypedData" });
+	}
+	else if(GroomAttributeType == EOptimusGroomAttributeTypes::IntVector2)
+	{
+		Defs.Add({ "ValueTypedData", "ReadValueTypedInt2", ContextName, CountFunctionName, false, "ValueTypedData" });
+	}
+	else if(GroomAttributeType == EOptimusGroomAttributeTypes::IntVector3)
+	{
+		Defs.Add({ "ValueTypedData", "ReadValueTypedInt3", ContextName, CountFunctionName, false, "ValueTypedData" });
+	}
+	else if(GroomAttributeType == EOptimusGroomAttributeTypes::IntVector4)
+	{
+		Defs.Add({ "ValueTypedData", "ReadValueTypedInt4", ContextName, CountFunctionName, false, "ValueTypedData" });
+	}
+	else if(GroomAttributeType == EOptimusGroomAttributeTypes::Uint)
+	{
+		Defs.Add({ "ValueTypedData", "ReadValueTypedUint", ContextName, CountFunctionName, false, "ValueTypedData" });
+	}
+	else if(GroomAttributeType == EOptimusGroomAttributeTypes::Float)
+	{
+		Defs.Add({ "ValueTypedData", "ReadValueTypedFloat", ContextName, CountFunctionName, false, "ValueTypedData" });
+	}
+	else if(GroomAttributeType == EOptimusGroomAttributeTypes::Vector2)
+	{
+		Defs.Add({ "ValueTypedData", "ReadValueTypedFloat2", ContextName, CountFunctionName, false, "ValueTypedData" });
+	}
+	else if(GroomAttributeType == EOptimusGroomAttributeTypes::Vector3)
+	{
+		Defs.Add({ "ValueTypedData", "ReadValueTypedFloat3", ContextName, CountFunctionName, false, "ValueTypedData" });
+	}
+	else if(GroomAttributeType == EOptimusGroomAttributeTypes::Vector4)
+	{
+		Defs.Add({ "ValueTypedData", "ReadValueTypedFloat4", ContextName, CountFunctionName, false, "ValueTypedData" });
+	}
+	else if(GroomAttributeType == EOptimusGroomAttributeTypes::LinearColor)
+	{
+		Defs.Add({ "ValueTypedData", "ReadValueTypedFloat4", ContextName, CountFunctionName, false, "ValueTypedData" });
+	}
+	else if(GroomAttributeType == EOptimusGroomAttributeTypes::Quat)
+	{
+		Defs.Add({ "ValueTypedData", "ReadValueTypedFloat4", ContextName, CountFunctionName, false, "ValueTypedData" });
+	}
+	else if(GroomAttributeType == EOptimusGroomAttributeTypes::Rotator)
+	{
+		Defs.Add({ "ValueTypedData", "ReadValueTypedFloat3", ContextName, CountFunctionName, false, "ValueTypedData" });
+	}
+	else if(GroomAttributeType == EOptimusGroomAttributeTypes::Transform)
+	{
+		Defs.Add({ "ValueTypedData", "ReadValueTypedFloat4x4", ContextName, CountFunctionName, false, "ValueTypedData" });
+	}
+	else if(GroomAttributeType == EOptimusGroomAttributeTypes::Matrix3x4)
+	{
+		Defs.Add({ "ValueTypedData", "ReadValueTypedFloat3x4", ContextName, CountFunctionName, false, "ValueTypedData" });
+	}
+}
+
+TArray<FOptimusCDIPinDefinition> UOptimusGroomAttributeReadDataInterface::GetPinDefinitions() const
+{
+	TArray<FOptimusCDIPinDefinition> Defs;
+	AddGroupDefinitions(GroomAttributeType, UOptimusGroomExecDataInterface::GetExecutionDomainName(GroomAttributeGroup), "ReadNumValues", Defs);
+	return Defs;
+}
+
+void UOptimusGroomAttributeReadDataInterface::RegisterPropertyChangeDelegatesForOwningNode(UOptimusNode* InNode)
+{
+	if(InNode)
+	{
+		OnPinDefinitionChangedDelegate.BindUObject(InNode, &UOptimusNode::RecreatePinsFromPinDefinitions);
+	}
+}
+ 
+TSubclassOf<UActorComponent> UOptimusGroomAttributeReadDataInterface::GetRequiredComponentClass() const
+{
+	return UMeshComponent::StaticClass();
+}
+
+void UOptimusGroomAttributeReadDataInterface::GetSupportedInputs(TArray<FShaderFunctionDefinition>& OutFunctions) const
+{
+	OutFunctions.AddDefaulted_GetRef()
+		.SetName(TEXT("ReadNumValues"))
+		.AddReturnType(EShaderFundamentalType::Uint);
+
+	// List all the hlsl return types
+	
+	OutFunctions.AddDefaulted_GetRef()
+		.SetName(TEXT("ReadValueTypedInt"))
+		.AddReturnType(EShaderFundamentalType::Int)
+		.AddParam(EShaderFundamentalType::Uint);
+	
+	OutFunctions.AddDefaulted_GetRef()
+		.SetName(TEXT("ReadValueTypedInt2"))
+		.AddReturnType(EShaderFundamentalType::Int, 2)
+		.AddParam(EShaderFundamentalType::Uint);
+	
+	OutFunctions.AddDefaulted_GetRef()
+		.SetName(TEXT("ReadValueTypedInt3"))
+		.AddReturnType(EShaderFundamentalType::Int, 3)
+		.AddParam(EShaderFundamentalType::Uint);
+
+	OutFunctions.AddDefaulted_GetRef()
+		.SetName(TEXT("ReadValueTypedInt4"))
+		.AddReturnType(EShaderFundamentalType::Int, 4)
+		.AddParam(EShaderFundamentalType::Uint);
+
+	OutFunctions.AddDefaulted_GetRef()
+		.SetName(TEXT("ReadValueTypedUint"))
+		.AddReturnType(EShaderFundamentalType::Uint)
+		.AddParam(EShaderFundamentalType::Uint);
+
+	OutFunctions.AddDefaulted_GetRef()
+		.SetName(TEXT("ReadValueTypedBool"))
+		.AddReturnType(EShaderFundamentalType::Bool)
+		.AddParam(EShaderFundamentalType::Uint);
+
+	OutFunctions.AddDefaulted_GetRef()
+		.SetName(TEXT("ReadValueTypedFloat"))
+		.AddReturnType(EShaderFundamentalType::Float)
+		.AddParam(EShaderFundamentalType::Uint);
+
+	OutFunctions.AddDefaulted_GetRef()
+		.SetName(TEXT("ReadValueTypedFloat2"))
+		.AddReturnType(EShaderFundamentalType::Float, 2)
+		.AddParam(EShaderFundamentalType::Uint);
+
+	OutFunctions.AddDefaulted_GetRef()
+		.SetName(TEXT("ReadValueTypedFloat3"))
+		.AddReturnType(EShaderFundamentalType::Float, 3)
+		.AddParam(EShaderFundamentalType::Uint);
+
+	OutFunctions.AddDefaulted_GetRef()
+		.SetName(TEXT("ReadValueTypedFloat4"))
+		.AddReturnType(EShaderFundamentalType::Float, 4)
+		.AddParam(EShaderFundamentalType::Uint);
+
+	OutFunctions.AddDefaulted_GetRef()
+		.SetName(TEXT("ReadValueTypedFloat4x4"))
+		.AddReturnType(EShaderFundamentalType::Float, 4, 4)
+		.AddParam(EShaderFundamentalType::Uint);
+
+	OutFunctions.AddDefaulted_GetRef()
+		.SetName(TEXT("ReadValueTypedFloat3x4"))
+		.AddReturnType(EShaderFundamentalType::Float, 3, 4)
+		.AddParam(EShaderFundamentalType::Uint);
+}
+
+BEGIN_SHADER_PARAMETER_STRUCT(FOptimusGroomAttributeReadParameters, )
+	SHADER_PARAMETER(uint32, NumValues)
+	SHADER_PARAMETER_RDG_BUFFER_SRV(Buffer<float>, ValueTypedData)
+END_SHADER_PARAMETER_STRUCT()
+
+void UOptimusGroomAttributeReadDataInterface::GetShaderParameters(TCHAR const* UID, FShaderParametersMetadataBuilder& InOutBuilder, FShaderParametersMetadataAllocations& InOutAllocations) const
+{
+	InOutBuilder.AddNestedStruct<FOptimusGroomAttributeReadParameters>(UID);
+}
+
+TCHAR const* UOptimusGroomAttributeReadDataInterface::TemplateFilePath = TEXT("/Plugin/Runtime/HairStrands/Private/Deformers/DeformerGroomAttributeRead.ush");
+
+TCHAR const* UOptimusGroomAttributeReadDataInterface::GetShaderVirtualPath() const
+{
+	return TemplateFilePath;
+}
+
+void UOptimusGroomAttributeReadDataInterface::GetShaderHash(FString& InOutKey) const
+{
+	GetShaderFileHash(TemplateFilePath, EShaderPlatform::SP_PCD3D_SM5).AppendString(InOutKey);
+}
+
+void UOptimusGroomAttributeReadDataInterface::GetPermutations(FComputeKernelPermutationVector& OutPermutationVector) const
+{}
+
+void UOptimusGroomAttributeReadDataInterface::GetHLSL(FString& OutHLSL, FString const& InDataInterfaceName) const
+{
+	TMap<FString, FStringFormatArg> TemplateArgs =
+	{
+		{ TEXT("DataInterfaceName"), InDataInterfaceName },
+	};
+
+	FString TemplateFile;
+	LoadShaderSourceFile(TemplateFilePath, EShaderPlatform::SP_PCD3D_SM5, &TemplateFile, nullptr);
+	OutHLSL += FString::Format(*TemplateFile, TemplateArgs);
+}
+
+UComputeDataProvider* UOptimusGroomAttributeReadDataInterface::CreateDataProvider(TObjectPtr<UObject> InBinding, uint64 InInputMask, uint64 InOutputMask) const
+{
+	UOptimusGroomAttributeReadDataProvider* Provider = NewObject<UOptimusGroomAttributeReadDataProvider>();
+	Provider->MeshComponent = Cast<UMeshComponent>(InBinding);
+	Provider->GroomAttributeGroup = GroomAttributeGroup;
+	Provider->GroomAttributeName = GroomAttributeName;
+	Provider->GroomAttributeType = GroomAttributeType;
+	return Provider;
+}
+
+FComputeDataProviderRenderProxy* UOptimusGroomAttributeReadDataProvider::GetRenderProxy()
+{
+	return new FOptimusGroomAttributeReadProviderProxy(MeshComponent, GroomAttributeName, GroomAttributeGroup, GroomAttributeType);
+}
+
+FOptimusGroomAttributeReadProviderProxy::FOptimusGroomAttributeReadProviderProxy(UMeshComponent* MeshComponent,
+	const FName& InAtributeName, const EOptimusGroomExecDomain InAttributeGroup, const EOptimusGroomAttributeTypes InAttributeType)
+{
+	GroomAttributeName = InAtributeName;;
+	GroomAttributeGroup = InAttributeGroup;
+	GroomAttributeType = InAttributeType;
+	
+	TArray<const UGroomComponent*> GroomComponents;
+	UE::Groom::Private::GatherGroomComponents(MeshComponent, GroomComponents);
+	UE::Groom::Private::GroomComponentsToInstances(GroomComponents, GroupInstances);
+	UE::Groom::Private::GetGroomInvocationElementGroups(GroomComponents,UOptimusGroomExecDataInterface::GetExecutionDomainName(GroomAttributeGroup), GroupElements);
+}
+
+bool FOptimusGroomAttributeReadProviderProxy::IsValid(FValidationData const& InValidationData) const
+{
+	if (InValidationData.ParameterStructSize != sizeof(FParameters))
+	{
+		return false;
+	}
+	if (InValidationData.NumInvocations != GroupInstances.Num())
+	{
+		return false;
+	}
+	const FName GroupName = UOptimusGroomExecDataInterface::GetExecutionDomainName(GroomAttributeGroup);
+	
+	if((GroomAttributeGroup == EOptimusGroomExecDomain::GuidesPoints) || (GroomAttributeGroup == EOptimusGroomExecDomain::GuidesEdges) ||
+		(GroomAttributeGroup == EOptimusGroomExecDomain::GuidesCurves) || (GroomAttributeGroup == EOptimusGroomExecDomain::GuidesObjects))
+	{
+		if(!UE::Groom::Private::HaveGuidesInstanceResources(GroupInstances))
+		{
+			return false;
+		}
+	}
+	else if ((GroomAttributeGroup == EOptimusGroomExecDomain::ControlPoint) || (GroomAttributeGroup == EOptimusGroomExecDomain::StrandsEdges) ||
+		(GroomAttributeGroup == EOptimusGroomExecDomain::Curve) || (GroomAttributeGroup == EOptimusGroomExecDomain::StrandsObjects))
+	{
+		if(!UE::Groom::Private::HaveStrandsInstanceResources(GroupInstances))
+		{
+			return false;
+		}
+	}
+	for (const TPair<UGroomAsset*, UE::Groom::Private::FGroupElements>& AssetPair : GroupElements)
+	{
+		if(!AssetPair.Key)
+		{
+			return false;
+		}
+		
+		if(!AssetPair.Key->GetDataflowSettings().GetRestCollection()->HasAttribute(GroomAttributeName, GroupName))
+		{
+			return false;
+		}
+	}
+	
+	return true;
+}
+
+void FOptimusGroomAttributeReadProviderProxy::GatherPermutations(FPermutationData& InOutPermutationData) const
+{
+}
+
+template<typename DataType>
+FORCEINLINE void ConvertAttributeData(const DataType& InputData, TArray<float>& DataValues, int32& DataIndex)
+{
+	DataValues[DataIndex++] = InputData;
+}
+
+template<typename DataType>
+FORCEINLINE int32 NumAttributeValues()
+{
+	return 1;;
+}
+
+template<>
+FORCEINLINE void ConvertAttributeData<FTransform3f>(const FTransform3f& InputData, TArray<float>& DataValues, int32& DataIndex)
+{
+	const FMatrix44f InputMatrix = InputData.ToMatrixWithScale();
+	DataValues[DataIndex++] = InputMatrix.M[0][0];
+	DataValues[DataIndex++] = InputMatrix.M[0][1];
+	DataValues[DataIndex++] = InputMatrix.M[0][2];
+	DataValues[DataIndex++] = InputMatrix.M[0][3];
+
+	DataValues[DataIndex++] = InputMatrix.M[1][0];
+	DataValues[DataIndex++] = InputMatrix.M[1][1];
+	DataValues[DataIndex++] = InputMatrix.M[1][2];
+	DataValues[DataIndex++] = InputMatrix.M[1][3];
+
+	DataValues[DataIndex++] = InputMatrix.M[2][0];
+	DataValues[DataIndex++] = InputMatrix.M[2][1];
+	DataValues[DataIndex++] = InputMatrix.M[2][2];
+	DataValues[DataIndex++] = InputMatrix.M[2][3];
+
+	DataValues[DataIndex++] = InputMatrix.M[3][0];
+	DataValues[DataIndex++] = InputMatrix.M[3][1];
+	DataValues[DataIndex++] = InputMatrix.M[3][2];
+	DataValues[DataIndex++] = InputMatrix.M[3][3];
+}
+
+template<>
+FORCEINLINE int32 NumAttributeValues<FTransform3f>()
+{
+	return 16;
+}
+
+template<>
+FORCEINLINE void ConvertAttributeData<FLinearColor>(const FLinearColor& InputData, TArray<float>& DataValues, int32& DataIndex)
+{
+	DataValues[DataIndex++] = InputData.R;
+	DataValues[DataIndex++] = InputData.G;
+	DataValues[DataIndex++] = InputData.B;
+	DataValues[DataIndex++] = InputData.A;
+}
+
+template<>
+FORCEINLINE int32 NumAttributeValues<FLinearColor>()
+{
+	return 4;
+}
+
+template<>
+FORCEINLINE void ConvertAttributeData<FQuat4f>(const FQuat4f& InputData, TArray<float>& DataValues, int32& DataIndex)
+{
+	DataValues[DataIndex++] = InputData.X;
+	DataValues[DataIndex++] = InputData.Y;
+	DataValues[DataIndex++] = InputData.Z;
+	DataValues[DataIndex++] = InputData.W;
+}
+
+template<>
+FORCEINLINE int32 NumAttributeValues<FQuat4f>()
+{
+	return 4;
+}
+
+template<>
+FORCEINLINE void ConvertAttributeData<FRotator3f>(const FRotator3f& InputData, TArray<float>& DataValues, int32& DataIndex)
+{
+	DataValues[DataIndex++] = InputData.Pitch;
+	DataValues[DataIndex++] = InputData.Yaw;
+	DataValues[DataIndex++] = InputData.Roll;
+}
+
+template<>
+FORCEINLINE int32 NumAttributeValues<FRotator3f>()
+{
+	return 3;
+}
+
+template<>
+FORCEINLINE void ConvertAttributeData<FIntVector4>(const FIntVector4& InputData, TArray<float>& DataValues, int32& DataIndex)
+{
+	DataValues[DataIndex++] = InputData.X;
+	DataValues[DataIndex++] = InputData.Y;
+	DataValues[DataIndex++] = InputData.Z;
+	DataValues[DataIndex++] = InputData.W;
+}
+
+template<>
+FORCEINLINE int32 NumAttributeValues<FIntVector4>()
+{
+	return 4;
+}
+
+template<>
+FORCEINLINE void ConvertAttributeData<FIntVector3>(const FIntVector3& InputData, TArray<float>& DataValues, int32& DataIndex)
+{
+	DataValues[DataIndex++] = InputData.X;
+	DataValues[DataIndex++] = InputData.Y;
+	DataValues[DataIndex++] = InputData.Z;
+}
+
+template<>
+FORCEINLINE int32 NumAttributeValues<FIntVector3>()
+{
+	return 3;
+}
+
+template<>
+FORCEINLINE void ConvertAttributeData<FIntVector2>(const FIntVector2& InputData, TArray<float>& DataValues, int32& DataIndex)
+{
+	DataValues[DataIndex++] = InputData.X;
+	DataValues[DataIndex++] = InputData.Y;
+}
+
+template<>
+FORCEINLINE int32 NumAttributeValues<FIntVector2>()
+{
+	return 2;
+}
+
+template<>
+FORCEINLINE void ConvertAttributeData<FVector4f>(const FVector4f& InputData, TArray<float>& DataValues, int32& DataIndex)
+{
+	DataValues[DataIndex++] = InputData.X;
+	DataValues[DataIndex++] = InputData.Y;
+	DataValues[DataIndex++] = InputData.Z;
+	DataValues[DataIndex++] = InputData.W;
+}
+
+template<>
+FORCEINLINE int32 NumAttributeValues<FVector4f>()
+{
+	return 4;
+}
+
+template<>
+FORCEINLINE void ConvertAttributeData<FVector3f>(const FVector3f& InputData, TArray<float>& DataValues, int32& DataIndex)
+{
+	DataValues[DataIndex++] = InputData.X;
+	DataValues[DataIndex++] = InputData.Y;
+	DataValues[DataIndex++] = InputData.Z;
+}
+
+template<>
+FORCEINLINE int32 NumAttributeValues<FVector3f>()
+{
+	return 3;
+}
+
+template<>
+FORCEINLINE void ConvertAttributeData<FVector2f>(const FVector2f& InputData, TArray<float>& DataValues, int32& DataIndex)
+{
+	DataValues[DataIndex++] = InputData.X;
+	DataValues[DataIndex++] = InputData.Y;
+}
+
+template<>
+FORCEINLINE int32 NumAttributeValues<FVector2f>()
+{
+	return 2;
+}
+
+template<typename DataType>
+FORCEINLINE void CreateInternalBuffers(const UGroomAsset* GroomAsset, const UE::Groom::Private::FGroupElements& GroupElements, const FName& AttributeName,
+	const FName& GroupName, FRDGBuilder& GraphBuilder, TArray<FRDGBufferRef>& AttributeBuffers)
+{
+	if(GroomAsset && GroomAsset->GetDataflowSettings().GetRestCollection() &&
+     		GroomAsset->GetDataflowSettings().GetRestCollection()->HasAttribute(AttributeName, GroupName))
+    {
+     	const TManagedArray<DataType>& AssetAttributeValues =
+     			GroomAsset->GetDataflowSettings().GetRestCollection()->GetAttribute<DataType>(AttributeName,GroupName);
+     		
+     	for(int32 GroupIndex : GroupElements.GroupIndices)
+     	{
+     		const uint32 NumAttributes = GroupElements.GroupOffsets[GroupIndex+1]-GroupElements.GroupOffsets[GroupIndex];
+     			
+     		TArray<float> GroupAttributeValues;
+     		GroupAttributeValues.SetNum(NumAttributes*NumAttributeValues<DataType>());
+ 
+     		int32 DataIndex = 0;
+     		for(uint32 AttributeIndex = 0; AttributeIndex < NumAttributes; ++AttributeIndex)
+     		{
+     			ConvertAttributeData<DataType>(AssetAttributeValues[AttributeIndex+GroupElements.GroupOffsets[GroupIndex]], GroupAttributeValues, DataIndex);
+     		}
+ 
+     		const FString BufferName = FString("Hair.Deformer.") + AttributeName.ToString();
+     		FRDGBufferRef& GroupAttributeBuffer = AttributeBuffers.Add_GetRef( GraphBuilder.CreateBuffer(FRDGBufferDesc::CreateStructuredDesc(sizeof(float),
+     		 GroupAttributeValues.Num()), *BufferName));
+     		GraphBuilder.QueueBufferUpload(GroupAttributeBuffer, GroupAttributeValues.GetData(), sizeof(float) * GroupAttributeValues.Num(), ERDGInitialDataFlags::None);
+     	}
+    }
+}
+
+FORCEINLINE void AddGroupResources(UGroomAsset* GroomAsset, const UE::Groom::Private::FGroupElements& GroupElements, const EOptimusGroomAttributeTypes GroomAttributeType,
+	const FName& GroomAttributeName,  const FName& GroomAttributeGroup, FRDGBuilder& GraphBuilder, TArray<FRDGBufferRef>& GroomAttributeBuffers)
+{
+	if(GroomAttributeType == EOptimusGroomAttributeTypes::Bool)
+	{
+		CreateInternalBuffers<bool>(GroomAsset, GroupElements, GroomAttributeName,
+			GroomAttributeGroup, GraphBuilder, GroomAttributeBuffers);
+	}
+	else if(GroomAttributeType == EOptimusGroomAttributeTypes::Int || GroomAttributeType == EOptimusGroomAttributeTypes::Uint)
+	{
+		CreateInternalBuffers<int32>(GroomAsset, GroupElements, GroomAttributeName,
+			GroomAttributeGroup, GraphBuilder, GroomAttributeBuffers);
+	}
+	else if(GroomAttributeType == EOptimusGroomAttributeTypes::IntVector2)
+	{
+		CreateInternalBuffers<FIntVector2>(GroomAsset, GroupElements, GroomAttributeName,
+			GroomAttributeGroup, GraphBuilder, GroomAttributeBuffers);
+	}
+	else if(GroomAttributeType == EOptimusGroomAttributeTypes::IntVector3)
+	{
+		CreateInternalBuffers<FIntVector3>(GroomAsset, GroupElements, GroomAttributeName,
+			GroomAttributeGroup, GraphBuilder, GroomAttributeBuffers);
+	}
+	else if(GroomAttributeType == EOptimusGroomAttributeTypes::IntVector4)
+	{
+		CreateInternalBuffers<FIntVector4>(GroomAsset, GroupElements, GroomAttributeName,
+			GroomAttributeGroup, GraphBuilder, GroomAttributeBuffers);
+	}
+	else if(GroomAttributeType == EOptimusGroomAttributeTypes::Float)
+	{
+		CreateInternalBuffers<float>(GroomAsset, GroupElements, GroomAttributeName,
+			GroomAttributeGroup, GraphBuilder, GroomAttributeBuffers);
+	}
+	else if(GroomAttributeType == EOptimusGroomAttributeTypes::Vector2)
+	{
+		CreateInternalBuffers<FVector2f>(GroomAsset, GroupElements,  GroomAttributeName,
+			GroomAttributeGroup, GraphBuilder, GroomAttributeBuffers);
+	}
+	else if(GroomAttributeType == EOptimusGroomAttributeTypes::Vector3 || GroomAttributeType == EOptimusGroomAttributeTypes::Rotator)
+	{
+		CreateInternalBuffers<FVector3f>(GroomAsset, GroupElements, GroomAttributeName,
+			GroomAttributeGroup, GraphBuilder, GroomAttributeBuffers);
+	}
+	else if(GroomAttributeType == EOptimusGroomAttributeTypes::Vector4)
+	{
+		CreateInternalBuffers<FVector4f>(GroomAsset, GroupElements, GroomAttributeName,
+			GroomAttributeGroup, GraphBuilder, GroomAttributeBuffers);
+	}
+	else if(GroomAttributeType == EOptimusGroomAttributeTypes::LinearColor)
+	{
+		CreateInternalBuffers<FLinearColor>(GroomAsset, GroupElements, GroomAttributeName,
+			GroomAttributeGroup, GraphBuilder, GroomAttributeBuffers);
+	}
+	else if(GroomAttributeType == EOptimusGroomAttributeTypes::Quat)
+	{
+		CreateInternalBuffers<FQuat4f>(GroomAsset, GroupElements, GroomAttributeName,
+			GroomAttributeGroup, GraphBuilder, GroomAttributeBuffers);
+	}
+	else if(GroomAttributeType == EOptimusGroomAttributeTypes::Transform || GroomAttributeType == EOptimusGroomAttributeTypes::Matrix3x4)
+	{
+		CreateInternalBuffers<FTransform3f>(GroomAsset, GroupElements, GroomAttributeName,
+			GroomAttributeGroup, GraphBuilder, GroomAttributeBuffers);
+	}
+}
+
+void FOptimusGroomAttributeReadProviderProxy::CreateInternalBuffers(FRDGBuilder& GraphBuilder)
+{
+	for (const TPair<UGroomAsset*,UE::Groom::Private::FGroupElements>& AssetPair : GroupElements)
+	{
+		AddGroupResources(AssetPair.Key, AssetPair.Value, GroomAttributeType, GroomAttributeName,
+				UOptimusGroomExecDataInterface::GetExecutionDomainName(GroomAttributeGroup),GraphBuilder,
+				AttributeValuesBuffers);
+	}
+}
+
+void FOptimusGroomAttributeReadProviderProxy::AllocateResources(FRDGBuilder& GraphBuilder, FAllocationData const& InAllocationData)
+{
+	if (!FallbackStructuredSRV) { FallbackStructuredSRV  = GraphBuilder.CreateSRV(GSystemTextures.GetDefaultStructuredBuffer(GraphBuilder, 16u)); }
+	
+	CreateInternalBuffers(GraphBuilder);
+	
+	AttributeValuesResources.Reset();
+	for (const TPair<UGroomAsset*,UE::Groom::Private::FGroupElements>& AssetPair : GroupElements)
+	{
+		for(int32 GroupIndex = 0; GroupIndex < AttributeValuesBuffers.Num(); ++GroupIndex)
+		{
+			FRDGBufferSRVRef& AttributeValuesResource = AttributeValuesResources.AddDefaulted_GetRef();
+			AttributeValuesResource = GraphBuilder.CreateSRV(AttributeValuesBuffers[GroupIndex]);
+		}
+	}
+}
+
+void FOptimusGroomAttributeReadProviderProxy::GatherDispatchData(FDispatchData const& InDispatchData)
+{
+	check(InDispatchData.NumInvocations == GroupInstances.Num());
+	check(InDispatchData.NumInvocations == AttributeValuesResources.Num());
+	
+	const TStridedView<FParameters> ParameterArray = MakeStridedParameterView<FParameters>(InDispatchData);
+
+	int32 InvocationIndex = 0;
+	for (const TPair<UGroomAsset*,UE::Groom::Private::FGroupElements>& AssetPair : GroupElements)
+	{
+		for(int32 GroupIndex : AssetPair.Value.GroupIndices)
+		{
+			FParameters& Parameters = ParameterArray[InvocationIndex];
+			if (const FHairGroupInstance* GroupInstance = GroupInstances[InvocationIndex])
+			{
+				Parameters.NumValues = AssetPair.Value.GroupOffsets[GroupIndex+1]-AssetPair.Value.GroupOffsets[GroupIndex];
+				Parameters.ValueTypedData = AttributeValuesResources[InvocationIndex];
+			}
+			else
+			{
+				Parameters.NumValues = 0;
+				Parameters.ValueTypedData = FallbackStructuredSRV;
+			}
+			++InvocationIndex;
+		}
+	}
+}
\ No newline at end of file
diff --git a/Engine/Plugins/Runtime/HairStrands/Source/HairStrandsDeformer/Private/DeformerGroomAttributeRead.h b/Engine/Plugins/Runtime/HairStrands/Source/HairStrandsDeformer/Private/DeformerGroomAttributeRead.h
new file mode 100644
--- /dev/null
+++ b/Engine/Plugins/Runtime/HairStrands/Source/HairStrandsDeformer/Private/DeformerGroomAttributeRead.h
@@ -0,0 +1,166 @@
+// Copyright Epic Games, Inc. All Rights Reserved.
+
+#pragma once
+
+#include "DeformerGroomDomainsExec.h"
+#include "RenderGraphDefinitions.h"
+#include "OptimusComputeDataInterface.h"
+#include "HairStrandsInterface.h"
+#include "ComputeFramework/ComputeDataProvider.h"
+#include "DeformerGroomAttributeRead.generated.h"
+
+namespace UE::Groom::Private
+{
+	struct FGroupElements;
+}
+
+class FOptimusGroomAttributeReadParameters;
+class UGroomComponent;
+class UGroomAsset;
+struct FHairGroupInstance;
+class UMeshComponent;
+
+/** List of attribute types on which you could paint in dataflow */
+UENUM()
+enum class EOptimusGroomAttributeTypes : uint8
+{
+	None = 0 UMETA(Hidden),
+	Bool UMETA(DisplayName = "Bool"),
+	Int UMETA(DisplayName = "Int"),
+	IntVector2 UMETA(DisplayName = "Int Vector 2"),
+	IntVector3 UMETA(DisplayName = "Int Vector 3"),
+	IntVector4 UMETA(DisplayName = "Int Vector 4"),
+	Uint UMETA(DisplayName = "Uint"),
+	Float UMETA(DisplayName = "Float"),
+	Vector2 UMETA(DisplayName = "Vector 2"),
+	Vector3 UMETA(DisplayName = "Vector 3"),
+	Vector4 UMETA(DisplayName = "Vector 4"),
+	LinearColor UMETA(DisplayName = "Linear Color"),
+	Quat UMETA(DisplayName = "Quat"),
+	Rotator UMETA(DisplayName = "Rotator"),
+	Transform UMETA(DisplayName = "Transform"),
+	Matrix3x4 UMETA(DisplayName = "Matrix 3x4"),
+};
+
+/** Compute Framework Data Interface for reading groom strands. */
+UCLASS(Category = ComputeFramework)
+class UOptimusGroomAttributeReadDataInterface : public UOptimusComputeDataInterface
+{
+	GENERATED_BODY()
+
+public:
+	//~ Begin UOptimusComputeDataInterface Interface
+	virtual FString GetDisplayName() const override;
+	virtual TArray<FOptimusCDIPinDefinition> GetPinDefinitions() const override;
+	virtual TSubclassOf<UActorComponent> GetRequiredComponentClass() const override;
+	virtual TArray<FOptimusCDIPropertyPinDefinition> GetPropertyPinDefinitions() const override;
+	virtual bool CanPinDefinitionChange() override { return true; };
+	virtual void RegisterPropertyChangeDelegatesForOwningNode(UOptimusNode* InNode) override;
+#if WITH_EDITOR
+	virtual void PostEditChangeChainProperty(FPropertyChangedChainEvent& PropertyChangedEvent) override;
+#endif
+	//~ End UOptimusComputeDataInterface Interface
+	
+	//~ Begin UComputeDataInterface Interface
+	virtual TCHAR const* GetClassName() const override { return TEXT("GroomAttributeRead"); }
+	virtual void GetSupportedInputs(TArray<FShaderFunctionDefinition>& OutFunctions) const override;
+	virtual void GetShaderParameters(TCHAR const* UID, FShaderParametersMetadataBuilder& InOutBuilder, FShaderParametersMetadataAllocations& InOutAllocations) const override;
+	virtual TCHAR const* GetShaderVirtualPath() const override;
+	virtual void GetPermutations(FComputeKernelPermutationVector& OutPermutationVector) const override;
+	virtual void GetShaderHash(FString& InOutKey) const override;
+	virtual void GetHLSL(FString& OutHLSL, FString const& InDataInterfaceName) const override;
+	virtual UComputeDataProvider* CreateDataProvider(TObjectPtr<UObject> InBinding, uint64 InInputMask, uint64 InOutputMask) const override;
+	//~ End UComputeDataInterface Interface
+
+private:
+	/** File holding the hlsl implementation */
+	static TCHAR const* TemplateFilePath;
+
+	static FName GetGroomAttributeName();
+
+	/** Attribute name */
+	UPROPERTY(EditAnywhere, Category = "Groom Attribute")
+	FName GroomAttributeName = NAME_None;
+
+	/** Attribute group */
+	UPROPERTY(EditAnywhere, Category = "Groom Attribute")
+	EOptimusGroomExecDomain GroomAttributeGroup = EOptimusGroomExecDomain::GuidesPoints;
+
+	/** Attribute type */
+	UPROPERTY(EditAnywhere, Category = "Groom Attribute")
+	EOptimusGroomAttributeTypes GroomAttributeType = EOptimusGroomAttributeTypes::Float;
+
+	/** Delegate to change the pin definition context based on the group selected */
+	FOnPinDefinitionChanged OnPinDefinitionChangedDelegate;
+};
+
+/** Compute Framework Data Provider for reading groom strands. */
+UCLASS(BlueprintType, editinlinenew, Category = ComputeFramework)
+class UOptimusGroomAttributeReadDataProvider : public UComputeDataProvider
+{
+	GENERATED_BODY()
+
+public:
+	UPROPERTY(EditAnywhere, BlueprintReadWrite, Category = Binding)
+	TObjectPtr<UMeshComponent> MeshComponent = nullptr;
+
+	/** Attribute name */
+	UPROPERTY(EditAnywhere, Category = "Groom Attribute")
+	FName GroomAttributeName = NAME_None;
+
+	/** Attribute group */
+	UPROPERTY(EditAnywhere, Category = "Groom Attribute")
+	EOptimusGroomExecDomain GroomAttributeGroup = EOptimusGroomExecDomain::GuidesPoints;
+
+	/** Attribute type */
+	UPROPERTY(EditAnywhere, Category = "Groom Attribute")
+	EOptimusGroomAttributeTypes GroomAttributeType = EOptimusGroomAttributeTypes::Float;
+
+	//~ Begin UComputeDataProvider Interface
+	virtual FComputeDataProviderRenderProxy* GetRenderProxy() override;
+	//~ End UComputeDataProvider Interface
+};
+
+class FOptimusGroomAttributeReadProviderProxy : public FComputeDataProviderRenderProxy
+{
+public:
+	FOptimusGroomAttributeReadProviderProxy(UMeshComponent* InGroomComponent,
+		const FName& InAtributeName, const EOptimusGroomExecDomain InAttributeGroup, const EOptimusGroomAttributeTypes InAttributeType);
+
+	//~ Begin FComputeDataProviderRenderProxy Interface
+	virtual bool IsValid(FValidationData const& InValidationData) const override;
+	virtual void AllocateResources(FRDGBuilder& GraphBuilder, FAllocationData const& InAllocationData) override;
+	virtual void GatherPermutations(FPermutationData& InOutPermutationData) const override;
+	virtual void GatherDispatchData(FDispatchData const& InDispatchData) override;
+	//~ End FComputeDataProviderRenderProxy Interface
+
+private:
+	using FParameters = FOptimusGroomAttributeReadParameters;
+
+	/** Create internal buffers */
+	void CreateInternalBuffers(FRDGBuilder& GraphBuilder);
+
+	/** List of instances (invocations) used in that data interface */
+	TArray<const FHairGroupInstance*> GroupInstances;
+
+	/** Groom assets group indices the proxy instance is linked to */
+	TArray<TPair<UGroomAsset*, UE::Groom::Private::FGroupElements>> GroupElements;
+
+	/** Groom attribute buffers built from the deformer */
+	TArray<FRDGBufferRef> AttributeValuesBuffers;
+	
+	/** Attribute Resources used to dispatch CS on GPU */
+	TArray<FRDGBufferSRVRef> AttributeValuesResources;
+
+	/** Attribute name */
+	FName GroomAttributeName = NAME_None;
+
+	/** Attribute group */
+	EOptimusGroomExecDomain GroomAttributeGroup = EOptimusGroomExecDomain::GuidesPoints;
+
+	/** Attribute type */
+	EOptimusGroomAttributeTypes GroomAttributeType = EOptimusGroomAttributeTypes::Float;
+
+	/** Fallback resources */
+	FRDGBufferSRVRef FallbackStructuredSRV = nullptr;
+};
diff --git a/Engine/Plugins/Runtime/HairStrands/Source/HairStrandsDeformer/Private/DeformerGroomComponentSource.cpp b/Engine/Plugins/Runtime/HairStrands/Source/HairStrandsDeformer/Private/DeformerGroomComponentSource.cpp
deleted file mode 100644
diff --git a/Engine/Plugins/Runtime/HairStrands/Source/HairStrandsDeformer/Private/DeformerGroomComponentSource.h b/Engine/Plugins/Runtime/HairStrands/Source/HairStrandsDeformer/Private/DeformerGroomComponentSource.h
deleted file mode 100644
diff --git a/Engine/Plugins/Runtime/HairStrands/Source/HairStrandsDeformer/Private/DeformerGroomDomainsExec.cpp b/Engine/Plugins/Runtime/HairStrands/Source/HairStrandsDeformer/Private/DeformerGroomDomainsExec.cpp
new file mode 100644
--- /dev/null
+++ b/Engine/Plugins/Runtime/HairStrands/Source/HairStrandsDeformer/Private/DeformerGroomDomainsExec.cpp
@@ -0,0 +1,199 @@
+// Copyright Epic Games, Inc. All Rights Reserved.
+
+#include "DeformerGroomDomainsExec.h"
+
+#include "DeformerGroomDomainsSource.h"
+#include "DeformerGroomInterfaceUtils.h"
+#include "GroomComponent.h"
+#include "OptimusExecutionDomain.h"
+#include "ComputeFramework/ShaderParamTypeDefinition.h"
+#include "ShaderCompilerCore.h"
+#include "ShaderParameterMetadataBuilder.h"
+#include "SkeletalRenderPublic.h"
+
+FString UOptimusGroomExecDataInterface::GetDisplayName() const
+{
+	return TEXT("Execute Groom");
+}
+
+FName UOptimusGroomExecDataInterface::GetCategory() const
+{
+	return CategoryName::ExecutionDataInterfaces;
+}
+
+TArray<FOptimusCDIPinDefinition> UOptimusGroomExecDataInterface::GetPinDefinitions() const
+{
+	TArray<FOptimusCDIPinDefinition> Defs;
+	Defs.Add({ "NumThreads", "ReadNumThreads" });
+	return Defs;
+}
+
+TSubclassOf<UActorComponent> UOptimusGroomExecDataInterface::GetRequiredComponentClass() const
+{
+	return UMeshComponent::StaticClass();
+}
+
+void UOptimusGroomExecDataInterface::GetSupportedInputs(TArray<FShaderFunctionDefinition>& OutFunctions) const
+{
+	OutFunctions.AddDefaulted_GetRef()
+		.SetName(TEXT("ReadNumThreads"))
+		.AddReturnType(EShaderFundamentalType::Int, 3);
+}
+
+BEGIN_SHADER_PARAMETER_STRUCT(FGroomExecDataInterfaceParameters, )
+	SHADER_PARAMETER(FIntVector, NumThreads)
+END_SHADER_PARAMETER_STRUCT()
+
+void UOptimusGroomExecDataInterface::GetShaderParameters(TCHAR const* UID, FShaderParametersMetadataBuilder& InOutBuilder, FShaderParametersMetadataAllocations& InOutAllocations) const
+{
+	InOutBuilder.AddNestedStruct<FGroomExecDataInterfaceParameters>(UID);
+}
+
+TCHAR const* UOptimusGroomExecDataInterface::TemplateFilePath = TEXT("/Plugin/Runtime/HairStrands/Private/Deformers/DeformerDataGroomDomainsExec.ush");
+
+TCHAR const* UOptimusGroomExecDataInterface::GetShaderVirtualPath() const
+{
+	return TemplateFilePath;
+}
+
+void UOptimusGroomExecDataInterface::GetShaderHash(FString& InOutKey) const
+{
+	GetShaderFileHash(TemplateFilePath, EShaderPlatform::SP_PCD3D_SM5).AppendString(InOutKey);
+}
+
+void UOptimusGroomExecDataInterface::GetHLSL(FString& OutHLSL, FString const& InDataInterfaceName) const
+{
+	TMap<FString, FStringFormatArg> TemplateArgs =
+	{
+		{ TEXT("DataInterfaceName"), InDataInterfaceName },
+	};
+
+	FString TemplateFile;
+	LoadShaderSourceFile(TemplateFilePath, EShaderPlatform::SP_PCD3D_SM5, &TemplateFile, nullptr);
+	OutHLSL += FString::Format(*TemplateFile, TemplateArgs);
+}
+
+UComputeDataProvider* UOptimusGroomExecDataInterface::CreateDataProvider(TObjectPtr<UObject> InBinding, uint64 InInputMask, uint64 InOutputMask) const
+{
+	UOptimusGroomExecDataProvider* Provider = NewObject<UOptimusGroomExecDataProvider>();
+	Provider->MeshComponent = Cast<UMeshComponent>(InBinding);
+	Provider->Domain = Domain;
+	return Provider;
+}
+
+FName UOptimusGroomExecDataInterface::GetExecutionDomainName(const EOptimusGroomExecDomain ExecutionDomain) 
+{
+	if (ExecutionDomain == EOptimusGroomExecDomain::ControlPoint)
+	{
+		return UOptimusGroomAssetComponentSource::FStrandsExecutionDomains::Points;
+	}
+	else if (ExecutionDomain == EOptimusGroomExecDomain::Curve)
+	{
+		return UOptimusGroomAssetComponentSource::FStrandsExecutionDomains::Curves;
+	}
+	else if (ExecutionDomain == EOptimusGroomExecDomain::StrandsEdges)
+	{
+		return UOptimusGroomAssetComponentSource::FStrandsExecutionDomains::Edges;
+	}
+	else if (ExecutionDomain == EOptimusGroomExecDomain::StrandsObjects)
+	{
+		return UOptimusGroomAssetComponentSource::FStrandsExecutionDomains::Objects;
+	}
+	else if (ExecutionDomain == EOptimusGroomExecDomain::GuidesPoints)
+	{
+		return UOptimusGroomAssetComponentSource::FGuidesExecutionDomains::Points;
+	}
+	else if (ExecutionDomain == EOptimusGroomExecDomain::GuidesCurves)
+	{
+		return UOptimusGroomAssetComponentSource::FGuidesExecutionDomains::Curves;
+	}
+	else if (ExecutionDomain == EOptimusGroomExecDomain::GuidesEdges)
+	{
+		return UOptimusGroomAssetComponentSource::FGuidesExecutionDomains::Edges;
+	}
+	else if (ExecutionDomain == EOptimusGroomExecDomain::GuidesObjects)
+	{
+		return UOptimusGroomAssetComponentSource::FGuidesExecutionDomains::Objects;
+	}
+
+	return NAME_None;
+}
+
+FName UOptimusGroomExecDataInterface::GetSelectedExecutionDomainName() const
+{
+	return GetExecutionDomainName(Domain);
+}
+
+FComputeDataProviderRenderProxy* UOptimusGroomExecDataProvider::GetRenderProxy()
+{
+	return new FOptimusGroomExecDataProviderProxy(MeshComponent, Domain);
+}
+
+FOptimusGroomExecDataProviderProxy::FOptimusGroomExecDataProviderProxy(UMeshComponent* MeshComponent, EOptimusGroomExecDomain InDomain)
+{
+	UE::Groom::Private::GatherGroupInstances(MeshComponent, GroupInstances);
+
+	TArray<const UGroomComponent*> GroomComponents;
+	UE::Groom::Private::GatherGroomComponents(MeshComponent, GroomComponents);
+	UE::Groom::Private::GroomComponentsToInstances(GroomComponents, GroupInstances);
+	UE::Groom::Private::GetGroomInvocationElementCounts(GroomComponents,UOptimusGroomExecDataInterface::GetExecutionDomainName(InDomain), GroupCounts);
+	
+	Domain = InDomain;
+}
+
+int32 FOptimusGroomExecDataProviderProxy::GetDispatchThreadCount(TArray<FIntVector>& ThreadCounts) const
+{
+	ThreadCounts.Reset();
+	ThreadCounts.Reserve(GroupCounts.Num());
+	
+	for (const int32& NumThreads : GroupCounts)
+	{	
+		ThreadCounts.Add(FIntVector(NumThreads, 1, 1));
+	}
+	return GroupCounts.Num();
+}
+
+bool FOptimusGroomExecDataProviderProxy::IsValid(FValidationData const& InValidationData) const
+{
+	if (InValidationData.ParameterStructSize != sizeof(FParameters))
+	{
+		return false;
+	}
+	if (InValidationData.NumInvocations != GroupInstances.Num())
+	{
+		return false;
+	}
+	if(Domain == EOptimusGroomExecDomain::ControlPoint || Domain == EOptimusGroomExecDomain::Curve ||
+		Domain == EOptimusGroomExecDomain::StrandsEdges || Domain == EOptimusGroomExecDomain::StrandsObjects)
+	{
+		if(!UE::Groom::Private::HaveStrandsInstanceResources(GroupInstances))
+		{
+			return false;
+			
+		}
+	}
+	if(Domain == EOptimusGroomExecDomain::GuidesPoints || Domain == EOptimusGroomExecDomain::GuidesCurves ||
+		Domain == EOptimusGroomExecDomain::GuidesEdges || Domain == EOptimusGroomExecDomain::GuidesObjects)
+	{
+		if(!UE::Groom::Private::HaveGuidesInstanceResources(GroupInstances))
+		{
+			return false;
+		}
+	}
+	return true;
+}
+
+void FOptimusGroomExecDataProviderProxy::GatherDispatchData(FDispatchData const& InDispatchData)
+{
+	TArray<FIntVector> ThreadCounts;
+	GetDispatchThreadCount(ThreadCounts);
+	const TStridedView<FParameters> ParameterArray = MakeStridedParameterView<FParameters>(InDispatchData);
+	
+	if(ThreadCounts.Num() == ParameterArray.Num())
+	{
+		for (int32 InvocationIndex = 0; InvocationIndex < ParameterArray.Num(); ++InvocationIndex)
+		{
+			ParameterArray[InvocationIndex].NumThreads = ThreadCounts[InvocationIndex];
+		}
+	}
+}
diff --git a/Engine/Plugins/Runtime/HairStrands/Source/HairStrandsDeformer/Private/DeformerGroomDomainsExec.h b/Engine/Plugins/Runtime/HairStrands/Source/HairStrandsDeformer/Private/DeformerGroomDomainsExec.h
new file mode 100644
--- /dev/null
+++ b/Engine/Plugins/Runtime/HairStrands/Source/HairStrandsDeformer/Private/DeformerGroomDomainsExec.h
@@ -0,0 +1,120 @@
+// Copyright Epic Games, Inc. All Rights Reserved.
+
+#pragma once
+
+#include "IOptimusDeprecatedExecutionDataInterface.h"
+#include "OptimusComputeDataInterface.h"
+#include "ComputeFramework/ComputeDataProvider.h"
+#include "DeformerGroomDomainsExec.generated.h"
+
+class FGroomExecDataInterfaceParameters;
+class UMeshComponent;
+struct FHairGroupInstance;
+
+UENUM()
+enum class EOptimusGroomExecDomain : uint8
+{
+	None = 0 UMETA(Hidden),
+	/** Run kernel with one thread per strands points. */
+	ControlPoint = 1 UMETA(DisplayName = "StrandsPoints"),
+	/** Run kernel with one thread per strands curves. */
+	Curve UMETA(DisplayName = "StrandsCurves"),
+	/** Run kernel with one thread per strands edges. */
+	StrandsEdges UMETA(DisplayName = "StrandsEdges"),
+	/** Run kernel with one thread per strands objects. */
+	StrandsObjects UMETA(DisplayName = "StrandsObjects"),
+	/** Run kernel with one thread per guides points. */
+	GuidesPoints UMETA(DisplayName = "GuidesPoints"),
+	/** Run kernel with one thread per guides curves. */
+	GuidesCurves UMETA(DisplayName = "GuidesCurves"),
+	/** Run kernel with one thread per guides edges. */
+	GuidesEdges UMETA(DisplayName = "GuidesEdges"),
+	/** Run kernel with one thread per guides objects. */
+	GuidesObjects UMETA(DisplayName = "GuidesObjects"),
+};
+
+/** Compute Framework Data Interface for executing kernels over a skinned mesh domain. */
+UCLASS(Category = ComputeFramework)
+class UOptimusGroomExecDataInterface :
+	public UOptimusComputeDataInterface,
+	public IOptimusDeprecatedExecutionDataInterface
+{
+	GENERATED_BODY()
+
+public:
+	//~ Begin UOptimusComputeDataInterface Interface
+	virtual FString GetDisplayName() const override;
+	virtual FName GetCategory() const override;
+	virtual bool IsVisible() const override {return false;};
+	virtual TArray<FOptimusCDIPinDefinition> GetPinDefinitions() const override;
+	virtual TSubclassOf<UActorComponent> GetRequiredComponentClass() const override;
+	//~ End UOptimusComputeDataInterface Interface
+	
+	//~ Begin UComputeDataInterface Interface
+	virtual TCHAR const* GetClassName() const override { return TEXT("GroomExec"); }
+	virtual bool IsExecutionInterface() const override { return true; }
+	virtual void GetSupportedInputs(TArray<FShaderFunctionDefinition>& OutFunctions) const override;
+	virtual void GetShaderParameters(TCHAR const* UID, FShaderParametersMetadataBuilder& InOutBuilder, FShaderParametersMetadataAllocations& InOutAllocations) const override;
+	virtual TCHAR const* GetShaderVirtualPath() const override;
+	virtual void GetShaderHash(FString& InOutKey) const override;
+	virtual void GetHLSL(FString& OutHLSL, FString const& InDataInterfaceName) const override;
+	virtual UComputeDataProvider* CreateDataProvider(TObjectPtr<UObject> InBinding, uint64 InInputMask, uint64 InOutputMask) const override;
+	//~ End UComputeDataInterface Interface
+
+	//~ Begin IOptimusDeprecatedExecutionDataInterface Interface
+	virtual FName GetSelectedExecutionDomainName() const override;
+	//~ End IOptimusDeprecatedExecutionDataInterface Interface
+
+	/** static function to retrieve the execution domain name */
+	static FName GetExecutionDomainName(const EOptimusGroomExecDomain ExecutionDomain);
+
+	/** Groom execution domain */
+	UPROPERTY(EditAnywhere, Category = Execution)
+	EOptimusGroomExecDomain Domain = EOptimusGroomExecDomain::ControlPoint;
+
+private:
+	static TCHAR const* TemplateFilePath;
+};
+
+/** Compute Framework Data Provider for executing kernels over a groom execution. */
+UCLASS(BlueprintType, editinlinenew, Category = ComputeFramework)
+class UOptimusGroomExecDataProvider : public UComputeDataProvider
+{
+	GENERATED_BODY()
+
+public:
+	UPROPERTY()
+	TObjectPtr<UMeshComponent> MeshComponent = nullptr;
+
+	/** Groom execution domain */
+	UPROPERTY()
+	EOptimusGroomExecDomain Domain = EOptimusGroomExecDomain::ControlPoint;
+
+	//~ Begin UComputeDataProvider Interface
+	virtual FComputeDataProviderRenderProxy* GetRenderProxy() override;
+	//~ End UComputeDataProvider Interface
+};
+
+class FOptimusGroomExecDataProviderProxy : public FComputeDataProviderRenderProxy
+{
+public:
+	FOptimusGroomExecDataProviderProxy(UMeshComponent* MeshComponent, EOptimusGroomExecDomain InDomain);
+
+	//~ Begin FComputeDataProviderRenderProxy Interface
+	virtual int32 GetDispatchThreadCount(TArray<FIntVector>& ThreadCounts) const override;
+	virtual bool IsValid(FValidationData const& InValidationData) const override;
+	virtual void GatherDispatchData(FDispatchData const& InDispatchData) override;
+	//~ End FComputeDataProviderRenderProxy Interface
+
+private:
+	using FParameters = FGroomExecDataInterfaceParameters;
+
+	/** List of instances (invocations) used in that data interface */
+	TArray<const FHairGroupInstance*> GroupInstances;
+
+	/** Number of elements for each invocation */
+	TArray<int32> GroupCounts;
+
+	/** Groom execution domain */
+	EOptimusGroomExecDomain Domain = EOptimusGroomExecDomain::ControlPoint;
+};
diff --git a/Engine/Plugins/Runtime/HairStrands/Source/HairStrandsDeformer/Private/DeformerGroomDomainsSource.cpp b/Engine/Plugins/Runtime/HairStrands/Source/HairStrandsDeformer/Private/DeformerGroomDomainsSource.cpp
new file mode 100644
--- /dev/null
+++ b/Engine/Plugins/Runtime/HairStrands/Source/HairStrandsDeformer/Private/DeformerGroomDomainsSource.cpp
@@ -0,0 +1,175 @@
+// Copyright Epic Games, Inc. All Rights Reserved.
+
+#include "DeformerGroomDomainsSource.h"
+#include "DeformerGroomInterfaceUtils.h"
+#include "GroomComponent.h"
+#include "GroomSolverComponent.h"
+
+#define LOCTEXT_NAMESPACE "DeformersGroomDomainsSource"
+
+FName UOptimusGroomAssetComponentSource::FStrandsExecutionDomains::Edges("StrandsEdges");
+FName UOptimusGroomAssetComponentSource::FStrandsExecutionDomains::Curves("StrandsCurves");
+FName UOptimusGroomAssetComponentSource::FStrandsExecutionDomains::Objects("StrandsObjects");
+FName UOptimusGroomAssetComponentSource::FStrandsExecutionDomains::Points("StrandsPoints");
+
+FName UOptimusGroomAssetComponentSource::FGuidesExecutionDomains::Edges("GuidesEdges");
+FName UOptimusGroomAssetComponentSource::FGuidesExecutionDomains::Curves("GuidesCurves");
+FName UOptimusGroomAssetComponentSource::FGuidesExecutionDomains::Objects("GuidesObjects");
+FName UOptimusGroomAssetComponentSource::FGuidesExecutionDomains::Points("GuidesPoints");
+
+FName UOptimusGroomAssetComponentSource::FMeshesExecutionDomains::Bones("MeshesBones");
+FName UOptimusGroomAssetComponentSource::FMeshesExecutionDomains::Vertices("MeshesVertices");
+
+FText UOptimusGroomAssetComponentSource::GetDisplayName() const
+{
+	return LOCTEXT("GroomAssetComponent", "Groom Asset Component");
+}
+
+TSubclassOf<UActorComponent> UOptimusGroomAssetComponentSource::GetComponentClass() const
+{
+	return UMeshComponent::StaticClass();
+}
+
+bool UOptimusGroomAssetComponentSource::IsUsableAsPrimarySource() const
+{
+	return GetComponentClass()->IsChildOf<UMeshComponent>();
+}
+
+TArray<FName> UOptimusGroomAssetComponentSource::GetExecutionDomains() const
+{
+	return {FStrandsExecutionDomains::Edges, FStrandsExecutionDomains::Objects, FStrandsExecutionDomains::Curves, FStrandsExecutionDomains::Points,
+				FGuidesExecutionDomains::Edges, FGuidesExecutionDomains::Objects, FGuidesExecutionDomains::Curves, FGuidesExecutionDomains::Points,
+				FMeshesExecutionDomains::Bones, FMeshesExecutionDomains::Vertices};
+}
+
+int32 UOptimusGroomAssetComponentSource::GetLodIndex(const UActorComponent* InComponent) const
+{
+	return 0;
+}
+
+uint32 UOptimusGroomAssetComponentSource::GetDefaultNumInvocations(const UActorComponent* ActorComponent, int32 LodIndex) const
+{
+	TArray<const UGroomComponent*> GroomComponents;
+	UE::Groom::Private::GatherGroomComponents(ActorComponent, GroomComponents);
+	
+	uint32 NumInvocations = 0;
+	for(const UGroomComponent* GroomComponent : GroomComponents)
+	{
+		const uint32 NumGroups = GroomComponent->GetGroupCount();
+		for (uint32 GroupIndex = 0; GroupIndex < NumGroups; ++GroupIndex)
+		{
+			if(GroomComponent->IsDeformationEnable(GroupIndex) && GroomComponent->GetGroupInstance(GroupIndex))
+			{
+				++NumInvocations;
+			}
+		}
+	}
+	return NumInvocations;
+}
+
+bool UOptimusGroomAssetComponentSource::GetComponentElementCountsForExecutionDomain(
+	FName DomainName, const UActorComponent* ActorComponent, int32 LodIndex, TArray<int32>& InvocationElementCounts) const
+{
+	TArray<const UGroomComponent*> GroomComponents;
+	UE::Groom::Private::GatherGroomComponents(ActorComponent, GroomComponents);
+
+	InvocationElementCounts.Reset();
+	const uint32 TotalCount = UE::Groom::Private::GetGroomInvocationElementCounts(GroomComponents, DomainName, InvocationElementCounts);
+	return true;//TotalCount != 0;
+}
+
+FName UOptimusGroomSolverComponentSource::FSolverExecutionDomains::Edges("SolverEdges");
+FName UOptimusGroomSolverComponentSource::FSolverExecutionDomains::Curves("SolverCurves");
+FName UOptimusGroomSolverComponentSource::FSolverExecutionDomains::Objects("SolverObjects");
+FName UOptimusGroomSolverComponentSource::FSolverExecutionDomains::Points("SolverPoints");
+
+FName UOptimusGroomSolverComponentSource::FDynamicExecutionDomains::Points("DynamicPoints");
+FName UOptimusGroomSolverComponentSource::FDynamicExecutionDomains::Curves("DynamicCurves");
+
+FName UOptimusGroomSolverComponentSource::FKinematicExecutionDomains::Points("KinematicPoints");
+FName UOptimusGroomSolverComponentSource::FKinematicExecutionDomains::Curves("KinematicCurves");
+
+FText UOptimusGroomSolverComponentSource::GetDisplayName() const
+{
+	return LOCTEXT("GroomSolverComponent", "Groom Solver Component");
+}
+
+TSubclassOf<UActorComponent> UOptimusGroomSolverComponentSource::GetComponentClass() const
+{
+	return UGroomSolverComponent::StaticClass();
+}
+
+bool UOptimusGroomSolverComponentSource::IsUsableAsPrimarySource() const
+{
+	return GetComponentClass()->IsChildOf<UGroomSolverComponent>();
+}
+
+TArray<FName> UOptimusGroomSolverComponentSource::GetExecutionDomains() const
+{
+	return {FSolverExecutionDomains::Points, FSolverExecutionDomains::Edges, FSolverExecutionDomains::Curves, FSolverExecutionDomains::Objects,
+		FDynamicExecutionDomains::Points, FDynamicExecutionDomains::Curves, FKinematicExecutionDomains::Points, FKinematicExecutionDomains::Curves};
+}
+
+int32 UOptimusGroomSolverComponentSource::GetLodIndex(const UActorComponent* InComponent) const
+{
+	return 0;
+}
+
+uint32 UOptimusGroomSolverComponentSource::GetDefaultNumInvocations(const UActorComponent* ActorComponent, int32 LodIndex) const
+{
+	return 1;
+}
+
+bool UOptimusGroomSolverComponentSource::GetComponentElementCountsForExecutionDomain(
+	FName DomainName,const UActorComponent* ActorComponent, int32 LodIndex, TArray<int32>& InvocationElementCounts) const
+{
+	if(const UGroomSolverComponent* GroomSolver = Cast<UGroomSolverComponent>(ActorComponent))
+	{
+		TArray<const UGroomComponent*> GroomComponents;
+		for(const TObjectPtr<UGroomComponent>& GroomComponent : GroomSolver->GetGroomComponents())
+		{
+			GroomComponents.Add(GroomComponent);
+		}
+		TArray<int32> LocalElementCounts;
+		uint32 TotalCount = 0;
+		
+		if(DomainName == FSolverExecutionDomains::Objects)
+		{
+			TotalCount = UE::Groom::Private::GetGroomInvocationElementCounts(GroomComponents, UOptimusGroomAssetComponentSource::FGuidesExecutionDomains::Objects, LocalElementCounts);
+		}
+		else if(DomainName == FSolverExecutionDomains::Curves)
+		{
+			TotalCount = UE::Groom::Private::GetGroomInvocationElementCounts(GroomComponents, UOptimusGroomAssetComponentSource::FGuidesExecutionDomains::Curves, LocalElementCounts);
+		}
+		else if (DomainName == FSolverExecutionDomains::Edges)
+		{
+			TotalCount = UE::Groom::Private::GetGroomInvocationElementCounts(GroomComponents, UOptimusGroomAssetComponentSource::FGuidesExecutionDomains::Edges, LocalElementCounts);
+		}
+		else if (DomainName == FSolverExecutionDomains::Points)
+		{
+			TotalCount = UE::Groom::Private::GetGroomInvocationElementCounts(GroomComponents, UOptimusGroomAssetComponentSource::FGuidesExecutionDomains::Points, LocalElementCounts, GroupSize);
+		}
+		else if (DomainName == FDynamicExecutionDomains::Curves)
+		{
+			TotalCount = GroomSolver->GetSolverSettings().CurveDynamicIndices.Num();
+		}
+		else if (DomainName == FDynamicExecutionDomains::Points)
+		{
+			TotalCount = GroomSolver->GetSolverSettings().PointDynamicIndices.Num();
+		}
+		else if (DomainName == FKinematicExecutionDomains::Curves)
+		{
+			TotalCount = GroomSolver->GetSolverSettings().CurveKinematicIndices.Num();
+		}
+		else if (DomainName == FKinematicExecutionDomains::Points)
+		{
+			TotalCount = GroomSolver->GetSolverSettings().PointKinematicIndices.Num();
+		}
+		InvocationElementCounts.Init(TotalCount,1);
+		return TotalCount != 0;
+	}
+	return false;
+}
+
+
+#undef LOCTEXT_NAMESPACE
diff --git a/Engine/Plugins/Runtime/HairStrands/Source/HairStrandsDeformer/Private/DeformerGroomDomainsSource.h b/Engine/Plugins/Runtime/HairStrands/Source/HairStrandsDeformer/Private/DeformerGroomDomainsSource.h
new file mode 100644
--- /dev/null
+++ b/Engine/Plugins/Runtime/HairStrands/Source/HairStrandsDeformer/Private/DeformerGroomDomainsSource.h
@@ -0,0 +1,94 @@
+// Copyright Epic Games, Inc. All Rights Reserved.
+
+#pragma once
+
+#include "OptimusComponentSource.h"
+#include "DeformerGroomDomainsSource.generated.h"
+
+class UActorComponent;
+class UGroomComponent;
+class UGroomSolverComponent;
+
+/** Execution domain that will be used by the groom component optimus data interfaces */
+UCLASS()
+class UOptimusGroomAssetComponentSource : public UOptimusComponentSource
+{
+	GENERATED_BODY()
+public:
+
+	struct FGuidesExecutionDomains
+	{
+		static FName Edges;
+		static FName Curves;
+		static FName Objects;
+		static FName Points;
+	};
+
+	struct FStrandsExecutionDomains
+	{
+		static FName Edges;
+		static FName Curves;
+		static FName Objects;
+		static FName Points;
+	};
+
+	struct FMeshesExecutionDomains
+	{
+		static FName Bones;
+		static FName Vertices;
+	};
+	
+	//~ Begin  UOptimusComponentSource implementations
+	virtual FText GetDisplayName() const override;
+	virtual FName GetBindingName() const override { return FName("Groom Asset"); }
+	virtual TSubclassOf<UActorComponent> GetComponentClass() const override;
+	virtual TArray<FName> GetExecutionDomains() const override;
+	virtual int32 GetLodIndex(const UActorComponent* InComponent) const override;
+	virtual uint32 GetDefaultNumInvocations(const UActorComponent* InComponent, int32 InLod) const override;
+	virtual bool GetComponentElementCountsForExecutionDomain(FName InDomainName, const UActorComponent* InComponent, int32 InLodIndex, TArray<int32>& OutInvocationElementCounts) const override;
+	virtual bool IsUsableAsPrimarySource() const override;
+	//~ End UOptimusComputeDataInterface Interface
+};
+
+/** Execution domain that will be used by the groom solver optimus data interfaces */
+UCLASS()
+class UOptimusGroomSolverComponentSource : public UOptimusComponentSource
+{
+	GENERATED_BODY()
+public:
+
+	static constexpr int32 GroupSize = 64;
+
+	struct FSolverExecutionDomains
+	{
+		static FName Edges;
+		static FName Curves;
+		static FName Objects;
+		static FName Points;
+	};
+
+	struct FDynamicExecutionDomains
+	{
+		static FName Points;
+		static FName Curves;
+	};
+
+	struct FKinematicExecutionDomains
+	{
+		static FName Points;
+		static FName Curves;
+	};
+	
+	//~ Begin  UOptimusComponentSource implementations
+	virtual FText GetDisplayName() const override;
+	virtual FName GetBindingName() const override { return FName("Groom Solver"); }
+	virtual TSubclassOf<UActorComponent> GetComponentClass() const override;
+	virtual TArray<FName> GetExecutionDomains() const override;
+	virtual int32 GetLodIndex(const UActorComponent* InComponent) const override;
+	virtual uint32 GetDefaultNumInvocations(const UActorComponent* InComponent, int32 InLod) const override;
+	virtual bool GetComponentElementCountsForExecutionDomain(FName InDomainName, const UActorComponent* InComponent, int32 InLodIndex, TArray<int32>& OutInvocationElementCounts) const override;
+	virtual bool IsUsableAsPrimarySource() const override;
+	//~ End UOptimusComputeDataInterface Interface
+};
+
+
diff --git a/Engine/Plugins/Runtime/HairStrands/Source/HairStrandsDeformer/Private/DeformerGroomGuidesRead.cpp b/Engine/Plugins/Runtime/HairStrands/Source/HairStrandsDeformer/Private/DeformerGroomGuidesRead.cpp
new file mode 100644
--- /dev/null
+++ b/Engine/Plugins/Runtime/HairStrands/Source/HairStrandsDeformer/Private/DeformerGroomGuidesRead.cpp
@@ -0,0 +1,300 @@
+// Copyright Epic Games, Inc. All Rights Reserved.
+
+#include "DeformerGroomGuidesRead.h"
+
+#include "ComputeFramework/ShaderParamTypeDefinition.h"
+#include "ComputeFramework/ComputeKernelPermutationVector.h"
+#include "ShaderCompilerCore.h"
+#include "ShaderParameterMetadataBuilder.h"
+#include "GroomComponent.h"
+#include "GroomInstance.h"
+#include "RenderGraphBuilder.h"
+#include "GlobalRenderResources.h"
+#include "DeformerGroomDomainsSource.h"
+#include "DeformerGroomInterfaceUtils.h"
+#include "HairStrandsInterpolation.h"
+#include "OptimusDataTypeRegistry.h"
+#include "OptimusDeformerInstance.h"
+#include "OptimusValueContainerStruct.h"
+#include "SystemTextures.h"
+
+
+FString UOptimusGroomGuidesReadDataInterface::GetDisplayName() const
+{
+	return TEXT("Groom Guides");
+}
+
+TArray<FOptimusCDIPinDefinition> UOptimusGroomGuidesReadDataInterface::GetPinDefinitions() const
+{
+	static const FName GuidesPoints(UOptimusGroomAssetComponentSource::FGuidesExecutionDomains::Points);
+	static const FName GuidesCurves(UOptimusGroomAssetComponentSource::FGuidesExecutionDomains::Curves);
+
+	TArray<FOptimusCDIPinDefinition> Defs;
+
+	// Deformation buffers
+	Defs.Add({ "NumGuidePoints",   "ReadNumPoints", false, "NumGuidesPoints" });
+	Defs.Add({ "NumGuideCurves",   "ReadNumCurves", false, "NumGuidesCurves" });
+	Defs.Add({ "Position",         "ReadPointRestPosition",          GuidesPoints,  "ReadNumPoints", false, "PointRestPosition"});
+	Defs.Add({ "CurveOffsetPoint", "ReadCurvePointOffset",  GuidesCurves,          "ReadNumCurves", false, "CurvePointOffset" });
+	Defs.Add({ "CurveNumPoint",    "ReadCurveNumPoints",     GuidesCurves,          "ReadNumCurves", false, "CurveNumPoints"});
+	Defs.Add({ "PointCurveIndex", "ReadPointCurveIndex", GuidesPoints,          "ReadNumPoints", false, "PointCurveIndex"  });
+	Defs.Add({ "CurveRestTransform", "ReadCurveRestTransform", GuidesCurves,   "ReadNumCurves", false, "CurveRestTransform"  });
+	Defs.Add({ "CurveDeformedTransform", "ReadCurveDeformedTransform", GuidesCurves,          "ReadNumCurves", true, "CurveDeformedTransform"  });
+	Defs.Add({ "ObjectRestTransform", "ReadObjectRestTransform", false, "ObjectRestTransform"  });
+	Defs.Add({ "ObjectDeformedTransform", "ReadObjectDeformedTransform", true, "ObjectDeformedTransform"  });
+	Defs.Add({ "GuidesObjectIndex", "ReadObjectIndex", true, "GuidesObjectIndex"  });
+	
+	return Defs;
+}
+ 
+TSubclassOf<UActorComponent> UOptimusGroomGuidesReadDataInterface::GetRequiredComponentClass() const
+{
+	return UMeshComponent::StaticClass();
+}
+
+void UOptimusGroomGuidesReadDataInterface::GetPermutations(FComputeKernelPermutationVector& OutPermutationVector) const
+{
+	OutPermutationVector.AddPermutation(TEXT("ENABLE_SKINNED_TRANSFORM"), 2);
+}
+
+void UOptimusGroomGuidesReadDataInterface::GetSupportedInputs(TArray<FShaderFunctionDefinition>& OutFunctions) const
+{
+	OutFunctions.AddDefaulted_GetRef()
+		.SetName(TEXT("ReadNumPoints"))
+		.AddReturnType(EShaderFundamentalType::Uint);
+
+	OutFunctions.AddDefaulted_GetRef()
+		.SetName(TEXT("ReadNumCurves"))
+		.AddReturnType(EShaderFundamentalType::Uint);
+
+	OutFunctions.AddDefaulted_GetRef()
+		.SetName(TEXT("ReadObjectIndex"))
+		.AddReturnType(EShaderFundamentalType::Uint);
+
+	OutFunctions.AddDefaulted_GetRef()
+		.SetName(TEXT("ReadPointRestPosition"))
+		.AddReturnType(EShaderFundamentalType::Float, 3)
+		.AddParam(EShaderFundamentalType::Uint);
+
+	OutFunctions.AddDefaulted_GetRef()
+		.SetName(TEXT("ReadCurvePointOffset"))
+		.AddReturnType(EShaderFundamentalType::Uint)
+		.AddParam(EShaderFundamentalType::Uint);
+
+	OutFunctions.AddDefaulted_GetRef()
+		.SetName(TEXT("ReadCurveNumPoints"))
+		.AddReturnType(EShaderFundamentalType::Uint)
+		.AddParam(EShaderFundamentalType::Uint);
+
+	OutFunctions.AddDefaulted_GetRef()
+		.SetName(TEXT("ReadPointCurveIndex"))
+		.AddReturnType(EShaderFundamentalType::Uint)
+		.AddParam(EShaderFundamentalType::Uint);
+		
+	OutFunctions.AddDefaulted_GetRef()
+		.SetName(TEXT("ReadCurveRestTransform"))
+		.AddReturnType(EShaderFundamentalType::Float, 3, 4)
+		.AddParam(EShaderFundamentalType::Uint);
+	
+	OutFunctions.AddDefaulted_GetRef()
+		.SetName(TEXT("ReadCurveDeformedTransform"))
+		.AddReturnType(EShaderFundamentalType::Float, 3, 4)
+		.AddParam(EShaderFundamentalType::Uint);
+		
+	OutFunctions.AddDefaulted_GetRef()
+		.SetName(TEXT("ReadObjectRestTransform"))
+		.AddReturnType(EShaderFundamentalType::Float, 3, 4);
+	
+	OutFunctions.AddDefaulted_GetRef()
+		.SetName(TEXT("ReadObjectDeformedTransform"))
+		.AddReturnType(EShaderFundamentalType::Float, 3, 4);
+}
+
+BEGIN_SHADER_PARAMETER_STRUCT(FOptimusGroomGuidesReadParameters, )
+	SHADER_PARAMETER(uint32, PointCount)
+	SHADER_PARAMETER(uint32, CurveCount)
+	SHADER_PARAMETER(uint32, ObjectIndex)
+	SHADER_PARAMETER(FVector3f, RestPositionOffset)
+	SHADER_PARAMETER_RDG_BUFFER_SRV(ByteAddressBuffer, PointRestPositions)
+	SHADER_PARAMETER_RDG_BUFFER_SRV(ByteAddressBuffer, CurvePointOffsets)
+	SHADER_PARAMETER_RDG_BUFFER_SRV(ByteAddressBuffer, PointCurveIndices)
+	SHADER_PARAMETER(FMatrix44f, ObjectRestTransform)
+	SHADER_PARAMETER(FMatrix44f, ObjectDeformedTransform)
+	SHADER_PARAMETER_RDG_BUFFER_SRV(Buffer<float4>, TriangleRestPositions)
+	SHADER_PARAMETER_RDG_BUFFER_SRV(Buffer<float4>, TriangleDeformedPositions)
+	SHADER_PARAMETER_RDG_BUFFER_SRV(Buffer<uint>, CurveBarycentricCoordinates)
+	SHADER_PARAMETER_RDG_BUFFER_SRV(Buffer<uint>, CurveTriangleIndices)
+END_SHADER_PARAMETER_STRUCT()
+
+void UOptimusGroomGuidesReadDataInterface::GetShaderParameters(TCHAR const* UID, FShaderParametersMetadataBuilder& InOutBuilder, FShaderParametersMetadataAllocations& InOutAllocations) const
+{
+	InOutBuilder.AddNestedStruct<FOptimusGroomGuidesReadParameters>(UID);
+}
+
+TCHAR const* UOptimusGroomGuidesReadDataInterface::TemplateFilePath = TEXT("/Plugin/Runtime/HairStrands/Private/Deformers/DeformerGroomGuidesRead.ush");
+
+TCHAR const* UOptimusGroomGuidesReadDataInterface::GetShaderVirtualPath() const
+{
+	return TemplateFilePath;
+}
+
+void UOptimusGroomGuidesReadDataInterface::GetShaderHash(FString& InOutKey) const
+{
+	GetShaderFileHash(TemplateFilePath, EShaderPlatform::SP_PCD3D_SM5).AppendString(InOutKey);
+}
+
+void UOptimusGroomGuidesReadDataInterface::GetHLSL(FString& OutHLSL, FString const& InDataInterfaceName) const
+{
+	TMap<FString, FStringFormatArg> TemplateArgs =
+	{
+		{ TEXT("DataInterfaceName"), InDataInterfaceName },
+	};
+
+	FString TemplateFile;
+	LoadShaderSourceFile(TemplateFilePath, EShaderPlatform::SP_PCD3D_SM5, &TemplateFile, nullptr);
+	OutHLSL += FString::Format(*TemplateFile, TemplateArgs);
+}
+
+UComputeDataProvider* UOptimusGroomGuidesReadDataInterface::CreateDataProvider(TObjectPtr<UObject> InBinding, uint64 InInputMask, uint64 InOutputMask) const
+{
+	UOptimusGroomGuidesReadDataProvider* Provider = NewObject<UOptimusGroomGuidesReadDataProvider>();
+	Provider->MeshComponent = Cast<UMeshComponent>(InBinding);
+	return Provider;
+}
+
+FComputeDataProviderRenderProxy* UOptimusGroomGuidesReadDataProvider::GetRenderProxy()
+{
+	return new FOptimusGroomGuidesReadDataProviderProxy(MeshComponent);
+}
+
+FOptimusGroomGuidesReadDataProviderProxy::FOptimusGroomGuidesReadDataProviderProxy(UMeshComponent* MeshComponent)
+{
+	TArray<const UGroomComponent*> GroomComponents;
+	UE::Groom::Private::GatherGroomComponents(MeshComponent, GroomComponents);
+	UE::Groom::Private::GroomComponentsToInstances(GroomComponents, GroupInstances);
+}
+
+bool FOptimusGroomGuidesReadDataProviderProxy::IsValid(FValidationData const& InValidationData) const
+{
+	if (InValidationData.ParameterStructSize != sizeof(FParameters))
+	{
+		return false;
+	}
+	if (InValidationData.NumInvocations != GroupInstances.Num())
+	{
+		return false;
+	}
+	if(!UE::Groom::Private::HaveGuidesInstanceResources(GroupInstances) || !UE::Groom::Private::HaveGuidesSkinnedResources(GroupInstances))
+	{
+		return false;
+	}
+	
+	return true;
+}
+
+struct FOptimusGroomGuidesReadPermutationIds
+{
+	uint32 EnableSkinnedTransform = 0;
+
+	FOptimusGroomGuidesReadPermutationIds(FComputeKernelPermutationVector const& PermutationVector)
+	{
+		static FString Name(TEXT("ENABLE_SKINNED_TRANSFORM"));
+		static uint32 Hash = GetTypeHash(Name);
+		EnableSkinnedTransform = PermutationVector.GetPermutationBits(Name, Hash, 1);
+	}
+};
+
+void FOptimusGroomGuidesReadDataProviderProxy::GatherPermutations(FPermutationData& InOutPermutationData) const
+{
+	const FOptimusGroomGuidesReadPermutationIds PermutationIds(InOutPermutationData.PermutationVector);
+
+	int32 InvocationIndex = 0;
+	for (const FHairGroupInstance* GroupInstance : GroupInstances)
+	{
+		if (GroupInstance)
+		{
+			InOutPermutationData.PermutationIds[InvocationIndex] |= ((GroupInstance->BindingType == EHairBindingType::Skinning) ? PermutationIds.EnableSkinnedTransform : 0);
+		}
+		++InvocationIndex;
+	}
+}
+
+void FOptimusGroomGuidesReadDataProviderProxy::AllocateResources(FRDGBuilder& GraphBuilder, FAllocationData const& InAllocationData)
+{
+	InstanceResources.Reset();
+	BindingResources.Reset();
+
+	int32 InvocationIndex = 0;
+	for (const FHairGroupInstance* GroupInstance : GroupInstances)
+	{
+		if (GroupInstance)
+		{
+			{
+				FInstanceResources& Resource = InstanceResources.AddDefaulted_GetRef();
+				
+				Resource.PointRestPositions	= RegisterAsSRV(GraphBuilder, GroupInstance->Guides.RestResource->PositionBuffer);
+				Resource.CurvePointOffsets = RegisterAsSRV(GraphBuilder, GroupInstance->Guides.RestResource->CurveBuffer);
+				Resource.PointCurveIndices = RegisterAsSRV(GraphBuilder, GroupInstance->Guides.RestResource->PointToCurveBuffer);
+			}
+			{
+				FBindingResources& Resource = BindingResources.AddDefaulted_GetRef();
+				if(GroupInstance->BindingType == EHairBindingType::Skinning)
+				{
+					FHairStrandsLODRestRootResource* RestLODDatas = GroupInstance->Guides.RestRootResource->LODs[GroupInstance->HairGroupPublicData->MeshLODIndex];
+					FHairStrandsLODDeformedRootResource* DeformedLODDatas = GroupInstance->Guides.DeformedRootResource->LODs[GroupInstance->HairGroupPublicData->MeshLODIndex];
+			
+					Resource.CurveTriangleIndices = RegisterAsSRV(GraphBuilder, RestLODDatas->RootToUniqueTriangleIndexBuffer);
+					Resource.TriangleRestPositions = RegisterAsSRV(GraphBuilder, RestLODDatas->RestUniqueTrianglePositionBuffer);
+					Resource.TriangleDeformedPositions = RegisterAsSRV(GraphBuilder, DeformedLODDatas->GetDeformedUniqueTrianglePositionBuffer(FHairStrandsLODDeformedRootResource::Current));
+					Resource.CurveBarycentricCoordinates = RegisterAsSRV(GraphBuilder, RestLODDatas->RootBarycentricBuffer);
+				}
+				{
+					Resource.ObjectDeformedTransform = FMatrix44f(GroupInstance->GetCurrentLocalToWorld().ToMatrixWithScale().GetTransposed());
+					Resource.ObjectRestTransform = FMatrix44f(GroupInstance->GetRestLocalToWorld().ToMatrixWithScale().GetTransposed());
+				}
+			}
+		}
+		++InvocationIndex;
+	}
+}
+
+void FOptimusGroomGuidesReadDataProviderProxy::GatherDispatchData(FDispatchData const& InDispatchData)
+{
+	check(InDispatchData.NumInvocations == GroupInstances.Num());
+	const TStridedView<FParameters> ParameterArray = MakeStridedParameterView<FParameters>(InDispatchData);
+	
+	for (int32 InvocationIndex = 0; InvocationIndex < ParameterArray.Num(); ++InvocationIndex)
+	{
+		if (const FHairGroupInstance* GroupInstance = GroupInstances[InvocationIndex])
+		{
+			FParameters& Parameters = ParameterArray[InvocationIndex];
+			
+			Parameters.PointCount = GroupInstance->Guides.RestResource->GetPointCount();
+			Parameters.CurveCount = GroupInstance->Guides.RestResource->GetCurveCount();
+			Parameters.ObjectIndex = InvocationIndex;
+			Parameters.RestPositionOffset = FVector3f(GroupInstance->Guides.RestResource->GetPositionOffset());
+
+			{
+				FInstanceResources& Resource = InstanceResources[InvocationIndex];
+				Parameters.PointRestPositions = Resource.PointRestPositions;
+				Parameters.CurvePointOffsets = Resource.CurvePointOffsets;
+				Parameters.PointCurveIndices = Resource.PointCurveIndices;
+			}
+
+			{
+				FBindingResources& Resource = BindingResources[InvocationIndex];
+				if(GroupInstance->BindingType == EHairBindingType::Skinning)
+				{
+					Parameters.CurveTriangleIndices = Resource.CurveTriangleIndices;
+					Parameters.CurveBarycentricCoordinates = Resource.CurveBarycentricCoordinates;
+					Parameters.TriangleRestPositions = Resource.TriangleRestPositions;
+					Parameters.TriangleDeformedPositions = Resource.TriangleDeformedPositions;
+				}
+				{
+					Parameters.ObjectDeformedTransform = Resource.ObjectDeformedTransform;
+					Parameters.ObjectRestTransform = Resource.ObjectRestTransform;
+				}
+			}
+		}
+	}
+}
\ No newline at end of file
diff --git a/Engine/Plugins/Runtime/HairStrands/Source/HairStrandsDeformer/Private/DeformerGroomGuidesRead.h b/Engine/Plugins/Runtime/HairStrands/Source/HairStrandsDeformer/Private/DeformerGroomGuidesRead.h
new file mode 100644
--- /dev/null
+++ b/Engine/Plugins/Runtime/HairStrands/Source/HairStrandsDeformer/Private/DeformerGroomGuidesRead.h
@@ -0,0 +1,102 @@
+// Copyright Epic Games, Inc. All Rights Reserved.
+
+#pragma once
+
+#include "RenderGraphDefinitions.h"
+#include "OptimusComputeDataInterface.h"
+#include "HairStrandsInterface.h"
+#include "IOptimusDeformerInstanceAccessor.h"
+#include "ComputeFramework/ComputeDataProvider.h"
+#include "DeformerGroomGuidesRead.generated.h"
+
+class FOptimusGroomGuidesReadParameters;
+class UGroomComponent;
+struct FHairGroupInstance;
+class UGroomAsset;
+class UMeshComponent;
+
+/** Compute Framework Data Interface for reading groom guides. */
+UCLASS(Category = ComputeFramework)
+class UOptimusGroomGuidesReadDataInterface : public UOptimusComputeDataInterface
+{
+	GENERATED_BODY()
+
+public:
+	
+	//~ Begin UOptimusComputeDataInterface Interface
+	virtual FString GetDisplayName() const override;
+	virtual TArray<FOptimusCDIPinDefinition> GetPinDefinitions() const override;
+	virtual TSubclassOf<UActorComponent> GetRequiredComponentClass() const override;
+	//~ End UOptimusComputeDataInterface Interface
+	
+	//~ Begin UComputeDataInterface Interface
+	virtual TCHAR const* GetClassName() const override { return TEXT("GroomGuidesRead"); }
+	virtual void GetSupportedInputs(TArray<FShaderFunctionDefinition>& OutFunctions) const override;
+	virtual void GetShaderParameters(TCHAR const* UID, FShaderParametersMetadataBuilder& InOutBuilder, FShaderParametersMetadataAllocations& InOutAllocations) const override;
+	virtual TCHAR const* GetShaderVirtualPath() const override;
+	virtual void GetPermutations(FComputeKernelPermutationVector& OutPermutationVector) const override;
+	virtual void GetShaderHash(FString& InOutKey) const override;
+	virtual void GetHLSL(FString& OutHLSL, FString const& InDataInterfaceName) const override;
+	virtual UComputeDataProvider* CreateDataProvider(TObjectPtr<UObject> InBinding, uint64 InInputMask, uint64 InOutputMask) const override;
+	//~ End UComputeDataInterface Interface
+
+private:
+	/** File holding the hlsl implementation */
+	static TCHAR const* TemplateFilePath;
+};
+
+/** Compute Framework Data Provider for reading groom guides. */
+UCLASS(BlueprintType, editinlinenew, Category = ComputeFramework)
+class UOptimusGroomGuidesReadDataProvider : public UComputeDataProvider
+{
+	GENERATED_BODY()
+
+public:
+	
+	UPROPERTY(EditAnywhere, BlueprintReadWrite, Category = Binding)
+	TObjectPtr<UMeshComponent> MeshComponent = nullptr;
+
+	//~ Begin UComputeDataProvider Interface
+	virtual FComputeDataProviderRenderProxy* GetRenderProxy() override;
+	//~ End UComputeDataProvider Interface
+};
+
+class FOptimusGroomGuidesReadDataProviderProxy : public FComputeDataProviderRenderProxy
+{
+public:
+	FOptimusGroomGuidesReadDataProviderProxy(UMeshComponent* MeshComponent);
+
+	//~ Begin FComputeDataProviderRenderProxy Interface
+	virtual bool IsValid(FValidationData const& InValidationData) const override;
+	virtual void AllocateResources(FRDGBuilder& GraphBuilder, FAllocationData const& InAllocationData) override;
+	virtual void GatherPermutations(FPermutationData& InOutPermutationData) const override;
+	virtual void GatherDispatchData(FDispatchData const& InDispatchData) override;
+	//~ End FComputeDataProviderRenderProxy Interface
+
+private:
+	using FParameters = FOptimusGroomGuidesReadParameters;
+
+	/** List of instances (invocations) used in that data interface */
+	TArray<const FHairGroupInstance*> GroupInstances;
+
+	/** Binding Resources used to dispatch CS on GPU */
+	struct FBindingResources
+	{
+		FMatrix44f ObjectRestTransform = FMatrix44f::Identity;
+		FMatrix44f ObjectDeformedTransform = FMatrix44f::Identity;
+		FRDGBufferSRVRef TriangleRestPositions = nullptr;
+		FRDGBufferSRVRef TriangleDeformedPositions = nullptr;
+		FRDGBufferSRVRef CurveBarycentricCoordinates = nullptr;
+		FRDGBufferSRVRef CurveTriangleIndices = nullptr;
+	};
+	TArray<FBindingResources> BindingResources;
+
+	/** Instance resources used to dispatch CS on GPU */
+	struct FInstanceResources
+	{
+		FRDGBufferSRVRef CurvePointOffsets = nullptr;
+		FRDGBufferSRVRef PointRestPositions = nullptr;
+		FRDGBufferSRVRef PointCurveIndices = nullptr;
+	};
+	TArray<FInstanceResources> InstanceResources;
+};
diff --git a/Engine/Plugins/Runtime/HairStrands/Source/HairStrandsDeformer/Private/DeformerGroomGuidesWrite.cpp b/Engine/Plugins/Runtime/HairStrands/Source/HairStrandsDeformer/Private/DeformerGroomGuidesWrite.cpp
new file mode 100644
--- /dev/null
+++ b/Engine/Plugins/Runtime/HairStrands/Source/HairStrandsDeformer/Private/DeformerGroomGuidesWrite.cpp
@@ -0,0 +1,169 @@
+// Copyright Epic Games, Inc. All Rights Reserved.
+
+#include "DeformerGroomGuidesWrite.h"
+
+#include "ComputeFramework/ShaderParamTypeDefinition.h"
+#include "RenderGraphBuilder.h"
+#include "ShaderParameterMetadataBuilder.h"
+#include "GroomComponent.h"
+#include "GroomInstance.h"
+#include "DeformerGroomDomainsSource.h"
+#include "DeformerGroomInterfaceUtils.h"
+#include "RenderGraphUtils.h"
+#include "SystemTextures.h"
+#include "HairStrandsInterpolation.h"
+
+FString UOptimusGroomGuidesWriteDataInterface::GetDisplayName() const
+{
+	return TEXT("Write Groom Guides");
+}
+
+FName UOptimusGroomGuidesWriteDataInterface::GetCategory() const
+{
+	return CategoryName::OutputDataInterfaces;
+}
+
+TArray<FOptimusCDIPinDefinition> UOptimusGroomGuidesWriteDataInterface::GetPinDefinitions() const
+{
+	static const FName GuidesPoints(UOptimusGroomAssetComponentSource::FGuidesExecutionDomains::Points);
+
+	TArray<FOptimusCDIPinDefinition> Defs;
+	
+	// Deformation buffers
+	Defs.Add({ "PointDeformedPosition", "WritePointDeformedPosition", GuidesPoints, "ReadNumPoints", true, "PointDeformedPosition" });
+
+	return Defs;
+}
+
+TSubclassOf<UActorComponent> UOptimusGroomGuidesWriteDataInterface::GetRequiredComponentClass() const
+{
+	return UMeshComponent::StaticClass();
+}
+
+void UOptimusGroomGuidesWriteDataInterface::GetSupportedInputs(TArray<FShaderFunctionDefinition>& OutFunctions) const
+{
+	OutFunctions.AddDefaulted_GetRef()
+		.SetName(TEXT("ReadNumPoints"))
+		.AddReturnType(EShaderFundamentalType::Uint);
+}
+
+void UOptimusGroomGuidesWriteDataInterface::GetSupportedOutputs(TArray<FShaderFunctionDefinition>& OutFunctions) const
+{
+	OutFunctions.AddDefaulted_GetRef()
+		.SetName(TEXT("WritePointDeformedPosition"))
+		.AddParam(EShaderFundamentalType::Uint)
+		.AddParam(EShaderFundamentalType::Float, 3);
+}
+
+BEGIN_SHADER_PARAMETER_STRUCT(FOptimusGroomGuidesWriteParameters, )
+	SHADER_PARAMETER(uint32, RegisteredIndex)
+	SHADER_PARAMETER(uint32, PointCount)
+	SHADER_PARAMETER(uint32, OutputStreamStart)
+	SHADER_PARAMETER_RDG_BUFFER_SRV(StructuredBuffer, DeformedPositionOffset)
+	SHADER_PARAMETER_RDG_BUFFER_SRV(ByteAddressBuffer, PointRestPositions)
+	SHADER_PARAMETER_RDG_BUFFER_UAV(RWByteAddressBuffer, PointDeformedPositions)
+END_SHADER_PARAMETER_STRUCT()
+
+void UOptimusGroomGuidesWriteDataInterface::GetShaderParameters(TCHAR const* UID, FShaderParametersMetadataBuilder& InOutBuilder, FShaderParametersMetadataAllocations& InOutAllocations) const
+{
+	InOutBuilder.AddNestedStruct<FOptimusGroomGuidesWriteParameters>(UID);
+}
+
+TCHAR const* UOptimusGroomGuidesWriteDataInterface::TemplateFilePath = TEXT("/Plugin/Runtime/HairStrands/Private/Deformers/DeformerGroomGuidesWrite.ush");
+
+TCHAR const* UOptimusGroomGuidesWriteDataInterface::GetShaderVirtualPath() const
+{
+	return TemplateFilePath;
+}
+
+void UOptimusGroomGuidesWriteDataInterface::GetShaderHash(FString& InOutKey) const
+{
+	GetShaderFileHash(TemplateFilePath, EShaderPlatform::SP_PCD3D_SM5).AppendString(InOutKey);
+}
+
+void UOptimusGroomGuidesWriteDataInterface::GetHLSL(FString& OutHLSL, FString const& InDataInterfaceName) const
+{
+	TMap<FString, FStringFormatArg> TemplateArgs =
+	{
+		{ TEXT("DataInterfaceName"), InDataInterfaceName },
+	};
+
+	FString TemplateFile;
+	LoadShaderSourceFile(TemplateFilePath, EShaderPlatform::SP_PCD3D_SM5, &TemplateFile, nullptr);
+	OutHLSL += FString::Format(*TemplateFile, TemplateArgs);
+}
+
+UComputeDataProvider* UOptimusGroomGuidesWriteDataInterface::CreateDataProvider(TObjectPtr<UObject> InBinding, uint64 InInputMask, uint64 InOutputMask) const
+{
+	UOptimusGroomGuidesWriteDataProvider* Provider = NewObject<UOptimusGroomGuidesWriteDataProvider>();
+	Provider->MeshComponent = Cast<UMeshComponent>(InBinding);
+	Provider->OutputMask = InOutputMask;
+	return Provider;
+}
+
+FComputeDataProviderRenderProxy* UOptimusGroomGuidesWriteDataProvider::GetRenderProxy()
+{
+	return new FOptimusGroomGuidesWriteProviderProxy(MeshComponent, OutputMask);
+}
+
+FOptimusGroomGuidesWriteProviderProxy::FOptimusGroomGuidesWriteProviderProxy(UMeshComponent* MeshComponent, uint64 InOutputMask)
+{
+	OutputMask = InOutputMask;
+	UE::Groom::Private::GatherGroupInstances(MeshComponent, GroupInstances);
+}
+
+bool FOptimusGroomGuidesWriteProviderProxy::IsValid(FValidationData const& InValidationData) const
+{
+	if (InValidationData.ParameterStructSize != sizeof(FParameters))
+	{
+		return false;
+	}
+	if (InValidationData.NumInvocations != GroupInstances.Num())
+	{
+		return false;
+	}
+	if(!UE::Groom::Private::HaveGuidesInstanceResources(GroupInstances))
+	{
+		return false;
+	}
+	return true;
+}
+
+void FOptimusGroomGuidesWriteProviderProxy::AllocateResources(FRDGBuilder& GraphBuilder, FAllocationData const& InAllocationData)
+{
+	InstanceResources.Reset();
+	for (const FHairGroupInstance* GroupInstance : GroupInstances)
+	{
+		if (GroupInstance)
+		{
+			FInstanceResources& Resource = InstanceResources.AddDefaulted_GetRef();
+			
+			Resource.DeformedPositionOffset = Register(GraphBuilder, GroupInstance->Guides.DeformedResource->GetPositionOffsetBuffer(FHairStrandsDeformedResource::EFrameType::Current), ERDGImportedBufferFlags::CreateSRV).SRV;
+			Resource.PointRestPositions = Register(GraphBuilder, GroupInstance->Guides.RestResource->PositionBuffer, ERDGImportedBufferFlags::CreateSRV).SRV;
+			Resource.PointDeformedPositions = Register(GraphBuilder, GroupInstance->Guides.DeformedResource->GetDeformerBuffer(GraphBuilder), ERDGImportedBufferFlags::CreateUAV).UAV;
+
+			//Resource.PointDeformedPositions = Register(GraphBuilder, GroupInstance->Guides.DeformedResource->GetBuffer(FHairStrandsDeformedResource::Current), ERDGImportedBufferFlags::CreateUAV).UAV;
+		}	
+	}
+}
+
+void FOptimusGroomGuidesWriteProviderProxy::GatherDispatchData(FDispatchData const& InDispatchData)
+{
+	const TStridedView<FParameters> ParameterArray = MakeStridedParameterView<FParameters>(InDispatchData);
+	for (int32 InvocationIndex = 0; InvocationIndex < ParameterArray.Num(); ++InvocationIndex)
+	{
+		if (const FHairGroupInstance* GroupInstance = GroupInstances[InvocationIndex])
+		{
+			FParameters& Parameters = ParameterArray[InvocationIndex];
+			FInstanceResources& Resource = InstanceResources[InvocationIndex];
+			
+			Parameters.PointCount = GroupInstance->Guides.RestResource->GetPointCount();
+			Parameters.OutputStreamStart = 0;
+			Parameters.RegisteredIndex = GroupInstance->RegisteredIndex;
+			
+			Parameters.DeformedPositionOffset 	=  Resource.DeformedPositionOffset;
+			Parameters.PointRestPositions = Resource.PointRestPositions;
+			Parameters.PointDeformedPositions = Resource.PointDeformedPositions;
+		}
+	}
+}
\ No newline at end of file
diff --git a/Engine/Plugins/Runtime/HairStrands/Source/HairStrandsDeformer/Private/DeformerGroomGuidesWrite.h b/Engine/Plugins/Runtime/HairStrands/Source/HairStrandsDeformer/Private/DeformerGroomGuidesWrite.h
new file mode 100644
--- /dev/null
+++ b/Engine/Plugins/Runtime/HairStrands/Source/HairStrandsDeformer/Private/DeformerGroomGuidesWrite.h
@@ -0,0 +1,93 @@
+// Copyright Epic Games, Inc. All Rights Reserved.
+
+#pragma once
+
+#include "OptimusComputeDataInterface.h"
+#include "ComputeFramework/ComputeDataProvider.h"
+#include "DeformerGroomGuidesWrite.generated.h"
+
+class FOptimusGroomGuidesWriteParameters;
+struct FHairGroupInstance;
+class FRDGBuffer;
+class FRDGBufferSRV;
+class FRDGBufferUAV;
+class UGroomComponent;
+class UMeshComponent;
+
+/** Compute Framework Data Interface for writing groom strands. */
+UCLASS(Category = ComputeFramework)
+class UOptimusGroomGuidesWriteDataInterface : public UOptimusComputeDataInterface
+{
+	GENERATED_BODY()
+
+public:
+	//~ Begin UOptimusComputeDataInterface Interface
+	virtual FString GetDisplayName() const override;
+	virtual FName GetCategory() const override;
+	virtual TArray<FOptimusCDIPinDefinition> GetPinDefinitions() const override;
+	virtual TSubclassOf<UActorComponent> GetRequiredComponentClass() const override;
+	//~ End UOptimusComputeDataInterface Interface
+	
+	//~ Begin UComputeDataInterface Interface
+	virtual TCHAR const* GetClassName() const override { return TEXT("GroomGuidesWrite"); }
+	virtual void GetSupportedInputs(TArray<FShaderFunctionDefinition>& OutFunctions) const override;
+	virtual void GetSupportedOutputs(TArray<FShaderFunctionDefinition>& OutFunctions) const override;
+	virtual void GetShaderParameters(TCHAR const* UID, FShaderParametersMetadataBuilder& InOutBuilder, FShaderParametersMetadataAllocations& InOutAllocations) const override;
+	virtual TCHAR const* GetShaderVirtualPath() const override;
+	virtual void GetShaderHash(FString& InOutKey) const override;
+	virtual void GetHLSL(FString& OutHLSL, FString const& InDataInterfaceName) const override;
+	virtual UComputeDataProvider* CreateDataProvider(TObjectPtr<UObject> InBinding, uint64 InInputMask, uint64 InOutputMask) const override;
+	//~ End UComputeDataInterface Interface
+
+private:
+	/** File holding the hlsl implementation */
+	static TCHAR const* TemplateFilePath;
+};
+
+/** Compute Framework Data Provider for writing skinned mesh. */
+UCLASS(BlueprintType, editinlinenew, Category = ComputeFramework)
+class UOptimusGroomGuidesWriteDataProvider : public UComputeDataProvider
+{
+	GENERATED_BODY()
+
+public:
+	UPROPERTY(EditAnywhere, BlueprintReadWrite, Category = Binding)
+	TObjectPtr<UMeshComponent> MeshComponent = nullptr;
+
+	/** Output mask used to know the buffer format */
+	uint64 OutputMask = 0;
+
+	//~ Begin UComputeDataProvider Interface
+	virtual FComputeDataProviderRenderProxy* GetRenderProxy() override;
+	//~ End UComputeDataProvider Interface
+};
+
+class FOptimusGroomGuidesWriteProviderProxy : public FComputeDataProviderRenderProxy
+{
+public:
+	FOptimusGroomGuidesWriteProviderProxy(UMeshComponent* MeshComponent, uint64 InOutputMask);
+
+	//~ Begin FComputeDataProviderRenderProxy Interface
+	virtual bool IsValid(FValidationData const& InValidationData) const override;
+	virtual void AllocateResources(FRDGBuilder& GraphBuilder, FAllocationData const& InAllocationData) override;
+	virtual void GatherDispatchData(FDispatchData const& InDispatchData) override;
+	//~ End FComputeDataProviderRenderProxy Interface
+
+private:
+	using FParameters = FOptimusGroomGuidesWriteParameters;
+
+	/** List of instances (invocations) used in that data interface */
+	TArray<const FHairGroupInstance*> GroupInstances;
+
+	/** Output mask for gpu buffer format*/
+	uint64 OutputMask = 0;
+
+	/** Instance resources used to dispatch CS on GPU */
+	struct FInstanceResources
+	{
+		FRDGBufferSRVRef DeformedPositionOffset = nullptr;
+		FRDGBufferSRVRef PointRestPositions = nullptr;
+		FRDGBufferUAVRef PointDeformedPositions = nullptr;
+	};
+	TArray<FInstanceResources> InstanceResources;
+};
\ No newline at end of file
diff --git a/Engine/Plugins/Runtime/HairStrands/Source/HairStrandsDeformer/Private/DeformerGroomInterfaceUtils.h b/Engine/Plugins/Runtime/HairStrands/Source/HairStrandsDeformer/Private/DeformerGroomInterfaceUtils.h
new file mode 100644
--- /dev/null
+++ b/Engine/Plugins/Runtime/HairStrands/Source/HairStrandsDeformer/Private/DeformerGroomInterfaceUtils.h
@@ -0,0 +1,379 @@
+﻿// Copyright Epic Games, Inc. All Rights Reserved.
+
+#pragma once
+#include "GroomComponent.h"
+#include "GroomSolverComponent.h"
+#include "RenderGraphBuilder.h"
+#include "SkeletalMeshDeformerHelpers.h"
+#include "Components/SkeletalMeshComponent.h"
+#include "SkeletalRenderPublic.h"
+#include "SkinnedMeshSceneProxyDesc.h"
+#include "Rendering/SkeletalMeshRenderData.h"
+#include "DeformerGroomDomainsSource.h"
+
+namespace UE::Groom::Private
+{
+	FORCEINLINE USkeletalMeshComponent* GetGroupSkelMesh(const UGroomComponent* GroomComponent, const int32 GroupIndex, int32& MeshLOD)
+	{
+		if(GroomComponent)
+		{
+			TArray<UActorComponent*> ActorComponents;
+			if(AActor* RootActor = GroomComponent->GetAttachmentRootActor())
+			{
+				RootActor->GetComponents(USkeletalMeshComponent::StaticClass(), ActorComponents);
+			}
+			for(UActorComponent* ActorComponent : ActorComponents)
+			{
+				if(USkeletalMeshComponent* SkeletalMeshComponent = Cast<USkeletalMeshComponent>(ActorComponent))
+				{
+					if(GroomComponent->GroomAsset && (SkeletalMeshComponent->GetSkeletalMeshAsset() ==
+						GroomComponent->GroomAsset->GetDataflowSettings().GetSkeletalMesh(GroupIndex)))
+					{
+						MeshLOD = GroomComponent->GroomAsset->GetDataflowSettings().GetMeshLOD(GroupIndex);
+						return SkeletalMeshComponent;
+					}
+				}
+			}
+		}
+		MeshLOD = INDEX_NONE;
+		return nullptr;
+	}
+	
+	FORCEINLINE void GroomComponentsToInstances(const TArray<const UGroomComponent*>& GroomComponents, TArray<const FHairGroupInstance*>& GroupInstances)
+ 	{
+ 		GroupInstances.Reset();
+ 		for(const UGroomComponent* GroomComponent : GroomComponents)
+ 		{
+ 			if(GroomComponent)
+ 			{
+ 				const uint32 NumGroups = GroomComponent->GetGroupCount();
+ 				for (uint32 GroupIndex = 0; GroupIndex < NumGroups; ++GroupIndex)
+ 				{
+ 					if(GroomComponent->IsDeformationEnable(GroupIndex) && GroomComponent->GetGroupInstance(GroupIndex))
+ 					{
+ 						GroupInstances.Add(GroomComponent->GetGroupInstance(GroupIndex));
+ 					}
+ 				}
+ 			}
+ 		}
+ 	}
+	static void GroomComponentsToSkelmeshes(const TArray<const UGroomComponent*>& GroomComponents, TArray<const FSkeletalMeshObject*>& SkeletalMeshes, TArray<FMatrix44f>& SkeletalTransforms, TArray<TArray<FMatrix44f>>& BonesRefToLocals)
+	{
+		SkeletalMeshes.Reset();
+		SkeletalTransforms.Reset();
+		BonesRefToLocals.Reset();
+
+		TArray<FTransform> BonesTransforms;
+		TArray<FMatrix44f> RefToLocals;
+		int32 MeshLOD = INDEX_NONE;
+		
+		for(const UGroomComponent* GroomComponent : GroomComponents)
+		{
+			if(GroomComponent)
+			{
+				const uint32 NumGroups = GroomComponent->GetGroupCount();
+				for (uint32 GroupIndex = 0; GroupIndex < NumGroups; ++GroupIndex)
+				{
+					if(GroomComponent->IsDeformationEnable(GroupIndex) && GroomComponent->GetGroupInstance(GroupIndex))
+					{
+						if(USkeletalMeshComponent* SkelMesh =  GetGroupSkelMesh(GroomComponent, GroupIndex, MeshLOD))
+						{
+							const FSkeletalMeshRenderData* const RenderData = SkelMesh->GetSkinnedAsset()->GetResourceForRendering();
+
+							if(RenderData && RenderData->LODRenderData.IsValidIndex(MeshLOD))
+							{
+								// Get the matching active indices used for skinning
+								TArray<FBoneIndexType> GuidesBones = RenderData->LODRenderData[MeshLOD].ActiveBoneIndices;
+								
+								// Get inv ref pose matrices
+								const TArray<FMatrix44f>* RefBasesInvMatrix = &SkelMesh->GetSkinnedAsset()->GetRefBasesInvMatrix();
+
+								// Get the component space transforms
+								BonesTransforms.Init(FTransform::Identity, RefBasesInvMatrix->Num());
+								
+								SkelMesh->GetSkinnedAsset()->FillComponentSpaceTransforms(
+									SkelMesh->GetBoneSpaceTransforms(), GuidesBones, BonesTransforms);
+
+								// Fill ref to local matrices
+								RefToLocals.Init(FMatrix44f::Identity, RefBasesInvMatrix->Num());
+								
+								for (int32 BoneIndex = 0; BoneIndex < RefToLocals.Num(); ++BoneIndex)
+								{
+									RefToLocals[BoneIndex] = (*RefBasesInvMatrix)[BoneIndex] * (FMatrix44f)BonesTransforms[BoneIndex].ToMatrixWithScale();
+								}
+							
+								const FTransform& BonesTransform = SkelMesh->GetComponentTransform();
+								const FTransform& GroupTransform = GroomComponent->GetGroupInstance(GroupIndex)->GetCurrentLocalToWorld();
+								//SkelMesh->GetCurrentRefToLocalMatrices(RefToLocals, SkelMesh->MeshObject->GetLOD(), &RequiredExtraBoneIndices);//MeshLOD);
+		
+								SkeletalTransforms.Add(FMatrix44f(BonesTransform.ToMatrixWithScale() * GroupTransform.ToInverseMatrixWithScale()));
+								BonesRefToLocals.Add(RefToLocals);
+								SkeletalMeshes.Add(SkelMesh->MeshObject);
+							}
+						}
+					}
+				}
+			}
+		}
+	}
+	
+	FORCEINLINE void GatherGroomComponents(const UActorComponent* ActorComponent, TArray<const UGroomComponent*>& GroomComponents)
+	{
+		GroomComponents.Reset();
+		if(const UGroomComponent* GroomComponent = Cast<UGroomComponent>(ActorComponent))
+		{
+			GroomComponents.Add(GroomComponent);
+		}
+		else if(const UGroomSolverComponent* GroomSolver = Cast<UGroomSolverComponent>(ActorComponent))
+		{
+			for(const TObjectPtr<UGroomComponent>& SolverGroom : GroomSolver->GetGroomComponents())
+			{
+				GroomComponents.Add(SolverGroom);
+			}
+		}
+	}
+	
+	FORCEINLINE void GatherGroupInstances(const UActorComponent* ActorComponent, TArray<const FHairGroupInstance*>& GroupInstances)
+	{
+		TArray<const UGroomComponent*> GroomComponents;
+		GatherGroomComponents(ActorComponent, GroomComponents);
+
+		GroomComponentsToInstances(GroomComponents, GroupInstances);
+	}
+
+	FORCEINLINE void GatherGroupSkelmeshes(const UActorComponent* ActorComponent, TArray<const FSkeletalMeshObject*>& SkeletalMeshes, TArray<FMatrix44f>& SkeletalTransforms, TArray<TArray<FMatrix44f>>& BonesRefToLocals)
+	{
+		TArray<const UGroomComponent*> GroomComponents;
+		GatherGroomComponents(ActorComponent, GroomComponents);
+
+		GroomComponentsToSkelmeshes(GroomComponents, SkeletalMeshes, SkeletalTransforms, BonesRefToLocals);
+	}
+
+	template<typename DataType>
+	FORCEINLINE bool HaveValidInstanceResources(const DataType& InstanceData)
+	{
+		if(InstanceData.RestResource && InstanceData.DeformedResource)
+		{
+			if(InstanceData.RestResource->IsInitialized() && InstanceData.RestResource->bIsInitialized &&
+			   InstanceData.DeformedResource->IsInitialized() && InstanceData.DeformedResource->bIsInitialized)
+			{
+				return true;
+			}
+		}
+		return false;
+	}
+
+	FORCEINLINE bool HaveGuidesInstanceResources(const TArray<const FHairGroupInstance*>& GroupInstances)
+	{
+		for(const FHairGroupInstance* GroupInstance : GroupInstances)
+		{
+			if(!GroupInstance || (GroupInstance && !HaveValidInstanceResources(GroupInstance->Guides)))
+			{
+				return false;
+			}
+		}
+		return true;
+	}
+
+	FORCEINLINE bool HaveStrandsInstanceResources(const TArray<const FHairGroupInstance*>& GroupInstances)
+	{
+		for(const FHairGroupInstance* GroupInstance : GroupInstances)
+		{
+			if(!GroupInstance || (GroupInstance && !HaveValidInstanceResources(GroupInstance->Strands)))
+			{
+				return false;
+			}
+		}
+		return true;
+	}
+
+	template<typename DataType>
+	FORCEINLINE bool HaveValidSkinnedResources(const DataType& SkinnedData, const int32 LODIndex)
+	{
+		if(SkinnedData.HasValidRootData() && (SkinnedData.RestRootResource->GetRootCount() > 0) &&
+		  (SkinnedData.RestRootResource->LODs.Num() == SkinnedData.DeformedRootResource->LODs.Num()))
+		{
+			if(LODIndex < SkinnedData.RestRootResource->LODs.Num())
+			{
+				FHairStrandsLODRestRootResource* RestLODDatas = SkinnedData.RestRootResource->LODs[LODIndex];
+				FHairStrandsLODDeformedRootResource* DeformedLODDatas = SkinnedData.DeformedRootResource->LODs[LODIndex];
+
+				return (RestLODDatas && RestLODDatas->IsValid() && DeformedLODDatas && DeformedLODDatas->IsValid());
+			}
+		}
+		return false;
+	}
+
+	FORCEINLINE bool HaveGuidesSkinnedResources(const TArray<const FHairGroupInstance*>& GroupInstances)
+	{
+		for(const FHairGroupInstance* GroupInstance : GroupInstances)
+		{
+			if(GroupInstance && GroupInstance->HairGroupPublicData)
+			{
+				if(GroupInstance->BindingType == EHairBindingType::Skinning)
+				{
+					if(!HaveValidSkinnedResources(GroupInstance->Guides, GroupInstance->HairGroupPublicData->MeshLODIndex ))
+					{
+						return false;
+					}
+				}
+			}
+			else
+			{
+				return false;
+			}
+		}
+		return true;
+	}
+
+	FORCEINLINE bool HaveStrandsSkinnedResources(const TArray<const FHairGroupInstance*>& GroupInstances)
+	{
+		for(const FHairGroupInstance* GroupInstance : GroupInstances)
+		{
+			if(GroupInstance && GroupInstance->HairGroupPublicData)
+			{
+				if(GroupInstance->BindingType == EHairBindingType::Skinning)
+				{
+					if(!HaveValidSkinnedResources(GroupInstance->Strands, GroupInstance->HairGroupPublicData->MeshLODIndex ))
+					{
+						return false;
+					}
+				}
+			}
+			else
+			{
+				return false;
+			}
+		}
+		return true;
+	}
+
+	FORCEINLINE void GetGroomGroupElementOffsets(const UGroomComponent* GroomComponent, const FName DomainName, TArray<int32>& GroupOffsets, const int32 GroupSize = INDEX_NONE)
+	{
+		if(GroomComponent)
+		{
+			const uint32 NumGroups = GroomComponent->GetGroupCount();
+			const float GroupAlignment = static_cast<float>(GroupSize);
+			int32 MeshLOD = INDEX_NONE;
+			
+			GroupOffsets.Init(0, NumGroups+1);
+			int32 NumElements = 0;
+			for (uint32 GroupIndex = 0; GroupIndex < NumGroups; ++GroupIndex)
+			{
+				GroupOffsets[GroupIndex] = NumElements;
+				
+				if(DomainName == UOptimusGroomAssetComponentSource::FStrandsExecutionDomains::Objects ||
+				   DomainName == UOptimusGroomAssetComponentSource::FGuidesExecutionDomains::Objects)
+				{
+					NumElements++;
+				}
+				if(const USkeletalMeshComponent* SkelMesh =  GetGroupSkelMesh(GroomComponent,GroupIndex, MeshLOD))
+				{
+					if (SkelMesh->MeshObject)
+					{
+						if(DomainName == UOptimusGroomAssetComponentSource::FMeshesExecutionDomains::Bones)
+						{
+							NumElements += SkelMesh->MeshObject->GetReferenceToLocalMatrices().Num();
+						}
+						else if(DomainName == UOptimusGroomAssetComponentSource::FMeshesExecutionDomains::Vertices)
+						{
+							NumElements += 0; // @todo : compute the total number of vertices to modify
+						}
+					}
+				}
+				
+				if (const FHairGroupInstance* Instance = GroomComponent->GetGroupInstance(GroupIndex))
+				{
+					if(Instance->Strands.IsValid())
+					{
+						if(DomainName == UOptimusGroomAssetComponentSource::FStrandsExecutionDomains::Curves)
+						{
+							NumElements += Instance->Strands.GetData().GetNumCurves();
+						}
+						else if (DomainName == UOptimusGroomAssetComponentSource::FStrandsExecutionDomains::Edges)
+						{
+							NumElements += Instance->Strands.GetData().GetNumPoints() - Instance->Strands.GetData().GetNumCurves();
+						}
+						else if (DomainName == UOptimusGroomAssetComponentSource::FStrandsExecutionDomains::Points)
+						{
+							NumElements += Instance->Strands.GetData().GetNumPoints();
+						}
+					}
+					if(Instance->Guides.IsValid())
+					{
+						if(DomainName == UOptimusGroomAssetComponentSource::FGuidesExecutionDomains::Curves)
+						{
+							NumElements += Instance->Guides.GetData().GetNumCurves();
+						}
+						else if (DomainName == UOptimusGroomAssetComponentSource::FGuidesExecutionDomains::Edges)
+						{
+							NumElements += Instance->Guides.GetData().GetNumPoints() - Instance->Guides.GetData().GetNumCurves();
+						}
+						else if (DomainName == UOptimusGroomAssetComponentSource::FGuidesExecutionDomains::Points)
+						{
+							NumElements += Instance->Guides.GetData().GetNumPoints();
+						}
+					}
+				}
+				NumElements = (GroupSize == INDEX_NONE) ? NumElements :
+				   FMath::CeilToInt32(NumElements / GroupAlignment) * GroupAlignment;
+			}
+			GroupOffsets[NumGroups] = NumElements;
+		}
+	}
+
+	FORCEINLINE int32 GetGroomInvocationElementCounts(const TArray<const UGroomComponent*>& GroomComponents, const FName DomainName, TArray<int32>& InvocationCounts, int32 GroupSize = INDEX_NONE)
+	{
+		TArray<int32> GroupOffsets;
+		int32 TotalCount = 0;
+		InvocationCounts.Reset();
+		for(const UGroomComponent* GroomComponent : GroomComponents)
+		{
+			if(GroomComponent)
+			{
+				GetGroomGroupElementOffsets(GroomComponent, DomainName, GroupOffsets, GroupSize);
+
+				const uint32 NumGroups = GroomComponent->GetGroupCount();
+				for (uint32 GroupIndex = 0; GroupIndex < NumGroups; ++GroupIndex)
+				{
+					if(GroomComponent->IsDeformationEnable(GroupIndex) && GroomComponent->GetGroupInstance(GroupIndex))
+					{
+						InvocationCounts.Add(GroupOffsets[GroupIndex+1]-GroupOffsets[GroupIndex]);
+						TotalCount += InvocationCounts.Last();
+					}
+				}
+			}
+		}
+		return TotalCount;
+	}
+
+	struct FGroupElements
+	{
+		TArray<int32> GroupIndices;
+		TArray<int32> GroupOffsets;
+	};
+	
+	FORCEINLINE void GetGroomInvocationElementGroups(const TArray<const UGroomComponent*>& GroomComponents, const FName DomainName, TArray<TPair<UGroomAsset*, FGroupElements>>& InvocationGroups, int32 GroupSize = INDEX_NONE)
+	{
+		InvocationGroups.Reset();
+		for(const UGroomComponent* GroomComponent : GroomComponents)
+		{
+			if(GroomComponent)
+			{
+				FGroupElements GroupData;
+				GetGroomGroupElementOffsets(GroomComponent, DomainName, GroupData.GroupOffsets, GroupSize);
+				
+				const uint32 NumGroups = GroomComponent->GetGroupCount();
+				for (uint32 GroupIndex = 0; GroupIndex < NumGroups; ++GroupIndex)
+				{ 
+					if(GroomComponent->IsDeformationEnable(GroupIndex) && GroomComponent->GetGroupInstance(GroupIndex))
+					{
+						GroupData.GroupIndices.Add(GroupIndex);
+					}
+				}
+				InvocationGroups.Add(MakeTuple(GroomComponent->GroomAsset, GroupData));
+			}
+			
+		}
+	}
+}
diff --git a/Engine/Plugins/Runtime/HairStrands/Source/HairStrandsDeformer/Private/DeformerGroomMeshesRead.cpp b/Engine/Plugins/Runtime/HairStrands/Source/HairStrandsDeformer/Private/DeformerGroomMeshesRead.cpp
new file mode 100644
--- /dev/null
+++ b/Engine/Plugins/Runtime/HairStrands/Source/HairStrandsDeformer/Private/DeformerGroomMeshesRead.cpp
@@ -0,0 +1,184 @@
+﻿// Copyright Epic Games, Inc. All Rights Reserved.
+
+#include "DeformerGroomMeshesRead.h"
+
+#include "ComputeFramework/ShaderParamTypeDefinition.h"
+#include "ShaderCompilerCore.h"
+#include "ShaderParameterMetadataBuilder.h"
+#include "GroomComponent.h"
+#include "RenderGraphBuilder.h"
+#include "GlobalRenderResources.h"
+#include "DeformerGroomDomainsSource.h"
+#include "DeformerGroomInterfaceUtils.h"
+#include "SystemTextures.h"
+#include "SkeletalRenderPublic.h"
+
+FString UOptimusGroomMeshesReadDataInterface::GetDisplayName() const
+{
+	return TEXT("Groom Meshes");
+}
+
+TArray<FOptimusCDIPinDefinition> UOptimusGroomMeshesReadDataInterface::GetPinDefinitions() const
+{
+	static const FName MeshesBones(UOptimusGroomAssetComponentSource::FMeshesExecutionDomains::Bones);
+
+	TArray<FOptimusCDIPinDefinition> Defs;
+
+	// Bones buffers
+	Defs.Add({ "NumMeshesBones",   "ReadNumBones", false, "NumMeshesBones" });
+	Defs.Add({ "BoneTransformMatrix", "ReadBoneTransformMatrix", MeshesBones,   "ReadNumBones", true, "BoneTransformMatrix"  });
+
+	return Defs;
+}
+ 
+TSubclassOf<UActorComponent> UOptimusGroomMeshesReadDataInterface::GetRequiredComponentClass() const
+{
+	return UMeshComponent::StaticClass();
+}
+
+void UOptimusGroomMeshesReadDataInterface::GetPermutations(FComputeKernelPermutationVector& OutPermutationVector) const
+{}
+
+void UOptimusGroomMeshesReadDataInterface::GetSupportedInputs(TArray<FShaderFunctionDefinition>& OutFunctions) const
+{
+	OutFunctions.AddDefaulted_GetRef()
+		.SetName(TEXT("ReadNumBones"))
+		.AddReturnType(EShaderFundamentalType::Uint);
+
+	OutFunctions.AddDefaulted_GetRef()
+		.SetName(TEXT("ReadBoneTransformMatrix"))
+		.AddReturnType(EShaderFundamentalType::Float, 3, 4)
+		.AddParam(EShaderFundamentalType::Uint);
+}
+
+BEGIN_SHADER_PARAMETER_STRUCT(FOptimusGroomMeshesReadParameters, )
+	SHADER_PARAMETER(uint32, NumBones)
+	SHADER_PARAMETER_RDG_BUFFER_SRV(Buffer<float4>, BoneTransformMatrices)
+END_SHADER_PARAMETER_STRUCT()
+
+void UOptimusGroomMeshesReadDataInterface::GetShaderParameters(TCHAR const* UID, FShaderParametersMetadataBuilder& InOutBuilder, FShaderParametersMetadataAllocations& InOutAllocations) const
+{
+	InOutBuilder.AddNestedStruct<FOptimusGroomMeshesReadParameters>(UID);
+}
+
+TCHAR const* UOptimusGroomMeshesReadDataInterface::TemplateFilePath = TEXT("/Plugin/Runtime/HairStrands/Private/Deformers/DeformerGroomMeshesRead.ush");
+
+TCHAR const* UOptimusGroomMeshesReadDataInterface::GetShaderVirtualPath() const
+{
+	return TemplateFilePath;
+}
+
+void UOptimusGroomMeshesReadDataInterface::GetShaderHash(FString& InOutKey) const
+{
+	GetShaderFileHash(TemplateFilePath, EShaderPlatform::SP_PCD3D_SM5).AppendString(InOutKey);
+}
+
+void UOptimusGroomMeshesReadDataInterface::GetHLSL(FString& OutHLSL, FString const& InDataInterfaceName) const
+{
+	TMap<FString, FStringFormatArg> TemplateArgs =
+	{
+		{ TEXT("DataInterfaceName"), InDataInterfaceName },
+	};
+
+	FString TemplateFile;
+	LoadShaderSourceFile(TemplateFilePath, EShaderPlatform::SP_PCD3D_SM5, &TemplateFile, nullptr);
+	OutHLSL += FString::Format(*TemplateFile, TemplateArgs);
+}
+
+UComputeDataProvider* UOptimusGroomMeshesReadDataInterface::CreateDataProvider(TObjectPtr<UObject> InBinding, uint64 InInputMask, uint64 InOutputMask) const
+{
+	UOptimusGroomMeshesReadDataProvider* Provider = NewObject<UOptimusGroomMeshesReadDataProvider>();
+	Provider->MeshComponent = Cast<UMeshComponent>(InBinding);
+	return Provider;
+}
+
+FComputeDataProviderRenderProxy* UOptimusGroomMeshesReadDataProvider::GetRenderProxy()
+{
+	return new FOptimusGroomMeshesReadDataProviderProxy(MeshComponent);
+}
+
+FOptimusGroomMeshesReadDataProviderProxy::FOptimusGroomMeshesReadDataProviderProxy(UMeshComponent* MeshComponent)
+{
+	UE::Groom::Private::GatherGroupSkelmeshes(MeshComponent, SkeletalMeshObjects, SkeletalMeshTransforms, BonesRefToLocals);
+}
+
+bool FOptimusGroomMeshesReadDataProviderProxy::IsValid(FValidationData const& InValidationData) const
+{
+	if (InValidationData.ParameterStructSize != sizeof(FParameters))
+	{
+		return false;
+	}
+	if (InValidationData.NumInvocations != SkeletalMeshObjects.Num())
+	{
+		return false;
+	}
+	
+	return true;
+}
+
+FORCEINLINE void CreateBonesBuffers(const FSkeletalMeshObject* SkeletalMeshObject, const FMatrix44f& SkeletalMeshTransform, const TArray<FMatrix44f>& BoneRefToLocals, FRDGBuilder& GraphBuilder,
+	TArray<FRDGBufferRef>& BonesMatricesBuffers, TArray<FRDGBufferSRVRef>& BoneMatricesResources)
+{
+	if(SkeletalMeshObject)
+	{
+		//const TArray<FMatrix44f>& BoneRefToLocals = SkeletalMeshObject->GetReferenceToLocalMatrices();
+		const int32 NumBones = BoneRefToLocals.Num();
+		
+		TArray<FVector4f> BoneVectors;
+		BoneVectors.Init(FVector4f::Zero(), NumBones * 3);
+
+		int32 VectorIndex = 0;
+		FMatrix3x4 BoneMatrix;
+		for(int32 BoneIndex = 0; BoneIndex < NumBones; ++BoneIndex)
+		{
+			const FMatrix44f LocalMatrix = BoneRefToLocals[BoneIndex] * SkeletalMeshTransform;
+			LocalMatrix.To3x4MatrixTranspose(static_cast<float*>(&BoneMatrix.M[0][0]));
+			
+			BoneVectors[VectorIndex++] = FVector4f(BoneMatrix.M[0][0], BoneMatrix.M[0][1], BoneMatrix.M[0][2], BoneMatrix.M[0][3]);
+			BoneVectors[VectorIndex++] = FVector4f(BoneMatrix.M[1][0], BoneMatrix.M[1][1], BoneMatrix.M[1][2], BoneMatrix.M[1][3]);
+			BoneVectors[VectorIndex++] = FVector4f(BoneMatrix.M[2][0], BoneMatrix.M[2][1], BoneMatrix.M[2][2], BoneMatrix.M[2][3]);
+		}
+
+		FRDGBufferRef& BonesBuffer = BonesMatricesBuffers.Add_GetRef( GraphBuilder.CreateBuffer(
+			FRDGBufferDesc::CreateStructuredDesc(sizeof(FVector4f), BoneVectors.Num()), TEXT("Hair.Deformer.BoneTransformMatrices")));
+		GraphBuilder.QueueBufferUpload(BonesBuffer, BoneVectors.GetData(), sizeof(FVector4f) * BoneVectors.Num(), ERDGInitialDataFlags::None);
+
+		BoneMatricesResources.Add(GraphBuilder.CreateSRV(BonesBuffer));
+	}
+}
+
+void FOptimusGroomMeshesReadDataProviderProxy::AllocateResources(FRDGBuilder& GraphBuilder, FAllocationData const& InAllocationData)
+{
+	if (!FallbackStructuredSRV) { FallbackStructuredSRV  = GraphBuilder.CreateSRV(GSystemTextures.GetDefaultStructuredBuffer(GraphBuilder, 16u)); }
+
+	BoneMatricesResources.Reset();
+	BoneMatricesBuffers.Reset();
+
+	int32 SkelIndex = 0;
+	for (const FSkeletalMeshObject* SkeletalMeshObject : SkeletalMeshObjects)
+	{
+		CreateBonesBuffers(SkeletalMeshObject, SkeletalMeshTransforms[SkelIndex], BonesRefToLocals[SkelIndex], GraphBuilder, BoneMatricesBuffers, BoneMatricesResources);
+		++SkelIndex;
+	}
+}
+
+void FOptimusGroomMeshesReadDataProviderProxy::GatherDispatchData(FDispatchData const& InDispatchData)
+{
+	check(InDispatchData.NumInvocations == SkeletalMeshObjects.Num());
+	const TStridedView<FParameters> ParameterArray = MakeStridedParameterView<FParameters>(InDispatchData);
+	
+	for (int32 InvocationIndex = 0; InvocationIndex < ParameterArray.Num(); ++InvocationIndex)
+	{
+		FParameters& Parameters = ParameterArray[InvocationIndex];
+		if (const FSkeletalMeshObject* SkeletalMeshObject = SkeletalMeshObjects[InvocationIndex])
+		{
+			Parameters.NumBones = BonesRefToLocals[InvocationIndex].Num();
+			Parameters.BoneTransformMatrices = BoneMatricesResources[InvocationIndex];
+		}
+		else
+		{
+			Parameters.NumBones = 0;
+			Parameters.BoneTransformMatrices = FallbackStructuredSRV;
+		}
+	}
+}
\ No newline at end of file
diff --git a/Engine/Plugins/Runtime/HairStrands/Source/HairStrandsDeformer/Private/DeformerGroomMeshesRead.h b/Engine/Plugins/Runtime/HairStrands/Source/HairStrandsDeformer/Private/DeformerGroomMeshesRead.h
new file mode 100644
--- /dev/null
+++ b/Engine/Plugins/Runtime/HairStrands/Source/HairStrandsDeformer/Private/DeformerGroomMeshesRead.h
@@ -0,0 +1,94 @@
+﻿// Copyright Epic Games, Inc. All Rights Reserved.
+
+#pragma once
+
+#include "RenderGraphDefinitions.h"
+#include "OptimusComputeDataInterface.h"
+#include "HairStrandsInterface.h"
+#include "IOptimusDeformerInstanceAccessor.h"
+#include "ComputeFramework/ComputeDataProvider.h"
+#include "DeformerGroomMeshesRead.generated.h"
+
+class FOptimusGroomMeshesReadParameters;
+class UGroomComponent;
+struct FHairGroupInstance;
+class UGroomAsset;
+class UMeshComponent;
+class FSkeletalMeshObject;
+
+/** Compute Framework Data Interface for reading groom meshes. */
+UCLASS(Category = ComputeFramework)
+class UOptimusGroomMeshesReadDataInterface : public UOptimusComputeDataInterface
+{
+	GENERATED_BODY()
+
+public:
+	
+	//~ Begin UOptimusComputeDataInterface Interface
+	virtual FString GetDisplayName() const override;
+	virtual TArray<FOptimusCDIPinDefinition> GetPinDefinitions() const override;
+	virtual TSubclassOf<UActorComponent> GetRequiredComponentClass() const override;
+	//~ End UOptimusComputeDataInterface Interface
+	
+	//~ Begin UComputeDataInterface Interface
+	virtual TCHAR const* GetClassName() const override { return TEXT("GroomMeshesRead"); }
+	virtual void GetSupportedInputs(TArray<FShaderFunctionDefinition>& OutFunctions) const override;
+	virtual void GetShaderParameters(TCHAR const* UID, FShaderParametersMetadataBuilder& InOutBuilder, FShaderParametersMetadataAllocations& InOutAllocations) const override;
+	virtual TCHAR const* GetShaderVirtualPath() const override;
+	virtual void GetPermutations(FComputeKernelPermutationVector& OutPermutationVector) const override;
+	virtual void GetShaderHash(FString& InOutKey) const override;
+	virtual void GetHLSL(FString& OutHLSL, FString const& InDataInterfaceName) const override;
+	virtual UComputeDataProvider* CreateDataProvider(TObjectPtr<UObject> InBinding, uint64 InInputMask, uint64 InOutputMask) const override;
+	//~ End UComputeDataInterface Interface
+
+private:
+	/** File holding the hlsl implementation */
+	static TCHAR const* TemplateFilePath;
+};
+
+/** Compute Framework Data Provider for reading groom meshes. */
+UCLASS(BlueprintType, editinlinenew, Category = ComputeFramework)
+class UOptimusGroomMeshesReadDataProvider : public UComputeDataProvider
+{
+	GENERATED_BODY()
+
+public:
+	
+	UPROPERTY(EditAnywhere, BlueprintReadWrite, Category = Binding)
+	TObjectPtr<UMeshComponent> MeshComponent = nullptr;
+
+	//~ Begin UComputeDataProvider Interface
+	virtual FComputeDataProviderRenderProxy* GetRenderProxy() override;
+	//~ End UComputeDataProvider Interface
+};
+
+class FOptimusGroomMeshesReadDataProviderProxy : public FComputeDataProviderRenderProxy
+{
+public:
+	FOptimusGroomMeshesReadDataProviderProxy(UMeshComponent* MeshComponent);
+
+	//~ Begin FComputeDataProviderRenderProxy Interface
+	virtual bool IsValid(FValidationData const& InValidationData) const override;
+	virtual void AllocateResources(FRDGBuilder& GraphBuilder, FAllocationData const& InAllocationData) override;
+	virtual void GatherDispatchData(FDispatchData const& InDispatchData) override;
+	//~ End FComputeDataProviderRenderProxy Interface
+
+private:
+	using FParameters = FOptimusGroomMeshesReadParameters;
+	
+	/** Skeletal mesh objects to extract the bones from */
+	TArray<const FSkeletalMeshObject*> SkeletalMeshObjects;
+
+	/** Skeletal mesh transforms in group space */
+	TArray<FMatrix44f> SkeletalMeshTransforms;
+
+	/** Bones refs to locals matrices */
+	TArray<TArray<FMatrix44f>> BonesRefToLocals;
+	
+	/** Bones Resources used to dispatch CS on GPU */
+	TArray<FRDGBufferRef> BoneMatricesBuffers;
+	TArray<FRDGBufferSRVRef> BoneMatricesResources;
+
+	/** Fallback resources */
+	FRDGBufferSRVRef FallbackStructuredSRV = nullptr;
+};
diff --git a/Engine/Plugins/Runtime/HairStrands/Source/HairStrandsDeformer/Private/DeformerGroomSolverRead.cpp b/Engine/Plugins/Runtime/HairStrands/Source/HairStrandsDeformer/Private/DeformerGroomSolverRead.cpp
new file mode 100644
--- /dev/null
+++ b/Engine/Plugins/Runtime/HairStrands/Source/HairStrandsDeformer/Private/DeformerGroomSolverRead.cpp
@@ -0,0 +1,476 @@
+﻿// Copyright Epic Games, Inc. All Rights Reserved.
+
+#include "DeformerGroomSolverRead.h"
+
+#include "DeformerGroomDomainsExec.h"
+#include "ComputeFramework/ShaderParamTypeDefinition.h"
+#include "ComputeFramework/ComputeKernelPermutationVector.h"
+#include "ShaderCompilerCore.h"
+#include "ShaderParameterMetadataBuilder.h"
+#include "GroomComponent.h"
+#include "GroomSolverComponent.h"
+#include "GroomInstance.h"
+#include "RenderGraphBuilder.h"
+#include "GlobalRenderResources.h"
+#include "DeformerGroomDomainsSource.h"
+#include "DeformerGroomInterfaceUtils.h"
+#include "HairStrandsInterpolation.h"
+#include "OptimusDataTypeRegistry.h"
+#include "OptimusDeformerInstance.h"
+#include "OptimusValueContainerStruct.h"
+#include "SystemTextures.h"
+
+FString UOptimusGroomSolverReadDataInterface::GetDisplayName() const
+{
+	return TEXT("Groom Solver");
+}
+
+FName UOptimusGroomSolverReadDataInterface::GetResetSimulationTriggerName()
+{
+	return GET_MEMBER_NAME_CHECKED(UOptimusGroomSolverReadDataInterface, ResetSimulationTrigger);
+}
+
+TArray<FOptimusCDIPropertyPinDefinition> UOptimusGroomSolverReadDataInterface::GetPropertyPinDefinitions() const
+{
+	TArray<FOptimusCDIPropertyPinDefinition> PropertyPinDefinitions;
+
+	const FOptimusDataTypeHandle NameType = FOptimusDataTypeRegistry::Get().FindType(*FNameProperty::StaticClass());
+	
+	PropertyPinDefinitions.Add(
+		{GetResetSimulationTriggerName(), NameType}
+	);
+
+	return PropertyPinDefinitions;
+}
+
+TArray<FOptimusCDIPinDefinition> UOptimusGroomSolverReadDataInterface::GetPinDefinitions() const
+{
+	static const FName SolverObjects(UOptimusGroomSolverComponentSource::FSolverExecutionDomains::Objects);
+	static const FName SolverPoints(UOptimusGroomSolverComponentSource::FSolverExecutionDomains::Points);
+	static const FName SolverCurves(UOptimusGroomSolverComponentSource::FSolverExecutionDomains::Curves);
+
+	static const FName DynamicPoints(UOptimusGroomSolverComponentSource::FDynamicExecutionDomains::Points);
+	static const FName DynamicCurves(UOptimusGroomSolverComponentSource::FDynamicExecutionDomains::Curves);
+	static const FName KinematicPoints(UOptimusGroomSolverComponentSource::FKinematicExecutionDomains::Points);
+	static const FName KinematicCurves(UOptimusGroomSolverComponentSource::FKinematicExecutionDomains::Curves);
+
+	TArray<FOptimusCDIPinDefinition> Defs;
+
+	Defs.Add({ "NumSolverPoints",  "ReadNumSolverPoints", true, "NumSolverPoints" });
+	Defs.Add({ "NumSolverCurves",  "ReadNumSolverCurves", true, "NumSolverCurves" });
+	Defs.Add({ "NumSolverObjects",  "ReadNumSolverObjects", true, "NumSolverObjects" });
+	Defs.Add({ "NumDynamicPoints",  "ReadNumDynamicPoints", true, "NumDynamicPoints" });
+	Defs.Add({ "NumDynamicCurves",  "ReadNumDynamicCurves", true, "NumDynamicCurves" });
+	Defs.Add({ "NumKinematicPoints",  "ReadNumKinematicPoints", true, "NumKinematicPoints" });
+	Defs.Add({ "NumKinematicCurves",  "ReadNumKinematicCurves", true, "NumKinematicCurves" });
+	Defs.Add({ "ObjectPointOffset", "ReadObjectPointOffset", SolverObjects,          "ReadNumSolverObjects", true, "ObjectPointOffset"  });
+	Defs.Add({ "ObjectNumPoints", "ReadObjectNumPoints", SolverObjects,          "ReadNumSolverObjects", true, "ObjectNumPoints"  });
+	Defs.Add({ "ObjectCurveOffset", "ReadObjectCurveOffset", SolverObjects,          "ReadNumSolverObjects", true, "ObjectCurveOffset"  });
+	Defs.Add({ "ObjectNumCurves", "ReadObjectNumCurves", SolverObjects,          "ReadNumSolverObjects", true, "ObjectNumCurves"  });
+	Defs.Add({ "PointObjectIndex", "ReadPointObjectIndex", SolverPoints,          "ReadNumSolverPoints", true, "PointObjectIndex"  });
+	Defs.Add({ "CurveObjectIndex", "ReadCurveObjectIndex", SolverCurves,          "ReadNumSolverCurves", true, "CurveObjectIndex"  });
+	Defs.Add({ "DynamicPointIndex", "ReadDynamicPointIndex", DynamicPoints,          "ReadNumDynamicPoints", true, "DynamicPointIndex"  });
+	Defs.Add({ "DynamicCurveIndex", "ReadDynamicCurveIndex", DynamicCurves,          "ReadNumDynamicCurves", true, "DynamicCurveIndex"  });
+	Defs.Add({ "KinematicPointIndex", "ReadKinematicPointIndex", KinematicPoints,          "ReadNumKinematicPoints", true, "KinematicPointIndex"  });
+	Defs.Add({ "KinematicCurveIndex", "ReadKinematicCurveIndex", KinematicCurves,          "ReadNumKinematicCurves", true, "KinematicCurveIndex"  });
+	Defs.Add({ "ObjectDistanceLod", "ReadObjectDistanceLod", SolverObjects,          "ReadObjectDistanceLod", true, "ObjectDistanceLod"  });
+
+	return Defs;
+}
+ 
+TSubclassOf<UActorComponent> UOptimusGroomSolverReadDataInterface::GetRequiredComponentClass() const
+{
+	return UGroomSolverComponent::StaticClass();
+}
+
+void UOptimusGroomSolverReadDataInterface::GetPermutations(FComputeKernelPermutationVector& OutPermutationVector) const
+{
+}
+
+void UOptimusGroomSolverReadDataInterface::GetSupportedInputs(TArray<FShaderFunctionDefinition>& OutFunctions) const
+{
+	OutFunctions.AddDefaulted_GetRef()
+		.SetName(TEXT("ReadNumSolverPoints"))
+		.AddReturnType(EShaderFundamentalType::Uint);
+
+	OutFunctions.AddDefaulted_GetRef()
+		.SetName(TEXT("ReadNumSolverCurves"))
+		.AddReturnType(EShaderFundamentalType::Uint);
+
+	OutFunctions.AddDefaulted_GetRef()
+		.SetName(TEXT("ReadNumSolverObjects"))
+		.AddReturnType(EShaderFundamentalType::Uint);
+
+	OutFunctions.AddDefaulted_GetRef()
+		.SetName(TEXT("ReadNumDynamicPoints"))
+		.AddReturnType(EShaderFundamentalType::Uint);
+
+	OutFunctions.AddDefaulted_GetRef()
+		.SetName(TEXT("ReadNumDynamicCurves"))
+		.AddReturnType(EShaderFundamentalType::Uint);
+
+	OutFunctions.AddDefaulted_GetRef()
+		.SetName(TEXT("ReadNumKinematicPoints"))
+		.AddReturnType(EShaderFundamentalType::Uint);
+
+	OutFunctions.AddDefaulted_GetRef()
+		.SetName(TEXT("ReadNumKinematicCurves"))
+		.AddReturnType(EShaderFundamentalType::Uint);
+
+	OutFunctions.AddDefaulted_GetRef()
+		.SetName(TEXT("ReadObjectPointOffset"))
+		.AddReturnType(EShaderFundamentalType::Uint)
+		.AddParam(EShaderFundamentalType::Uint);
+		
+	OutFunctions.AddDefaulted_GetRef()
+		.SetName(TEXT("ReadObjectCurveOffset"))
+		.AddReturnType(EShaderFundamentalType::Uint)
+		.AddParam(EShaderFundamentalType::Uint);
+
+	OutFunctions.AddDefaulted_GetRef()
+		.SetName(TEXT("ReadObjectNumPoints"))
+		.AddReturnType(EShaderFundamentalType::Uint)
+		.AddParam(EShaderFundamentalType::Uint);
+		
+	OutFunctions.AddDefaulted_GetRef()
+		.SetName(TEXT("ReadObjectNumCurves"))
+		.AddReturnType(EShaderFundamentalType::Uint)
+		.AddParam(EShaderFundamentalType::Uint);
+
+	OutFunctions.AddDefaulted_GetRef()
+		.SetName(TEXT("ReadPointObjectIndex"))
+		.AddReturnType(EShaderFundamentalType::Int)
+		.AddParam(EShaderFundamentalType::Uint);
+		
+	OutFunctions.AddDefaulted_GetRef()
+		.SetName(TEXT("ReadCurveObjectIndex"))
+		.AddReturnType(EShaderFundamentalType::Int)
+		.AddParam(EShaderFundamentalType::Uint);
+
+	OutFunctions.AddDefaulted_GetRef()
+		.SetName(TEXT("ReadDynamicPointIndex"))
+		.AddReturnType(EShaderFundamentalType::Int)
+		.AddParam(EShaderFundamentalType::Uint);
+
+	OutFunctions.AddDefaulted_GetRef()
+		.SetName(TEXT("ReadDynamicCurveIndex"))
+		.AddReturnType(EShaderFundamentalType::Int)
+		.AddParam(EShaderFundamentalType::Uint);
+
+	OutFunctions.AddDefaulted_GetRef()
+		.SetName(TEXT("ReadKinematicPointIndex"))
+		.AddReturnType(EShaderFundamentalType::Int)
+		.AddParam(EShaderFundamentalType::Uint);
+
+	OutFunctions.AddDefaulted_GetRef()
+		.SetName(TEXT("ReadKinematicCurveIndex"))
+		.AddReturnType(EShaderFundamentalType::Int)
+		.AddParam(EShaderFundamentalType::Uint);
+
+	OutFunctions.AddDefaulted_GetRef()
+		.SetName(TEXT("ReadObjectDistanceLod"))
+		.AddReturnType(EShaderFundamentalType::Uint)
+		.AddParam(EShaderFundamentalType::Uint);
+}
+
+BEGIN_SHADER_PARAMETER_STRUCT(FOptimusGroomSolverReadParameters, )
+	SHADER_PARAMETER(uint32, NumSolverPoints)
+	SHADER_PARAMETER(uint32, NumSolverCurves)
+	SHADER_PARAMETER(uint32, NumSolverObjects)
+	SHADER_PARAMETER(uint32, NumDynamicPoints)
+	SHADER_PARAMETER(uint32, NumDynamicCurves)
+	SHADER_PARAMETER(uint32, NumKinematicPoints)
+	SHADER_PARAMETER(uint32, NumKinematicCurves)
+	SHADER_PARAMETER_RDG_BUFFER_SRV(Buffer<uint>, ObjectPointOffsets)
+	SHADER_PARAMETER_RDG_BUFFER_SRV(Buffer<uint>, ObjectCurveOffsets)
+	SHADER_PARAMETER_RDG_BUFFER_SRV(Buffer<uint>, ObjectNumPoints)
+	SHADER_PARAMETER_RDG_BUFFER_SRV(Buffer<uint>, ObjectNumCurves)
+	SHADER_PARAMETER_RDG_BUFFER_SRV(Buffer<int>, PointObjectIndices)
+	SHADER_PARAMETER_RDG_BUFFER_SRV(Buffer<int>, CurveObjectIndices)
+	SHADER_PARAMETER_RDG_BUFFER_SRV(Buffer<int>, DynamicPointIndices)
+	SHADER_PARAMETER_RDG_BUFFER_SRV(Buffer<int>, DynamicCurveIndices)
+	SHADER_PARAMETER_RDG_BUFFER_SRV(Buffer<int>, KinematicPointIndices)
+	SHADER_PARAMETER_RDG_BUFFER_SRV(Buffer<int>, KinematicCurveIndices)
+	SHADER_PARAMETER_RDG_BUFFER_SRV(Buffer<uint>, ObjectDistanceLods)
+END_SHADER_PARAMETER_STRUCT()
+
+void UOptimusGroomSolverReadDataInterface::GetShaderParameters(TCHAR const* UID, FShaderParametersMetadataBuilder& InOutBuilder, FShaderParametersMetadataAllocations& InOutAllocations) const
+{
+	InOutBuilder.AddNestedStruct<FOptimusGroomSolverReadParameters>(UID);
+}
+
+TCHAR const* UOptimusGroomSolverReadDataInterface::TemplateFilePath = TEXT("/Plugin/Runtime/HairStrands/Private/Deformers/DeformerGroomSolverRead.ush");
+
+TCHAR const* UOptimusGroomSolverReadDataInterface::GetShaderVirtualPath() const
+{
+	return TemplateFilePath;
+}
+
+void UOptimusGroomSolverReadDataInterface::GetShaderHash(FString& InOutKey) const
+{
+	GetShaderFileHash(TemplateFilePath, EShaderPlatform::SP_PCD3D_SM5).AppendString(InOutKey);
+}
+
+void UOptimusGroomSolverReadDataInterface::GetHLSL(FString& OutHLSL, FString const& InDataInterfaceName) const
+{
+	TMap<FString, FStringFormatArg> TemplateArgs =
+	{
+		{ TEXT("DataInterfaceName"), InDataInterfaceName },
+	};
+
+	FString TemplateFile;
+	LoadShaderSourceFile(TemplateFilePath, EShaderPlatform::SP_PCD3D_SM5, &TemplateFile, nullptr);
+	OutHLSL += FString::Format(*TemplateFile, TemplateArgs);
+}
+
+UComputeDataProvider* UOptimusGroomSolverReadDataInterface::CreateDataProvider(TObjectPtr<UObject> InBinding, uint64 InInputMask, uint64 InOutputMask) const
+{
+	UOptimusGroomSolverReadDataProvider* Provider = NewObject<UOptimusGroomSolverReadDataProvider>();
+	Provider->SolverComponent = Cast<UGroomSolverComponent>(InBinding);
+	Provider->WeakDataInterface = this;
+	return Provider;
+}
+
+void UOptimusGroomSolverReadDataProvider::SetDeformerInstance(UOptimusDeformerInstance* InInstance)
+{
+	DeformerInstance = InInstance;
+}
+
+FComputeDataProviderRenderProxy* UOptimusGroomSolverReadDataProvider::GetRenderProxy()
+{
+	static constexpr int32 MaxResetCount = 4;
+	bool bResetSimulationTrigger = false;
+	if (const UOptimusGroomSolverReadDataInterface* DataInterface = WeakDataInterface.Get())
+	{
+		FOptimusValueContainerStruct ValueContainer =
+					DeformerInstance->GetDataInterfacePropertyOverride(
+						DataInterface,
+						UOptimusGroomSolverReadDataInterface::GetResetSimulationTriggerName()
+						);
+
+		TValueOrError<FName, EPropertyBagResult> Value = ValueContainer.Value.GetValueName(FOptimusValueContainerStruct::ValuePropertyName);
+		if (Value.HasValue())
+		{
+			const FName& ResetSimulationGraphName = Value.GetValue();
+
+			TArray<const UGroomComponent*> GroomComponents;
+			UE::Groom::Private::GatherGroomComponents(SolverComponent, GroomComponents);
+
+			for(const UGroomComponent* GroomComponent : GroomComponents)
+			{
+				if(GroomComponent->bResetSimulation)
+				{
+					bResetSimulationCount = 0;
+					break;
+				}
+			}
+			if(bResetSimulationCount > 0)
+			{
+				TArray<const FHairGroupInstance*> GroupInstances;
+				UE::Groom::Private::GroomComponentsToInstances(GroomComponents, GroupInstances);
+				
+				if(!UE::Groom::Private::HaveGuidesInstanceResources(GroupInstances) || !UE::Groom::Private::HaveGuidesSkinnedResources(GroupInstances))
+				{
+					bResetSimulationCount = 0;
+				}
+			}
+			bResetSimulationTrigger = (bResetSimulationCount < MaxResetCount);
+			if(bResetSimulationTrigger)
+			{
+				DeformerInstance->EnqueueTriggerGraph(ResetSimulationGraphName);
+				bResetSimulationCount++;
+			}
+		}
+	}
+	return new FOptimusGroomSolverReadDataProviderProxy(SolverComponent, bResetSimulationTrigger);
+}
+
+FOptimusGroomSolverReadDataProviderProxy::FOptimusGroomSolverReadDataProviderProxy(UGroomSolverComponent* SolverComponent, const bool bResetSimulation)
+{
+	TArray<const UGroomComponent*> GroomComponents;
+	UE::Groom::Private::GatherGroomComponents(SolverComponent, GroomComponents);
+
+	UE::Groom::Private::GroomComponentsToInstances(GroomComponents, GroupInstances);
+	
+	UE::Groom::Private::GetGroomInvocationElementCounts(GroomComponents,
+		UOptimusGroomAssetComponentSource::FGuidesExecutionDomains::Points, PointsCounts);
+
+	UE::Groom::Private::GetGroomInvocationElementCounts(GroomComponents,
+		UOptimusGroomAssetComponentSource::FGuidesExecutionDomains::Points, AlignedPoints, UOptimusGroomSolverComponentSource::GroupSize);
+
+	UE::Groom::Private::GetGroomInvocationElementCounts(GroomComponents,
+		UOptimusGroomAssetComponentSource::FGuidesExecutionDomains::Curves, CurvesCounts);
+	
+	SolverSettings = SolverComponent->GetSolverSettings();
+	bResetSimulationTrigger = bResetSimulation;
+}
+
+bool FOptimusGroomSolverReadDataProviderProxy::IsValid(FValidationData const& InValidationData) const
+{
+	if (InValidationData.ParameterStructSize != sizeof(FParameters))
+	{
+		return false;
+	}
+	if(!UE::Groom::Private::HaveGuidesInstanceResources(GroupInstances) || !UE::Groom::Private::HaveGuidesSkinnedResources(GroupInstances))
+	{
+		return false;
+	}
+	if(bResetSimulationTrigger)
+	{
+		return false;
+	}
+	
+	return true;
+}
+
+void FOptimusGroomSolverReadDataProviderProxy::GatherPermutations(FPermutationData& InOutPermutationData) const
+{
+}
+
+FORCEINLINE void FillElementsData(const TArray<int32>& ElementsCounts, const TArray<int32>& AlignedCounts,
+	TArray<uint32>& ObjectOffsets, TArray<uint32>& ObjectCounts, TArray<int32>& ElementObjects, int32& NumElements)
+{
+	const uint32 NumObjects = ElementsCounts.Num();
+	
+	ObjectOffsets.Init(0, NumObjects+1);
+	ObjectCounts.Init(0, NumObjects);
+	uint32 ElementOffset = 0;
+	for (uint32 ObjectIndex = 0; ObjectIndex < NumObjects; ++ObjectIndex)
+	{
+		ObjectCounts[ObjectIndex] = static_cast<uint32>(ElementsCounts[ObjectIndex]);
+		ObjectOffsets[ObjectIndex] = ElementOffset;
+		ElementOffset += AlignedCounts[ObjectIndex];
+	}
+	
+	ObjectOffsets[NumObjects] = ElementOffset;
+	NumElements = ElementOffset;
+
+	ElementObjects.Init(INDEX_NONE, NumElements);
+	for (uint32 ObjectIndex = 0; ObjectIndex < NumObjects; ++ObjectIndex)
+	{
+		for(uint32 ElementIndex = ObjectOffsets[ObjectIndex],
+			ElementEnd = ObjectOffsets[ObjectIndex] + ObjectCounts[ObjectIndex]; ElementIndex < ElementEnd; ++ElementIndex)
+		{
+			ElementObjects[ElementIndex] = ObjectIndex;
+		}
+	}
+}
+
+void FOptimusGroomSolverReadDataProviderProxy::CreateInternalBuffers(FRDGBuilder& GraphBuilder)
+{
+	TArray<uint32> ObjectPointOffsets;
+	TArray<uint32> ObjectNumPoints;
+	TArray<int32> PointObjectIndices;
+
+	FillElementsData(PointsCounts, AlignedPoints, ObjectPointOffsets,  ObjectNumPoints,
+		PointObjectIndices, NumPoints);
+
+	FString BufferName = FString("Hair.Deformer.ObjectPointOffsets");
+	ObjectPointOffsetsBuffer = GraphBuilder.CreateBuffer(FRDGBufferDesc::CreateStructuredDesc(sizeof(uint32), ObjectPointOffsets.Num()), *BufferName);
+	ObjectPointOffsetsResource = GraphBuilder.CreateSRV(ObjectPointOffsetsBuffer);
+	GraphBuilder.QueueBufferUpload(ObjectPointOffsetsBuffer, ObjectPointOffsets.GetData(), sizeof(uint32) * ObjectPointOffsets.Num(), ERDGInitialDataFlags::None);
+
+	BufferName = FString("Hair.Deformer.ObjectNumPoints");
+	ObjectNumPointsBuffer = GraphBuilder.CreateBuffer(FRDGBufferDesc::CreateStructuredDesc(sizeof(uint32), ObjectNumPoints.Num()), *BufferName);
+	ObjectNumPointsResource = GraphBuilder.CreateSRV(ObjectNumPointsBuffer);
+	GraphBuilder.QueueBufferUpload(ObjectNumPointsBuffer, ObjectNumPoints.GetData(), sizeof(uint32) * ObjectNumPoints.Num(), ERDGInitialDataFlags::None);
+
+	BufferName = FString("Hair.Deformer.PointObjectIndices");
+	PointObjectIndicesBuffer = GraphBuilder.CreateBuffer(FRDGBufferDesc::CreateStructuredDesc(sizeof(int32), PointObjectIndices.Num()), *BufferName);
+	PointObjectIndicesResource = GraphBuilder.CreateSRV(PointObjectIndicesBuffer);
+	GraphBuilder.QueueBufferUpload(PointObjectIndicesBuffer, PointObjectIndices.GetData(), sizeof(int32) * PointObjectIndices.Num(), ERDGInitialDataFlags::None);
+
+	BufferName = FString("Hair.Deformer.DynamicPointIndices");
+	DynamicPointIndicesBuffer = GraphBuilder.CreateBuffer(FRDGBufferDesc::CreateStructuredDesc(sizeof(int32), SolverSettings.PointDynamicIndices.Num()), *BufferName);
+	DynamicPointIndicesResource = GraphBuilder.CreateSRV(DynamicPointIndicesBuffer);
+	GraphBuilder.QueueBufferUpload(DynamicPointIndicesBuffer, SolverSettings.PointDynamicIndices.GetData(), sizeof(int32) * SolverSettings.PointDynamicIndices.Num(), ERDGInitialDataFlags::None);
+
+	BufferName = FString("Hair.Deformer.KinematicPointIndices");
+	KinematicPointIndicesBuffer = GraphBuilder.CreateBuffer(FRDGBufferDesc::CreateStructuredDesc(sizeof(int32), SolverSettings.PointKinematicIndices.Num()), *BufferName);
+	KinematicPointIndicesResource = GraphBuilder.CreateSRV(KinematicPointIndicesBuffer);
+	GraphBuilder.QueueBufferUpload(KinematicPointIndicesBuffer, SolverSettings.PointKinematicIndices.GetData(), sizeof(int32) * SolverSettings.PointKinematicIndices.Num(), ERDGInitialDataFlags::None);
+	
+	TArray<uint32> ObjectCurveOffsets;
+	TArray<uint32> ObjectNumCurves;
+	TArray<int32> CurveObjectIndices;
+	
+	FillElementsData(CurvesCounts, CurvesCounts, ObjectCurveOffsets,  ObjectNumCurves,
+	CurveObjectIndices, NumCurves);
+
+	BufferName = FString("Hair.Deformer.ObjectCurveOffsets");
+	ObjectCurveOffsetsBuffer = GraphBuilder.CreateBuffer(FRDGBufferDesc::CreateStructuredDesc(sizeof(uint32), ObjectCurveOffsets.Num()), *BufferName);
+	ObjectCurveOffsetsResource = GraphBuilder.CreateSRV(ObjectCurveOffsetsBuffer);
+	GraphBuilder.QueueBufferUpload(ObjectCurveOffsetsBuffer, ObjectCurveOffsets.GetData(), sizeof(uint32) * ObjectCurveOffsets.Num(), ERDGInitialDataFlags::None);
+	
+	BufferName = FString("Hair.Deformer.ObjectNumCurves");
+	ObjectNumCurvesBuffer = GraphBuilder.CreateBuffer(FRDGBufferDesc::CreateStructuredDesc(sizeof(uint32), ObjectNumCurves.Num()), *BufferName);
+	ObjectNumCurvesResource = GraphBuilder.CreateSRV(ObjectNumCurvesBuffer);
+	GraphBuilder.QueueBufferUpload(ObjectNumCurvesBuffer, ObjectNumCurves.GetData(), sizeof(uint32) * ObjectNumCurves.Num(), ERDGInitialDataFlags::None);
+
+	BufferName = FString("Hair.Deformer.CurveObjectIndices");
+	CurveObjectIndicesBuffer = GraphBuilder.CreateBuffer(FRDGBufferDesc::CreateStructuredDesc(sizeof(int32), CurveObjectIndices.Num()), *BufferName);
+	CurveObjectIndicesResource = GraphBuilder.CreateSRV(CurveObjectIndicesBuffer);
+	GraphBuilder.QueueBufferUpload(CurveObjectIndicesBuffer, CurveObjectIndices.GetData(), sizeof(int32) * CurveObjectIndices.Num(), ERDGInitialDataFlags::None);
+
+	BufferName = FString("Hair.Deformer.DynamicCurveIndices");
+	DynamicCurveIndicesBuffer = GraphBuilder.CreateBuffer(FRDGBufferDesc::CreateStructuredDesc(sizeof(int32), SolverSettings.CurveDynamicIndices.Num()), *BufferName);
+	DynamicCurveIndicesResource = GraphBuilder.CreateSRV(DynamicCurveIndicesBuffer);
+	GraphBuilder.QueueBufferUpload(DynamicCurveIndicesBuffer, SolverSettings.CurveDynamicIndices.GetData(), sizeof(int32) * SolverSettings.CurveDynamicIndices.Num(), ERDGInitialDataFlags::None);
+
+	BufferName = FString("Hair.Deformer.KinematicCurveIndices");
+	KinematicCurveIndicesBuffer = GraphBuilder.CreateBuffer(FRDGBufferDesc::CreateStructuredDesc(sizeof(int32), SolverSettings.CurveKinematicIndices.Num()), *BufferName);
+	KinematicCurveIndicesResource = GraphBuilder.CreateSRV(KinematicCurveIndicesBuffer);
+	GraphBuilder.QueueBufferUpload(KinematicCurveIndicesBuffer, SolverSettings.CurveKinematicIndices.GetData(), sizeof(int32) * SolverSettings.CurveKinematicIndices.Num(), ERDGInitialDataFlags::None);
+
+	BufferName = FString("Hair.Deformer.ObjectDistanceLods");
+	ObjectDistanceLodsBuffer = GraphBuilder.CreateBuffer(FRDGBufferDesc::CreateStructuredDesc(sizeof(uint32), SolverSettings.ObjectDistanceLods.Num()), *BufferName);
+	ObjectDistanceLodsResource = GraphBuilder.CreateSRV(ObjectDistanceLodsBuffer);
+	GraphBuilder.QueueBufferUpload(ObjectDistanceLodsBuffer, SolverSettings.ObjectDistanceLods.GetData(), sizeof(uint32) * SolverSettings.ObjectDistanceLods.Num(), ERDGInitialDataFlags::None);
+
+}
+
+void FOptimusGroomSolverReadDataProviderProxy::AllocateResources(FRDGBuilder& GraphBuilder, FAllocationData const& InAllocationData)
+{
+	CreateInternalBuffers(GraphBuilder);
+}
+
+void FOptimusGroomSolverReadDataProviderProxy::GatherDispatchData(FDispatchData const& InDispatchData)
+{ 
+	const TStridedView<FParameters> ParameterArray = MakeStridedParameterView<FParameters>(InDispatchData);
+	check(!ParameterArray.IsEmpty());
+
+	for (int32 InvocationIndex = 0; InvocationIndex < ParameterArray.Num(); ++InvocationIndex)
+	{
+		FParameters& Parameters = ParameterArray[InvocationIndex];
+
+		Parameters.NumSolverPoints = NumPoints;
+		Parameters.NumSolverCurves = NumCurves;
+		Parameters.NumSolverObjects = GroupInstances.Num();
+
+		Parameters.NumDynamicPoints = SolverSettings.PointDynamicIndices.Num();
+		Parameters.NumDynamicCurves = SolverSettings.CurveDynamicIndices.Num();
+		Parameters.NumKinematicPoints = SolverSettings.PointKinematicIndices.Num();
+		Parameters.NumKinematicCurves = SolverSettings.CurveKinematicIndices.Num();
+
+		Parameters.ObjectPointOffsets = ObjectPointOffsetsResource;
+		Parameters.ObjectCurveOffsets = ObjectCurveOffsetsResource;
+
+		Parameters.ObjectNumPoints = ObjectNumPointsResource;
+		Parameters.ObjectNumCurves = ObjectNumCurvesResource;
+
+		Parameters.PointObjectIndices = PointObjectIndicesResource;
+		Parameters.CurveObjectIndices = CurveObjectIndicesResource;
+
+		Parameters.DynamicPointIndices = DynamicPointIndicesResource;
+		Parameters.DynamicCurveIndices = DynamicCurveIndicesResource;
+
+		Parameters.KinematicPointIndices = KinematicPointIndicesResource;
+		Parameters.KinematicCurveIndices = KinematicCurveIndicesResource;
+
+		Parameters.ObjectDistanceLods = ObjectDistanceLodsResource;
+	}
+}
+
+void FOptimusGroomSolverReadDataProviderProxy::PostGraphDispatch(FRDGBuilder& GraphBuilder) const
+{
+}
\ No newline at end of file
diff --git a/Engine/Plugins/Runtime/HairStrands/Source/HairStrandsDeformer/Private/DeformerGroomSolverRead.h b/Engine/Plugins/Runtime/HairStrands/Source/HairStrandsDeformer/Private/DeformerGroomSolverRead.h
new file mode 100644
--- /dev/null
+++ b/Engine/Plugins/Runtime/HairStrands/Source/HairStrandsDeformer/Private/DeformerGroomSolverRead.h
@@ -0,0 +1,158 @@
+﻿// Copyright Epic Games, Inc. All Rights Reserved.
+
+#pragma once
+
+#include "GroomSolverComponent.h"
+#include "RenderGraphDefinitions.h"
+#include "OptimusComputeDataInterface.h"
+#include "HairStrandsInterface.h"
+#include "IOptimusDeformerInstanceAccessor.h"
+#include "ComputeFramework/ComputeDataProvider.h"
+#include "DeformerGroomSolverRead.generated.h"
+
+namespace UE::Groom::Private
+{
+	struct FGroupElements;
+}
+
+class FOptimusGroomSolverReadParameters;
+class UGroomComponent;
+struct FHairGroupInstance;
+class UGroomAsset;
+
+/** Compute Framework Data Interface for reading groom guides. */
+UCLASS(Category = ComputeFramework)
+class UOptimusGroomSolverReadDataInterface : public UOptimusComputeDataInterface
+{
+	GENERATED_BODY()
+
+public:
+
+	static FName GetResetSimulationTriggerName();
+	
+	//~ Begin UOptimusComputeDataInterface Interface
+	virtual FString GetDisplayName() const override;
+	virtual TArray<FOptimusCDIPinDefinition> GetPinDefinitions() const override;
+	virtual TSubclassOf<UActorComponent> GetRequiredComponentClass() const override;
+	virtual TArray<FOptimusCDIPropertyPinDefinition> GetPropertyPinDefinitions() const override;
+	//~ End UOptimusComputeDataInterface Interface
+	
+	//~ Begin UComputeDataInterface Interface
+	virtual TCHAR const* GetClassName() const override { return TEXT("GroomSolverRead"); }
+	virtual void GetSupportedInputs(TArray<FShaderFunctionDefinition>& OutFunctions) const override;
+	virtual void GetShaderParameters(TCHAR const* UID, FShaderParametersMetadataBuilder& InOutBuilder, FShaderParametersMetadataAllocations& InOutAllocations) const override;
+	virtual TCHAR const* GetShaderVirtualPath() const override;
+	virtual void GetPermutations(FComputeKernelPermutationVector& OutPermutationVector) const override;
+	virtual void GetShaderHash(FString& InOutKey) const override;
+	virtual void GetHLSL(FString& OutHLSL, FString const& InDataInterfaceName) const override;
+	virtual UComputeDataProvider* CreateDataProvider(TObjectPtr<UObject> InBinding, uint64 InInputMask, uint64 InOutputMask) const override;
+	virtual bool CanSupportUnifiedDispatch() const override { return true; }
+	//~ End UComputeDataInterface Interface
+
+private:
+	/** File holding the hlsl implementation */
+	static TCHAR const* TemplateFilePath;
+
+	UPROPERTY(EditAnywhere, Category = "Simulation Trigger")
+	FName ResetSimulationTrigger = NAME_None;
+};
+
+/** Compute Framework Data Provider for reading groom guides. */
+UCLASS(BlueprintType, editinlinenew, Category = ComputeFramework)
+class UOptimusGroomSolverReadDataProvider : public UComputeDataProvider, public IOptimusDeformerInstanceAccessor
+{
+	GENERATED_BODY()
+
+public:
+
+	//~ Begin UComputeDataProvider Interface
+	virtual void SetDeformerInstance(UOptimusDeformerInstance* InInstance) override;
+	//~ End UComputeDataProvider Interface
+	
+	UPROPERTY(EditAnywhere, BlueprintReadWrite, Category = Binding)
+	TObjectPtr<UGroomSolverComponent> SolverComponent = nullptr;
+
+	//~ Begin UComputeDataProvider Interface
+	virtual FComputeDataProviderRenderProxy* GetRenderProxy() override;
+	//~ End UComputeDataProvider Interface
+
+	/** Deformer instance using this provider */
+	UPROPERTY()
+	TObjectPtr<UOptimusDeformerInstance> DeformerInstance = nullptr;
+	
+	/** Data interface from which this provider has been created */
+	TWeakObjectPtr<const UOptimusGroomSolverReadDataInterface> WeakDataInterface;
+
+	/** Reset simulation count */
+	int32 bResetSimulationCount = 0;
+};
+
+class FOptimusGroomSolverReadDataProviderProxy : public FComputeDataProviderRenderProxy
+{
+public:
+	FOptimusGroomSolverReadDataProviderProxy(UGroomSolverComponent* SolverComponent, const bool bResetSimulation);
+
+	//~ Begin FComputeDataProviderRenderProxy Interface
+	virtual bool IsValid(FValidationData const& InValidationData) const override;
+	virtual void AllocateResources(FRDGBuilder& GraphBuilder, FAllocationData const& InAllocationData) override;
+	virtual void GatherPermutations(FPermutationData& InOutPermutationData) const override;
+	virtual void GatherDispatchData(FDispatchData const& InDispatchData) override;
+	virtual void PostGraphDispatch(FRDGBuilder& GraphBuilder) const override;
+	//~ End FComputeDataProviderRenderProxy Interface
+
+private:
+	using FParameters = FOptimusGroomSolverReadParameters;
+
+	/** Create internal buffers */
+	void CreateInternalBuffers(FRDGBuilder& GraphBuilder);
+
+	/** Solver settings used in the solver */
+	FGroomSolverSettings SolverSettings;
+
+	/** Groom solver buffers built from the deformer */
+	FRDGBufferRef ObjectPointOffsetsBuffer = nullptr;
+	FRDGBufferRef ObjectCurveOffsetsBuffer = nullptr;
+	FRDGBufferRef ObjectNumPointsBuffer = nullptr;
+	FRDGBufferRef ObjectNumCurvesBuffer = nullptr;
+	FRDGBufferRef PointObjectIndicesBuffer = nullptr;
+	FRDGBufferRef CurveObjectIndicesBuffer = nullptr;
+	FRDGBufferRef DynamicPointIndicesBuffer = nullptr;
+	FRDGBufferRef DynamicCurveIndicesBuffer = nullptr;
+	FRDGBufferRef KinematicPointIndicesBuffer = nullptr;
+	FRDGBufferRef KinematicCurveIndicesBuffer = nullptr;
+	FRDGBufferRef ObjectDistanceLodsBuffer = nullptr;
+
+	/** Invocation points count */
+	TArray<int32> PointsCounts;
+
+	/** Invocation curves count */
+	TArray<int32> CurvesCounts;
+
+	/** Invocation aligned points */
+	TArray<int32> AlignedPoints;
+
+	/** Total number of points */
+	int32 NumPoints = 0;
+
+	/** Total number of curves */
+	int32 NumCurves = 0;
+
+	/** Groom solver resources built from the deformer */
+	FRDGBufferSRVRef ObjectPointOffsetsResource = nullptr;
+	FRDGBufferSRVRef ObjectCurveOffsetsResource = nullptr;
+	FRDGBufferSRVRef ObjectNumPointsResource = nullptr;
+	FRDGBufferSRVRef ObjectNumCurvesResource = nullptr;
+	FRDGBufferSRVRef PointObjectIndicesResource = nullptr;
+	FRDGBufferSRVRef CurveObjectIndicesResource = nullptr;
+	FRDGBufferSRVRef DynamicPointIndicesResource = nullptr;
+	FRDGBufferSRVRef DynamicCurveIndicesResource = nullptr;
+	FRDGBufferSRVRef KinematicPointIndicesResource = nullptr;
+	FRDGBufferSRVRef KinematicCurveIndicesResource = nullptr;
+	FRDGBufferSRVRef ObjectDistanceLodsResource = nullptr;
+
+	/** List of instances (invocations) used in that data interface */
+	TArray<const FHairGroupInstance*> GroupInstances;
+
+	/** Reset simulation trigger */
+	bool bResetSimulationTrigger = false;
+};
diff --git a/Engine/Plugins/Runtime/HairStrands/Source/HairStrandsDeformer/Private/DeformerGroomStrandsRead.cpp b/Engine/Plugins/Runtime/HairStrands/Source/HairStrandsDeformer/Private/DeformerGroomStrandsRead.cpp
new file mode 100644
--- /dev/null
+++ b/Engine/Plugins/Runtime/HairStrands/Source/HairStrandsDeformer/Private/DeformerGroomStrandsRead.cpp
@@ -0,0 +1,358 @@
+// Copyright Epic Games, Inc. All Rights Reserved.
+
+#include "DeformerGroomStrandsRead.h"
+
+#include "ComputeFramework/ShaderParamTypeDefinition.h"
+#include "ShaderCompilerCore.h"
+#include "ShaderParameterMetadataBuilder.h"
+#include "GroomComponent.h"
+#include "GroomInstance.h"
+#include "RenderGraphBuilder.h"
+#include "GlobalRenderResources.h"
+#include "DeformerGroomDomainsSource.h"
+#include "DeformerGroomInterfaceUtils.h"
+#include "HairStrandsInterpolation.h"
+#include "SystemTextures.h"
+#include "ComputeFramework/ComputeKernelPermutationVector.h"
+
+FString UOptimusGroomStrandsReadDataInterface::GetDisplayName() const
+{
+	return TEXT("Groom Strands");
+}
+
+TArray<FOptimusCDIPinDefinition> UOptimusGroomStrandsReadDataInterface::GetPinDefinitions() const
+{
+	static const FName StrandsPoints(UOptimusGroomAssetComponentSource::FStrandsExecutionDomains::Points);
+	static const FName StrandsCurves(UOptimusGroomAssetComponentSource::FStrandsExecutionDomains::Curves);
+
+	TArray<FOptimusCDIPinDefinition> Defs;
+	Defs.Add({ "NumControlPoints", "ReadNumPoints", false, {"NumStrandsPoints"} });
+	Defs.Add({ "NumCurves", "ReadNumCurves", false, {"NumStrandsCurves"} });
+
+	// Deformation buffers
+	Defs.Add({ "Position", "ReadPointRestPosition", StrandsPoints, "ReadNumPoints", false, "PointRestPosition" });
+	Defs.Add({ "CurveOffsetPoint", "ReadCurvePointOffset", StrandsCurves,          "ReadNumCurves", false, "CurvePointOffset"  });
+	Defs.Add({ "CurveNumPoint", "ReadCurveNumPoints", StrandsCurves,          "ReadNumCurves", false, "CurveNumPoints"  });
+	Defs.Add({ "PointCurveIndex", "ReadPointCurveIndex", StrandsPoints,          "ReadNumPoints", false, "PointCurveIndex"  });
+	Defs.Add({ "CurveRestTransform", "ReadCurveRestTransform", StrandsCurves,   "ReadNumCurves", false, "CurveRestTransform"  });
+	Defs.Add({ "CurveDeformedTransform", "ReadCurveDeformedTransform", StrandsCurves,          "ReadNumCurves", true, "CurveDeformedTransform"  });
+
+	// Interpolation buffers
+	Defs.Add({ "GuideIndex", "ReadPointGuideIndices", StrandsPoints,   "ReadNumPoints", false, "PointGuideIndices"  });
+	Defs.Add({ "PointGuideWeights", "ReadPointGuideWeights", StrandsPoints,   "ReadNumPoints", false, "PointGuideWeights"  });
+	
+	// Geometry buffers
+	Defs.Add({ "CoordU", "ReadPointCoordU", StrandsPoints,   "ReadNumPoints", false, "PointCurveCoordU"  });
+	Defs.Add({ "Length", "ReadPointLength", StrandsPoints,   "ReadNumPoints", false, "PointCurveLength"  });
+	
+	// Material buffers
+	Defs.Add({ "Radius", "ReadPointRadius", StrandsPoints, "ReadNumPoints", true, "PointMaterialRadius" });
+	Defs.Add({ "RootUV", "ReadPointRootUV", StrandsPoints, "ReadNumPoints", true, "PointMaterialRootUV" });
+	Defs.Add({ "Seed", "ReadPointSeed", StrandsPoints,   "ReadNumPoints", true, "PointMaterialSeed"  });
+	Defs.Add({ "ClumpId", "ReadPointClumpId", StrandsPoints,   "ReadNumPoints", true, "PointMaterialClumpId"  });
+	Defs.Add({ "Color", "ReadPointColor", StrandsPoints,   "ReadNumPoints", true, "PointMaterialColor"  });
+	Defs.Add({ "Roughness", "ReadPointRoughness", StrandsPoints,   "ReadNumPoints", true, "PointMaterialRoughness"  });
+	Defs.Add({ "AO", "ReadPointAO", StrandsPoints,   "ReadNumPoints", true, "PointMaterialAO"  });
+	
+	return Defs;
+}
+ 
+TSubclassOf<UActorComponent> UOptimusGroomStrandsReadDataInterface::GetRequiredComponentClass() const
+{
+	return UMeshComponent::StaticClass();
+}
+
+void UOptimusGroomStrandsReadDataInterface::GetSupportedInputs(TArray<FShaderFunctionDefinition>& OutFunctions) const
+{
+	OutFunctions.AddDefaulted_GetRef()
+		.SetName(TEXT("ReadNumPoints"))
+		.AddReturnType(EShaderFundamentalType::Uint);
+
+	OutFunctions.AddDefaulted_GetRef()
+		.SetName(TEXT("ReadNumCurves"))
+		.AddReturnType(EShaderFundamentalType::Uint);
+
+	OutFunctions.AddDefaulted_GetRef()
+		.SetName(TEXT("ReadPointRestPosition"))
+		.AddReturnType(EShaderFundamentalType::Float, 3)
+		.AddParam(EShaderFundamentalType::Uint);
+
+	OutFunctions.AddDefaulted_GetRef()
+		.SetName(TEXT("ReadPointRadius"))
+		.AddReturnType(EShaderFundamentalType::Float)
+		.AddParam(EShaderFundamentalType::Uint);
+
+	OutFunctions.AddDefaulted_GetRef()
+		.SetName(TEXT("ReadPointCoordU"))
+		.AddReturnType(EShaderFundamentalType::Float)
+		.AddParam(EShaderFundamentalType::Uint);
+
+	OutFunctions.AddDefaulted_GetRef()
+		.SetName(TEXT("ReadPointLength"))
+		.AddReturnType(EShaderFundamentalType::Float)
+		.AddParam(EShaderFundamentalType::Uint);
+
+	OutFunctions.AddDefaulted_GetRef()
+		.SetName(TEXT("ReadPointRootUV"))
+		.AddReturnType(EShaderFundamentalType::Float, 2)
+		.AddParam(EShaderFundamentalType::Uint);
+
+	OutFunctions.AddDefaulted_GetRef()
+		.SetName(TEXT("ReadPointSeed"))
+		.AddReturnType(EShaderFundamentalType::Float)
+		.AddParam(EShaderFundamentalType::Uint);
+
+	OutFunctions.AddDefaulted_GetRef()
+		.SetName(TEXT("ReadPointClumpId"))
+		.AddReturnType(EShaderFundamentalType::Uint)
+		.AddParam(EShaderFundamentalType::Uint);
+
+	OutFunctions.AddDefaulted_GetRef()
+		.SetName(TEXT("ReadPointColor"))
+		.AddReturnType(EShaderFundamentalType::Float, 3)
+		.AddParam(EShaderFundamentalType::Uint);
+
+	OutFunctions.AddDefaulted_GetRef()
+		.SetName(TEXT("ReadPointRoughness"))
+		.AddReturnType(EShaderFundamentalType::Float)
+		.AddParam(EShaderFundamentalType::Uint);
+
+	OutFunctions.AddDefaulted_GetRef()
+		.SetName(TEXT("ReadPointAO"))
+		.AddReturnType(EShaderFundamentalType::Float)
+		.AddParam(EShaderFundamentalType::Uint);
+
+	OutFunctions.AddDefaulted_GetRef()
+		.SetName(TEXT("ReadCurvePointOffset"))
+		.AddReturnType(EShaderFundamentalType::Uint)
+		.AddParam(EShaderFundamentalType::Uint);
+
+	OutFunctions.AddDefaulted_GetRef()
+		.SetName(TEXT("ReadCurveNumPoints"))
+		.AddReturnType(EShaderFundamentalType::Uint)
+		.AddParam(EShaderFundamentalType::Uint);
+
+	OutFunctions.AddDefaulted_GetRef()
+		.SetName(TEXT("ReadPointGuideIndices"))
+		.AddReturnType(EShaderFundamentalType::Int, 2)
+		.AddParam(EShaderFundamentalType::Uint);
+
+	OutFunctions.AddDefaulted_GetRef()
+		.SetName(TEXT("ReadPointGuideWeights"))
+		.AddReturnType(EShaderFundamentalType::Float, 2)
+		.AddParam(EShaderFundamentalType::Uint);
+
+	OutFunctions.AddDefaulted_GetRef()
+		.SetName(TEXT("ReadPointCurveIndex"))
+		.AddReturnType(EShaderFundamentalType::Uint)
+		.AddParam(EShaderFundamentalType::Uint);
+	
+	OutFunctions.AddDefaulted_GetRef()
+		.SetName(TEXT("ReadCurveRestTransform"))
+		.AddReturnType(EShaderFundamentalType::Float, 3, 4)
+		.AddParam(EShaderFundamentalType::Uint);
+	
+	OutFunctions.AddDefaulted_GetRef()
+		.SetName(TEXT("ReadCurveDeformedTransform"))
+		.AddReturnType(EShaderFundamentalType::Float, 3, 4)
+		.AddParam(EShaderFundamentalType::Uint);
+	
+}
+
+BEGIN_SHADER_PARAMETER_STRUCT(FOptimusGroomStrandsReadParameters, )
+	SHADER_PARAMETER_STRUCT_INCLUDE(FHairStrandsInstanceCommonParameters, Common)
+	SHADER_PARAMETER_STRUCT_INCLUDE(FHairStrandsInstanceResourceParameters, Resources)
+	SHADER_PARAMETER_STRUCT_INCLUDE(FHairStrandsInstanceInterpolationParameters, Interpolation)
+	SHADER_PARAMETER(FMatrix44f, RigidRestTransform)
+	SHADER_PARAMETER(FMatrix44f, RigidDeformedTransform)
+	SHADER_PARAMETER_RDG_BUFFER_SRV(Buffer<float4>, TriangleRestPositions)
+	SHADER_PARAMETER_RDG_BUFFER_SRV(Buffer<float4>, TriangleDeformedPositions)
+	SHADER_PARAMETER_RDG_BUFFER_SRV(Buffer<uint>, CurveBarycentricCoordinates)
+	SHADER_PARAMETER_RDG_BUFFER_SRV(Buffer<uint>, CurveTriangleIndices)
+END_SHADER_PARAMETER_STRUCT()
+
+void UOptimusGroomStrandsReadDataInterface::GetShaderParameters(TCHAR const* UID, FShaderParametersMetadataBuilder& InOutBuilder, FShaderParametersMetadataAllocations& InOutAllocations) const
+{
+	InOutBuilder.AddNestedStruct<FOptimusGroomStrandsReadParameters>(UID);
+}
+
+TCHAR const* UOptimusGroomStrandsReadDataInterface::TemplateFilePath = TEXT("/Plugin/Runtime/HairStrands/Private/Deformers/DeformerGroomStrandsRead.ush");
+
+TCHAR const* UOptimusGroomStrandsReadDataInterface::GetShaderVirtualPath() const
+{
+	return TemplateFilePath;
+}
+
+void UOptimusGroomStrandsReadDataInterface::GetShaderHash(FString& InOutKey) const
+{
+	GetShaderFileHash(TemplateFilePath, EShaderPlatform::SP_PCD3D_SM5).AppendString(InOutKey);
+}
+
+void UOptimusGroomStrandsReadDataInterface::GetPermutations(FComputeKernelPermutationVector& OutPermutationVector) const
+{
+	OutPermutationVector.AddPermutation(TEXT("ENABLE_SKINNED_TRANSFORM"), 2);
+}
+
+void UOptimusGroomStrandsReadDataInterface::GetHLSL(FString& OutHLSL, FString const& InDataInterfaceName) const
+{
+	TMap<FString, FStringFormatArg> TemplateArgs =
+	{
+		{ TEXT("DataInterfaceName"), InDataInterfaceName },
+	};
+
+	FString TemplateFile;
+	LoadShaderSourceFile(TemplateFilePath, EShaderPlatform::SP_PCD3D_SM5, &TemplateFile, nullptr);
+	OutHLSL += FString::Format(*TemplateFile, TemplateArgs);
+}
+
+UComputeDataProvider* UOptimusGroomStrandsReadDataInterface::CreateDataProvider(TObjectPtr<UObject> InBinding, uint64 InInputMask, uint64 InOutputMask) const
+{
+	UOptimusGroomStrandsReadDataProvider* Provider = NewObject<UOptimusGroomStrandsReadDataProvider>();
+	Provider->MeshComponent = Cast<UMeshComponent>(InBinding);
+	return Provider;
+}
+
+FComputeDataProviderRenderProxy* UOptimusGroomStrandsReadDataProvider::GetRenderProxy()
+{
+	return new FOptimusGroomStrandsReadProviderProxy(MeshComponent);
+}
+
+FOptimusGroomStrandsReadProviderProxy::FOptimusGroomStrandsReadProviderProxy(UMeshComponent* MeshComponent)
+{
+	UE::Groom::Private::GatherGroupInstances(MeshComponent, GroupInstances);
+}
+
+bool FOptimusGroomStrandsReadProviderProxy::IsValid(FValidationData const& InValidationData) const
+{
+	if (InValidationData.ParameterStructSize != sizeof(FParameters))
+	{
+		return false;
+	}
+	if (InValidationData.NumInvocations != GroupInstances.Num())
+	{
+		return false;
+	}
+	if(!UE::Groom::Private::HaveStrandsInstanceResources(GroupInstances) || !UE::Groom::Private::HaveStrandsSkinnedResources(GroupInstances))
+	{
+		return false;
+	}
+	
+	return true;
+}
+
+struct FOptimusGroomStrandsReadPermutationIds
+{
+	uint32 EnableSkinnedTransform = 0;
+
+	FOptimusGroomStrandsReadPermutationIds(FComputeKernelPermutationVector const& PermutationVector)
+	{
+		static FString Name(TEXT("ENABLE_SKINNED_TRANSFORM"));
+		static uint32 Hash = GetTypeHash(Name);
+		EnableSkinnedTransform = PermutationVector.GetPermutationBits(Name, Hash, 1);
+	}
+};
+
+void FOptimusGroomStrandsReadProviderProxy::GatherPermutations(FPermutationData& InOutPermutationData) const
+{
+	const FOptimusGroomStrandsReadPermutationIds PermutationIds(InOutPermutationData.PermutationVector);
+
+	int32 InvocationIndex = 0;
+	for (const FHairGroupInstance* GroupInstance : GroupInstances)
+	{
+		if (GroupInstance)
+		{
+			InOutPermutationData.PermutationIds[InvocationIndex] |= ((GroupInstance->BindingType == EHairBindingType::Skinning) ? PermutationIds.EnableSkinnedTransform : 0);
+		}
+		++InvocationIndex;
+	}
+}
+
+void FOptimusGroomStrandsReadProviderProxy::AllocateResources(FRDGBuilder& GraphBuilder, FAllocationData const& InAllocationData)
+{
+	if (!FallbackByteAddressSRV) { FallbackByteAddressSRV = GraphBuilder.CreateSRV(GSystemTextures.GetDefaultByteAddressBuffer(GraphBuilder, 16u)); }
+
+	GroupResources.Reset();
+	GroupInterpolations.Reset();
+	BindingResources.Reset();
+	
+	for (const FHairGroupInstance* GroupInstance : GroupInstances)
+	{
+		if (GroupInstance)
+		{
+			{
+				FHairStrandsInstanceResourceParameters& Resource = GroupResources.AddDefaulted_GetRef();
+				Resource.PositionBuffer	= RegisterAsSRV(GraphBuilder, GroupInstance->Strands.RestResource->PositionBuffer);
+				Resource.PositionOffsetBuffer = RegisterAsSRV(GraphBuilder, GroupInstance->Strands.RestResource->PositionOffsetBuffer);
+				Resource.CurveBuffer = RegisterAsSRV(GraphBuilder, GroupInstance->Strands.RestResource->CurveBuffer);
+				Resource.PointToCurveBuffer	= RegisterAsSRV(GraphBuilder, GroupInstance->Strands.RestResource->PointToCurveBuffer);
+				Resource.CurveAttributeBuffer = RegisterAsSRV(GraphBuilder, GroupInstance->Strands.RestResource->CurveAttributeBuffer);
+				Resource.PointAttributeBuffer = RegisterAsSRV(GraphBuilder, GroupInstance->Strands.RestResource->PointAttributeBuffer);
+			}
+
+			{
+				FHairStrandsInstanceInterpolationParameters& Resource = GroupInterpolations.AddDefaulted_GetRef();
+				Resource.CurveInterpolationBuffer = GroupInstance->Strands.InterpolationResource ? RegisterAsSRV(GraphBuilder, GroupInstance->Strands.InterpolationResource->CurveInterpolationBuffer) : nullptr;
+				Resource.PointInterpolationBuffer = GroupInstance->Strands.InterpolationResource ? RegisterAsSRV(GraphBuilder, GroupInstance->Strands.InterpolationResource->PointInterpolationBuffer) : nullptr;
+			}
+			{
+				FBindingResources& Resource = BindingResources.AddDefaulted_GetRef();
+				if(GroupInstance->BindingType == EHairBindingType::Skinning)
+				{
+					FHairStrandsLODRestRootResource* RestLODDatas = GroupInstance->Strands.RestRootResource->LODs[GroupInstance->HairGroupPublicData->MeshLODIndex];
+					FHairStrandsLODDeformedRootResource* DeformedLODDatas = GroupInstance->Strands.DeformedRootResource->LODs[GroupInstance->HairGroupPublicData->MeshLODIndex];
+				
+					Resource.CurveTriangleIndices = RegisterAsSRV(GraphBuilder, RestLODDatas->RootToUniqueTriangleIndexBuffer);
+					Resource.TriangleRestPositions = RegisterAsSRV(GraphBuilder, RestLODDatas->RestUniqueTrianglePositionBuffer);
+					Resource.TriangleDeformedPositions = RegisterAsSRV(GraphBuilder, DeformedLODDatas->GetDeformedUniqueTrianglePositionBuffer(FHairStrandsLODDeformedRootResource::Current));
+					Resource.CurveBarycentricCoordinates = RegisterAsSRV(GraphBuilder, RestLODDatas->RootBarycentricBuffer);
+				}
+				else
+				{
+					Resource.RigidDeformedTransform = FMatrix44f(GroupInstance->Debug.RigidCurrentLocalToWorld.ToMatrixWithScale());
+					Resource.RigidRestTransform = FMatrix44f(GroupInstance->Debug.RigidRestLocalToWorld.ToMatrixWithScale());
+				}
+			}
+		}
+	}
+}
+
+void FOptimusGroomStrandsReadProviderProxy::GatherDispatchData(FDispatchData const& InDispatchData)
+{
+	check(InDispatchData.NumInvocations == GroupInstances.Num());
+	const TStridedView<FParameters> ParameterArray = MakeStridedParameterView<FParameters>(InDispatchData);
+
+	for (int32 InvocationIndex = 0; InvocationIndex < ParameterArray.Num(); ++InvocationIndex)
+	{
+		if (const FHairGroupInstance* GroupInstance = GroupInstances[InvocationIndex])
+		{
+			FParameters& Parameters = ParameterArray[InvocationIndex];
+			
+			const FHairGroupPublicData::FVertexFactoryInput VFInput = ComputeHairStrandsVertexInputData(GroupInstance, EGroomViewMode::None);
+			Parameters.Common = VFInput.Strands.Common;
+			Parameters.Resources = GroupResources[InvocationIndex];
+			
+			if (GroupInterpolations[InvocationIndex].CurveInterpolationBuffer != nullptr)
+			{
+				Parameters.Interpolation = GroupInterpolations[InvocationIndex];
+			}
+			else
+			{
+				Parameters.Interpolation.CurveInterpolationBuffer = FallbackByteAddressSRV;
+				Parameters.Interpolation.PointInterpolationBuffer = FallbackByteAddressSRV;
+			}
+			FBindingResources& Resource = BindingResources[InvocationIndex];
+			if(GroupInstance->BindingType == EHairBindingType::Skinning)
+			{          
+				Parameters.CurveTriangleIndices = Resource.CurveTriangleIndices;
+				Parameters.CurveBarycentricCoordinates = Resource.CurveBarycentricCoordinates;
+				Parameters.TriangleRestPositions = Resource.TriangleRestPositions;
+				Parameters.TriangleDeformedPositions = Resource.TriangleDeformedPositions;
+			}
+			else
+			{
+				Parameters.RigidDeformedTransform = Resource.RigidDeformedTransform;
+				Parameters.RigidRestTransform = Resource.RigidRestTransform;
+			}
+		}
+	}
+}
\ No newline at end of file
diff --git a/Engine/Plugins/Runtime/HairStrands/Source/HairStrandsDeformer/Private/DeformerGroomStrandsRead.h b/Engine/Plugins/Runtime/HairStrands/Source/HairStrandsDeformer/Private/DeformerGroomStrandsRead.h
new file mode 100644
--- /dev/null
+++ b/Engine/Plugins/Runtime/HairStrands/Source/HairStrandsDeformer/Private/DeformerGroomStrandsRead.h
@@ -0,0 +1,96 @@
+// Copyright Epic Games, Inc. All Rights Reserved.
+
+#pragma once
+
+#include "RenderGraphDefinitions.h"
+#include "OptimusComputeDataInterface.h"
+#include "HairStrandsInterface.h"
+#include "ComputeFramework/ComputeDataProvider.h"
+#include "DeformerGroomStrandsRead.generated.h"
+
+class FOptimusGroomStrandsReadParameters;
+class UGroomComponent;
+struct FHairGroupInstance;
+class UMeshComponent;
+
+/** Compute Framework Data Interface for reading groom strands. */
+UCLASS(Category = ComputeFramework)
+class UOptimusGroomStrandsReadDataInterface : public UOptimusComputeDataInterface
+{
+	GENERATED_BODY()
+
+public:
+	//~ Begin UOptimusComputeDataInterface Interface
+	virtual FString GetDisplayName() const override;
+	virtual TArray<FOptimusCDIPinDefinition> GetPinDefinitions() const override;
+	virtual TSubclassOf<UActorComponent> GetRequiredComponentClass() const override;
+	//~ End UOptimusComputeDataInterface Interface
+	
+	//~ Begin UComputeDataInterface Interface
+	virtual TCHAR const* GetClassName() const override { return TEXT("GroomStrandsRead"); }
+	virtual void GetSupportedInputs(TArray<FShaderFunctionDefinition>& OutFunctions) const override;
+	virtual void GetShaderParameters(TCHAR const* UID, FShaderParametersMetadataBuilder& InOutBuilder, FShaderParametersMetadataAllocations& InOutAllocations) const override;
+	virtual TCHAR const* GetShaderVirtualPath() const override;
+	virtual void GetPermutations(FComputeKernelPermutationVector& OutPermutationVector) const override;
+	virtual void GetShaderHash(FString& InOutKey) const override;
+	virtual void GetHLSL(FString& OutHLSL, FString const& InDataInterfaceName) const override;
+	virtual UComputeDataProvider* CreateDataProvider(TObjectPtr<UObject> InBinding, uint64 InInputMask, uint64 InOutputMask) const override;
+	//~ End UComputeDataInterface Interface
+
+private:
+	/** File holding the hlsl implementation */
+	static TCHAR const* TemplateFilePath;
+};
+
+/** Compute Framework Data Provider for reading groom strands. */
+UCLASS(BlueprintType, editinlinenew, Category = ComputeFramework)
+class UOptimusGroomStrandsReadDataProvider : public UComputeDataProvider
+{
+	GENERATED_BODY()
+
+public:
+	UPROPERTY(EditAnywhere, BlueprintReadWrite, Category = Binding)
+	TObjectPtr<UMeshComponent> MeshComponent = nullptr;
+
+	//~ Begin UComputeDataProvider Interface
+	virtual FComputeDataProviderRenderProxy* GetRenderProxy() override;
+	//~ End UComputeDataProvider Interface
+};
+
+class FOptimusGroomStrandsReadProviderProxy : public FComputeDataProviderRenderProxy
+{
+public:
+	FOptimusGroomStrandsReadProviderProxy(UMeshComponent* InGroomComponent);
+
+	//~ Begin FComputeDataProviderRenderProxy Interface
+	virtual bool IsValid(FValidationData const& InValidationData) const override;
+	virtual void AllocateResources(FRDGBuilder& GraphBuilder, FAllocationData const& InAllocationData) override;
+	virtual void GatherPermutations(FPermutationData& InOutPermutationData) const override;
+	virtual void GatherDispatchData(FDispatchData const& InDispatchData) override;
+	//~ End FComputeDataProviderRenderProxy Interface
+
+private:
+	using FParameters = FOptimusGroomStrandsReadParameters;
+
+	/** List of instances (invocations) used in that data interface */
+	TArray<const FHairGroupInstance*> GroupInstances;
+
+	/** Binding Resources */
+	struct FBindingResources
+	{
+		FMatrix44f RigidRestTransform = FMatrix44f::Identity;
+		FMatrix44f RigidDeformedTransform = FMatrix44f::Identity;
+		FRDGBufferSRVRef TriangleRestPositions = nullptr;
+		FRDGBufferSRVRef TriangleDeformedPositions = nullptr;
+		FRDGBufferSRVRef CurveBarycentricCoordinates = nullptr;
+		FRDGBufferSRVRef CurveTriangleIndices = nullptr;
+	};
+	TArray<FBindingResources> BindingResources;
+
+	/** Resources used to dispatch CS on GPU */
+	TArray<FHairStrandsInstanceResourceParameters> GroupResources;
+	TArray<FHairStrandsInstanceInterpolationParameters> GroupInterpolations;
+
+	/** Fallback resources */
+	FRDGBufferSRVRef FallbackByteAddressSRV = nullptr;
+};
diff --git a/Engine/Plugins/Runtime/HairStrands/Source/HairStrandsDeformer/Private/DeformerGroomStrandsWrite.cpp b/Engine/Plugins/Runtime/HairStrands/Source/HairStrandsDeformer/Private/DeformerGroomStrandsWrite.cpp
new file mode 100644
--- /dev/null
+++ b/Engine/Plugins/Runtime/HairStrands/Source/HairStrandsDeformer/Private/DeformerGroomStrandsWrite.cpp
@@ -0,0 +1,272 @@
+// Copyright Epic Games, Inc. All Rights Reserved.
+
+#include "DeformerGroomStrandsWrite.h"
+
+#include "ComputeFramework/ShaderParamTypeDefinition.h"
+#include "RenderGraphBuilder.h"
+#include "ShaderParameterMetadataBuilder.h"
+#include "GroomComponent.h"
+#include "GroomInstance.h"
+#include "DeformerGroomDomainsSource.h"
+#include "DeformerGroomInterfaceUtils.h"
+#include "RenderGraphUtils.h"
+#include "SystemTextures.h"
+#include "HairStrandsInterpolation.h"
+
+FString UOptimusGroomStrandsWriteDataInterface::GetDisplayName() const
+{
+	return TEXT("Write Groom Strands");
+}
+
+FName UOptimusGroomStrandsWriteDataInterface::GetCategory() const
+{
+	return CategoryName::OutputDataInterfaces;
+}
+
+TArray<FOptimusCDIPinDefinition> UOptimusGroomStrandsWriteDataInterface::GetPinDefinitions() const
+{
+	static const FName StrandsPoints(UOptimusGroomAssetComponentSource::FStrandsExecutionDomains::Points);
+	static const FName StrandsCurves(UOptimusGroomAssetComponentSource::FStrandsExecutionDomains::Curves);
+
+	TArray<FOptimusCDIPinDefinition> Defs;
+	
+	// Deformation buffers
+	Defs.Add({ "Position", "WritePointDeformedPosition", StrandsPoints, "ReadNumPoints", true, "PointDeformedPosition" });
+	Defs.Add({ "PositionAndRadius", "WritePointPositionRadius", StrandsPoints, "ReadNumPoints", true, "PointPositionRadius" });
+
+	// Material buffers
+	Defs.Add({ "Radius", "WritePointRadius", StrandsPoints, "ReadNumPoints", true, "PointMaterialRadius" });
+	Defs.Add({ "RootUV", "WriteCurveRootUV", StrandsCurves, "ReadNumCurves", true, "CurveMaterialRootUV" });
+	Defs.Add({ "Seed", "WriteCurveSeed", StrandsCurves, "ReadNumCurves", true, "CurveMaterialSeed" });
+	Defs.Add({ "ClumpId", "WriteCurveClumpId", StrandsCurves, "ReadNumCurves", true, "CurveMaterialClumpId" });
+	Defs.Add({ "Color", "WritePointColor", StrandsPoints, "ReadNumPoints", true, "PointMaterialColor" });
+	Defs.Add({ "Roughness", "WritePointRoughness", StrandsPoints, "ReadNumPoints", true, "PointMaterialRoughness" });
+	Defs.Add({ "AO", "WritePointAO", StrandsPoints, "ReadNumPoints", true, "PointMaterialAO" });
+
+	return Defs;
+}
+
+TSubclassOf<UActorComponent> UOptimusGroomStrandsWriteDataInterface::GetRequiredComponentClass() const
+{
+	return UMeshComponent::StaticClass();
+}
+
+void UOptimusGroomStrandsWriteDataInterface::GetSupportedInputs(TArray<FShaderFunctionDefinition>& OutFunctions) const
+{
+	OutFunctions.AddDefaulted_GetRef()
+		.SetName(TEXT("ReadNumPoints"))
+		.AddReturnType(EShaderFundamentalType::Uint);
+
+	OutFunctions.AddDefaulted_GetRef()
+		.SetName(TEXT("ReadNumCurves"))
+		.AddReturnType(EShaderFundamentalType::Uint);
+}
+
+void UOptimusGroomStrandsWriteDataInterface::GetSupportedOutputs(TArray<FShaderFunctionDefinition>& OutFunctions) const
+{
+	OutFunctions.AddDefaulted_GetRef()
+		.SetName(TEXT("WritePointDeformedPosition"))
+		.AddParam(EShaderFundamentalType::Uint)
+		.AddParam(EShaderFundamentalType::Float, 3);
+
+	OutFunctions.AddDefaulted_GetRef()
+		.SetName(TEXT("WritePointRadius"))
+		.AddParam(EShaderFundamentalType::Uint)
+		.AddParam(EShaderFundamentalType::Float);
+
+	OutFunctions.AddDefaulted_GetRef()
+		.SetName(TEXT("WritePointPositionRadius"))
+		.AddParam(EShaderFundamentalType::Uint)
+		.AddParam(EShaderFundamentalType::Float, 4);
+
+	OutFunctions.AddDefaulted_GetRef()
+		.SetName(TEXT("WriteCurveRootUV"))
+		.AddParam(EShaderFundamentalType::Uint)
+		.AddParam(EShaderFundamentalType::Float, 2);
+
+	OutFunctions.AddDefaulted_GetRef()
+		.SetName(TEXT("WriteCurveSeed"))
+		.AddParam(EShaderFundamentalType::Uint)
+		.AddParam(EShaderFundamentalType::Float);
+
+	OutFunctions.AddDefaulted_GetRef()
+		.SetName(TEXT("WriteCurveClumpId"))
+		.AddParam(EShaderFundamentalType::Uint)
+		.AddParam(EShaderFundamentalType::Float);
+
+	OutFunctions.AddDefaulted_GetRef()
+		.SetName(TEXT("WritePointColor"))
+		.AddParam(EShaderFundamentalType::Uint)
+		.AddParam(EShaderFundamentalType::Float, 3);
+
+	OutFunctions.AddDefaulted_GetRef()
+		.SetName(TEXT("WritePointRoughness"))
+		.AddParam(EShaderFundamentalType::Uint)
+		.AddParam(EShaderFundamentalType::Float);
+
+	OutFunctions.AddDefaulted_GetRef()
+		.SetName(TEXT("WritePointAO"))
+		.AddParam(EShaderFundamentalType::Uint)
+		.AddParam(EShaderFundamentalType::Float);
+}
+
+BEGIN_SHADER_PARAMETER_STRUCT(FOptimusGroomStrandsWriteParameters, )
+	SHADER_PARAMETER_STRUCT_INCLUDE(FHairStrandsInstanceCommonParameters, Common)
+	SHADER_PARAMETER(uint32, OutputStreamStart)
+	SHADER_PARAMETER_RDG_BUFFER_SRV(StructuredBuffer, DeformedPositionOffset)
+	SHADER_PARAMETER_RDG_BUFFER_SRV(ByteAddressBuffer, PointRestPositions)
+	SHADER_PARAMETER_RDG_BUFFER_UAV(RWByteAddressBuffer, PointDeformedPositions)
+	SHADER_PARAMETER_RDG_BUFFER_UAV(RWByteAddressBuffer, CurveAttributeBuffer)
+	SHADER_PARAMETER_RDG_BUFFER_UAV(RWByteAddressBuffer, PointAttributeBuffer)
+END_SHADER_PARAMETER_STRUCT()
+
+void UOptimusGroomStrandsWriteDataInterface::GetShaderParameters(TCHAR const* UID, FShaderParametersMetadataBuilder& InOutBuilder, FShaderParametersMetadataAllocations& InOutAllocations) const
+{
+	InOutBuilder.AddNestedStruct<FOptimusGroomStrandsWriteParameters>(UID);
+}
+
+TCHAR const* UOptimusGroomStrandsWriteDataInterface::TemplateFilePath = TEXT("/Plugin/Runtime/HairStrands/Private/Deformers/DeformerGroomStrandsWrite.ush");
+
+TCHAR const* UOptimusGroomStrandsWriteDataInterface::GetShaderVirtualPath() const
+{
+	return TemplateFilePath;
+}
+
+void UOptimusGroomStrandsWriteDataInterface::GetShaderHash(FString& InOutKey) const
+{
+	GetShaderFileHash(TemplateFilePath, EShaderPlatform::SP_PCD3D_SM5).AppendString(InOutKey);
+}
+
+void UOptimusGroomStrandsWriteDataInterface::GetHLSL(FString& OutHLSL, FString const& InDataInterfaceName) const
+{
+	TMap<FString, FStringFormatArg> TemplateArgs =
+	{
+		{ TEXT("DataInterfaceName"), InDataInterfaceName },
+	};
+
+	FString TemplateFile;
+	LoadShaderSourceFile(TemplateFilePath, EShaderPlatform::SP_PCD3D_SM5, &TemplateFile, nullptr);
+	OutHLSL += FString::Format(*TemplateFile, TemplateArgs);
+}
+
+UComputeDataProvider* UOptimusGroomStrandsWriteDataInterface::CreateDataProvider(TObjectPtr<UObject> InBinding, uint64 InInputMask, uint64 InOutputMask) const
+{
+	UOptimusGroomStrandsWriteDataProvider* Provider = NewObject<UOptimusGroomStrandsWriteDataProvider>();
+	Provider->MeshComponent = Cast<UMeshComponent>(InBinding);
+	Provider->OutputMask = InOutputMask;
+	return Provider;
+}
+
+FComputeDataProviderRenderProxy* UOptimusGroomStrandsWriteDataProvider::GetRenderProxy()
+{
+	return new FOptimusGroomStrandsWriteProviderProxy(MeshComponent, OutputMask);
+}
+
+FOptimusGroomStrandsWriteProviderProxy::FOptimusGroomStrandsWriteProviderProxy(UMeshComponent* MeshComponent, uint64 InOutputMask)
+{
+	OutputMask = InOutputMask;
+
+	TArray<const UGroomComponent*> GroomComponents;
+	UE::Groom::Private::GatherGroupInstances(MeshComponent, GroupInstances);
+}
+
+bool FOptimusGroomStrandsWriteProviderProxy::IsValid(FValidationData const& InValidationData) const
+{
+	if (InValidationData.ParameterStructSize != sizeof(FParameters))
+	{
+		return false;
+	}
+	if (InValidationData.NumInvocations != GroupInstances.Num())
+	{
+		return false;
+	}
+	if(!UE::Groom::Private::HaveStrandsInstanceResources(GroupInstances))
+	{
+		return false;
+	}
+
+	return true;
+}
+
+void FOptimusGroomStrandsWriteProviderProxy::AllocateResources(FRDGBuilder& GraphBuilder, FAllocationData const& InAllocationData)
+{
+	if (!FallbackPositionBufferSRV) { FallbackPositionBufferSRV = GraphBuilder.CreateSRV(GSystemTextures.GetDefaultByteAddressBuffer(GraphBuilder, 16u));}
+	if (!FallbackPositionBufferUAV) { FallbackPositionBufferUAV = GraphBuilder.CreateUAV(GraphBuilder.CreateBuffer(FRDGBufferDesc::CreateByteAddressDesc(16u), TEXT("Groom.Deformer.FallbackDeformedPositionBuffer")), ERDGUnorderedAccessViewFlags::SkipBarrier);}
+	if (!FallbackAttributeBufferUAV) { FallbackAttributeBufferUAV = GraphBuilder.CreateUAV(GraphBuilder.CreateBuffer(FRDGBufferDesc::CreateByteAddressDesc(16u), TEXT("Groom.Deformer.FallbackDeformedAttributeBuffer")), ERDGUnorderedAccessViewFlags::SkipBarrier);}
+
+	for (const FHairGroupInstance* GroupInstance : GroupInstances)
+	{
+		if (GroupInstance)
+		{
+			FDispatchResources& Resource = GroupResources.AddDefaulted_GetRef();
+			
+			// Positions / Radius
+			if ((OutputMask & 0x7) && GroupInstance->Strands.RestResource && GroupInstance->Strands.DeformedResource)
+			{
+				Resource.DeformedPositionOffset 	= Register(GraphBuilder, GroupInstance->Strands.RestResource->PositionOffsetBuffer, ERDGImportedBufferFlags::CreateSRV).SRV;
+				Resource.PointRestPositions 		= Register(GraphBuilder, GroupInstance->Strands.RestResource->PositionBuffer, ERDGImportedBufferFlags::CreateSRV).SRV;
+				Resource.PointDeformedPositions 	= Register(GraphBuilder, GroupInstance->Strands.DeformedResource->GetDeformerBuffer(GraphBuilder), ERDGImportedBufferFlags::CreateUAV).UAV;
+			}
+			
+			// Always copy the attributes from the rest asset, so that if deformer write different attribute at different tick, it all remains consistent with the source data
+			if(GroupInstance->Strands.RestResource && GroupInstance->Strands.DeformedResource)
+			{
+				// Curve Attributes
+				if (OutputMask & 0x38)
+				{
+					FRDGExternalBuffer CurveAttributeBufferExt = GroupInstance->Strands.DeformedResource->GetDeformerCurveAttributeBuffer(GraphBuilder);
+					if (CurveAttributeBufferExt.Buffer)
+					{
+						FRDGImportedBuffer DstCurveAttributeBuffer = Register(GraphBuilder, CurveAttributeBufferExt, ERDGImportedBufferFlags::CreateUAV);
+						FRDGImportedBuffer SrcCurveAttributeBuffer = Register(GraphBuilder, GroupInstance->Strands.RestResource->CurveAttributeBuffer, ERDGImportedBufferFlags::None);
+						if (SrcCurveAttributeBuffer.Buffer)
+						{
+							AddCopyBufferPass(GraphBuilder, DstCurveAttributeBuffer.Buffer, SrcCurveAttributeBuffer.Buffer);
+						}
+						Resource.CurveAttributeBuffer = DstCurveAttributeBuffer.UAV;
+					}
+				}
+
+				// Point Attributes
+				if (OutputMask & 0x1c0)
+				{
+					FRDGExternalBuffer PointAttributeBufferExt = GroupInstance->Strands.DeformedResource->GetDeformerPointAttributeBuffer(GraphBuilder);
+					if (PointAttributeBufferExt.Buffer)
+					{
+						FRDGImportedBuffer DstPointAttributeBuffer = Register(GraphBuilder, PointAttributeBufferExt, ERDGImportedBufferFlags::CreateUAV);
+						FRDGImportedBuffer SrcPointAttributeBuffer = Register(GraphBuilder, GroupInstance->Strands.RestResource->PointAttributeBuffer, ERDGImportedBufferFlags::None);
+						if(SrcPointAttributeBuffer.Buffer)
+						{
+							AddCopyBufferPass(GraphBuilder, DstPointAttributeBuffer.Buffer, SrcPointAttributeBuffer.Buffer);
+						}
+						Resource.PointAttributeBuffer = DstPointAttributeBuffer.UAV;
+					}
+				}
+			}
+		}	
+	}
+}
+
+void FOptimusGroomStrandsWriteProviderProxy::GatherDispatchData(FDispatchData const& InDispatchData)
+{
+	const TStridedView<FParameters> ParameterArray = MakeStridedParameterView<FParameters>(InDispatchData);
+	for (int32 InvocationIndex = 0; InvocationIndex < ParameterArray.Num(); ++InvocationIndex)
+	{
+		if (const FHairGroupInstance* GroupInstance = GroupInstances[InvocationIndex])
+		{
+			FDispatchResources& Resource = GroupResources[InvocationIndex];
+
+			const bool bValid = Resource.PointDeformedPositions != nullptr && Resource.PointRestPositions != nullptr && Resource.DeformedPositionOffset != nullptr;
+			const FHairGroupPublicData::FVertexFactoryInput VFInput = ComputeHairStrandsVertexInputData(GroupInstance, EGroomViewMode::None);		
+
+			FParameters& Parameters = ParameterArray[InvocationIndex];
+			Parameters.Common 					= VFInput.Strands.Common;
+			Parameters.OutputStreamStart 		= 0;
+			Parameters.DeformedPositionOffset 	= bValid ? Resource.DeformedPositionOffset : FallbackPositionBufferSRV;
+			Parameters.PointRestPositions 		= bValid ? Resource.PointRestPositions : FallbackPositionBufferSRV;
+			Parameters.PointDeformedPositions 	= bValid ? Resource.PointDeformedPositions : FallbackPositionBufferUAV;
+			Parameters.PointAttributeBuffer 	= bValid ? Resource.PointAttributeBuffer : FallbackAttributeBufferUAV;
+			Parameters.CurveAttributeBuffer 	= bValid ? Resource.CurveAttributeBuffer : FallbackAttributeBufferUAV;
+		}
+	}
+}
\ No newline at end of file
diff --git a/Engine/Plugins/Runtime/HairStrands/Source/HairStrandsDeformer/Private/DeformerGroomStrandsWrite.h b/Engine/Plugins/Runtime/HairStrands/Source/HairStrandsDeformer/Private/DeformerGroomStrandsWrite.h
new file mode 100644
--- /dev/null
+++ b/Engine/Plugins/Runtime/HairStrands/Source/HairStrandsDeformer/Private/DeformerGroomStrandsWrite.h
@@ -0,0 +1,101 @@
+// Copyright Epic Games, Inc. All Rights Reserved.
+
+#pragma once
+
+#include "OptimusComputeDataInterface.h"
+#include "ComputeFramework/ComputeDataProvider.h"
+#include "DeformerGroomStrandsWrite.generated.h"
+
+class FOptimusGroomStrandsWriteParameters;
+struct FHairGroupInstance;
+class FRDGBuffer;
+class FRDGBufferSRV;
+class FRDGBufferUAV;
+class UGroomComponent;
+class IGroomComponentContainer;
+class UMeshComponent;
+
+/** Compute Framework Data Interface for writing groom strands. */
+UCLASS(Category = ComputeFramework)
+class UOptimusGroomStrandsWriteDataInterface : public UOptimusComputeDataInterface
+{
+	GENERATED_BODY()
+
+public:
+	//~ Begin UOptimusComputeDataInterface Interface
+	virtual FString GetDisplayName() const override;
+	virtual FName GetCategory() const override;
+	virtual TArray<FOptimusCDIPinDefinition> GetPinDefinitions() const override;
+	virtual TSubclassOf<UActorComponent> GetRequiredComponentClass() const override;
+	//~ End UOptimusComputeDataInterface Interface
+	
+	//~ Begin UComputeDataInterface Interface
+	virtual TCHAR const* GetClassName() const override { return TEXT("GroomStrandsWrite"); }
+	virtual void GetSupportedInputs(TArray<FShaderFunctionDefinition>& OutFunctions) const override;
+	virtual void GetSupportedOutputs(TArray<FShaderFunctionDefinition>& OutFunctions) const override;
+	virtual void GetShaderParameters(TCHAR const* UID, FShaderParametersMetadataBuilder& InOutBuilder, FShaderParametersMetadataAllocations& InOutAllocations) const override;
+	virtual TCHAR const* GetShaderVirtualPath() const override;
+	virtual void GetShaderHash(FString& InOutKey) const override;
+	virtual void GetHLSL(FString& OutHLSL, FString const& InDataInterfaceName) const override;
+	virtual UComputeDataProvider* CreateDataProvider(TObjectPtr<UObject> InBinding, uint64 InInputMask, uint64 InOutputMask) const override;
+	//~ End UComputeDataInterface Interface
+
+private:
+	/** File holding the hlsl implementation */
+	static TCHAR const* TemplateFilePath;
+};
+
+/** Compute Framework Data Provider for writing skinned mesh. */
+UCLASS(BlueprintType, editinlinenew, Category = ComputeFramework)
+class UOptimusGroomStrandsWriteDataProvider : public UComputeDataProvider
+{
+	GENERATED_BODY()
+
+public:
+	UPROPERTY(EditAnywhere, BlueprintReadWrite, Category = Binding)
+	TObjectPtr<UMeshComponent> MeshComponent = nullptr;
+
+	/** Output mask used to know the buffer format */
+	uint64 OutputMask = 0;
+
+	//~ Begin UComputeDataProvider Interface
+	virtual FComputeDataProviderRenderProxy* GetRenderProxy() override;
+	//~ End UComputeDataProvider Interface
+};
+
+class FOptimusGroomStrandsWriteProviderProxy : public FComputeDataProviderRenderProxy
+{
+public:
+	FOptimusGroomStrandsWriteProviderProxy(UMeshComponent* MeshComponent, uint64 InOutputMask);
+
+	//~ Begin FComputeDataProviderRenderProxy Interface
+	virtual bool IsValid(FValidationData const& InValidationData) const override;
+	virtual void AllocateResources(FRDGBuilder& GraphBuilder, FAllocationData const& InAllocationData) override;
+	virtual void GatherDispatchData(FDispatchData const& InDispatchData) override;
+	//~ End FComputeDataProviderRenderProxy Interface
+
+private:
+	using FParameters = FOptimusGroomStrandsWriteParameters;
+
+	/** List of instances (invocations) used in that data interface */
+	TArray<const FHairGroupInstance*> GroupInstances;
+
+	/** Output mask for gpu buffer format*/
+	uint64 OutputMask = 0;
+
+	/** Resources used to dispatch CS on GPU */
+	struct FDispatchResources
+	{
+		FRDGBufferSRVRef DeformedPositionOffset = nullptr;
+		FRDGBufferSRVRef PointRestPositions = nullptr;
+		FRDGBufferUAVRef PointDeformedPositions = nullptr;
+		FRDGBufferUAVRef PointAttributeBuffer = nullptr;
+		FRDGBufferUAVRef CurveAttributeBuffer = nullptr;
+	};
+	TArray<FDispatchResources> GroupResources;
+
+	/** Fallback resources */
+	FRDGBufferSRVRef FallbackPositionBufferSRV= nullptr;
+	FRDGBufferUAVRef FallbackPositionBufferUAV = nullptr;
+	FRDGBufferUAVRef FallbackAttributeBufferUAV = nullptr;
+};
\ No newline at end of file
diff --git a/Engine/Plugins/Runtime/HairStrands/Source/HairStrandsSolver/Private/GroomSolverComponent.cpp b/Engine/Plugins/Runtime/HairStrands/Source/HairStrandsSolver/Private/GroomSolverComponent.cpp
--- a/Engine/Plugins/Runtime/HairStrands/Source/HairStrandsSolver/Private/GroomSolverComponent.cpp
+++ b/Engine/Plugins/Runtime/HairStrands/Source/HairStrandsSolver/Private/GroomSolverComponent.cpp
@@ -78,34 +78,89 @@ void UGroomSolverComponent::OnUnregister()
 	GroomSolverProxy.DeformerInstance = nullptr;
 }
 
-void UGroomSolverComponent::SelectSimulatedCurves(TArray<int32>& NumSimulatedCurves) const
+void UGroomSolverComponent::SelectDynamicCurves() 
 {
 	TArray<FVector> ViewLocations;
 	UE::Groom::Private::GatherViewLocations(GetWorld(), ViewLocations);
 
+	SolverSettings.CurveDynamicIndices.Reset();
+	SolverSettings.CurveKinematicIndices.Reset();
+	SolverSettings.PointDynamicIndices.Reset();
+	SolverSettings.PointKinematicIndices.Reset();
+	SolverSettings.ObjectDistanceLods.Reset();
+
+	static constexpr int32 GroupSize = 64;
+
 	if (ViewLocations.Num() > 0)
 	{
+		uint32 CurveOffset = 0, PointOffset = 0;
 		for(const TObjectPtr<UGroomComponent>& GroomComponent : GetGroomComponents())
 		{
-			const FBoxSphereBounds GroomBounds = GroomComponent->CalcBounds(GroomComponent->GetComponentTransform());
-						
-			FVector::FReal MinDistance = FMath::Square(WORLD_MAX);
-			for (const FVector& ViewLocation : ViewLocations)
+			if(GroomComponent && GroomComponent->GroomAsset)
 			{
-				MinDistance = FMath::Min(MinDistance, (GroomBounds.Origin - ViewLocation).SizeSquared());
-			}
-			MinDistance = FMath::Sqrt(MinDistance);
+				const FBoxSphereBounds GroomBounds = GroomComponent->CalcBounds(GroomComponent->GetComponentTransform());
+						
+				FVector::FReal MinDistance = FMath::Square(WORLD_MAX);
+				for (const FVector& ViewLocation : ViewLocations)
+				{
+					MinDistance = FMath::Min(MinDistance, (GroomBounds.Origin - ViewLocation).SizeSquared());
+				}
+				MinDistance = FMath::Sqrt(MinDistance);
 
-			const FVector::FReal DistanceRatio = FMath::Clamp(MinDistance / GetSolverSettings().MaxLODDistance, 0.0f, 1.0f);
+				const FVector::FReal DistanceRatio = FMath::Clamp(MinDistance / GetSolverSettings().MaxLODDistance, 0.0f, 1.0f);
 
-			const uint32 NumGroups = GroomComponent->GetGroupCount();
-			for (uint32 GroupIndex = 0; GroupIndex < NumGroups; ++GroupIndex)
-			{
-				if (const FHairGroupInstance* Instance = GroomComponent->GetGroupInstance(GroupIndex))
+				const uint32 NumGroups = GroomComponent->GetGroupCount();
+				for (uint32 GroupIndex = 0; GroupIndex < NumGroups; ++GroupIndex)
 				{
-					const int32 NumCurves = Instance->Guides.IsValid() ? Instance->Guides.GetData().GetNumCurves() : 0;
-					NumSimulatedCurves.Add((1.0-DistanceRatio) * NumCurves);
+					const FHairGroupInstance* GroupInstance = GroomComponent->GetGroupInstance(GroupIndex);
+					if(GroomComponent->IsDeformationEnable(GroupIndex) && GroupInstance)
+					{
+						const int32 NumTotalCurves = GroupInstance->Guides.GetData().GetNumCurves();
+						const int32 NumTotalPoints = GroupInstance->Guides.GetData().GetNumPoints();
 
+						const int32 NumCurvePoints = NumTotalPoints / NumTotalCurves;
+						const int32 NumDynamicCurves = (1.0-DistanceRatio) * NumTotalCurves;
+						
+						const int32 NumDynamicPoints = NumDynamicCurves * NumCurvePoints;
+						const int32 NumKinematicPoints = NumTotalPoints-NumDynamicPoints;
+						
+						const int32 NumDynamicPadding = FMath::CeilToInt32(NumDynamicPoints / static_cast<float>(GroupSize)) * GroupSize - NumDynamicPoints;
+						const int32 NumKinematicPadding = FMath::CeilToInt32(NumKinematicPoints / static_cast<float>(GroupSize)) * GroupSize - NumKinematicPoints;
+						
+						for (int32 CurveIndex = 0; CurveIndex < NumDynamicCurves; ++CurveIndex)
+						{
+							SolverSettings.CurveDynamicIndices.Add(CurveOffset + CurveIndex);
+						}
+						for (int32 CurveIndex = NumDynamicCurves; CurveIndex < NumTotalCurves; ++CurveIndex)
+						{
+							SolverSettings.CurveKinematicIndices.Add(CurveOffset + CurveIndex);
+						}
+						for (int32 PointIndex = 0; PointIndex < NumDynamicPoints; ++PointIndex)
+						{
+							SolverSettings.PointDynamicIndices.Add(PointOffset + PointIndex);
+						}
+						for (int32 PointIndex = 0; PointIndex < NumDynamicPadding; ++PointIndex)
+						{
+							SolverSettings.PointDynamicIndices.Add(INDEX_NONE);
+						}
+						for (int32 PointIndex = NumDynamicPoints; PointIndex < NumTotalPoints; ++PointIndex)
+						{
+							SolverSettings.PointKinematicIndices.Add(PointOffset + PointIndex);
+						}
+						for (int32 PointIndex = 0; PointIndex < NumKinematicPadding; ++PointIndex)
+						{
+							SolverSettings.PointKinematicIndices.Add(INDEX_NONE);
+						}
+
+						const int32 NumObjectLODs = FMath::CeilLogTwo(NumTotalCurves);
+						const uint32 ObjectLod = NumObjectLODs - 1 - FMath::FloorLog2(NumDynamicCurves);
+
+						SolverSettings.ObjectDistanceLods.Add(ObjectLod);
+						
+						CurveOffset += NumTotalCurves;
+						PointOffset += FMath::CeilToInt32(NumTotalPoints / static_cast<float>(GroupSize)) * GroupSize;
+					}
+				
 					// Keep it for now for debugging since it is WIP
 					//UE_LOG(LogTemp, Log, TEXT("Num Curves = %d Distance Ratio = %f Simulated Curves = %f | Num LODs = %d"), NumCurves, DistanceRatio, (1.0-DistanceRatio) * NumCurves, FMath::CeilLogTwo(NumCurves));
 				}
@@ -122,6 +177,8 @@ void UGroomSolverComponent::TickComponent(float DeltaTime, enum ELevelTick TickT
 	{
 		MarkRenderDynamicDataDirty();
 	}
+
+	SelectDynamicCurves();
 }
 
 void UGroomSolverComponent::SendRenderDynamicData_Concurrent()
diff --git a/Engine/Plugins/Runtime/HairStrands/Source/HairStrandsSolver/Public/GroomSolverComponent.h b/Engine/Plugins/Runtime/HairStrands/Source/HairStrandsSolver/Public/GroomSolverComponent.h
--- a/Engine/Plugins/Runtime/HairStrands/Source/HairStrandsSolver/Public/GroomSolverComponent.h
+++ b/Engine/Plugins/Runtime/HairStrands/Source/HairStrandsSolver/Public/GroomSolverComponent.h
@@ -32,6 +32,21 @@ struct FGroomSolverSettings
 	/** Maximum LOD distance (if distance between the component and the views is higher that this threshold, no simulation)*/
 	UPROPERTY(EditAnywhere, Category = "LODs")
 	float MaxLODDistance = 1000;
+
+	/** List of dynamic curves */
+	TArray<int32> CurveDynamicIndices;
+
+	/** List of kinematic curves */
+	TArray<int32> CurveKinematicIndices;
+
+	/** List of dynamic points */
+	TArray<int32> PointDynamicIndices;
+
+	/** List of kinematic points */
+	TArray<int32> PointKinematicIndices;
+	
+	/** List of object curve lods */
+	TArray<int32> ObjectDistanceLods;
 };
 
 /** Dataflow groom solver proxy used in dataflow simulation */ 
@@ -142,8 +157,8 @@ public:
 
 private :
 
-	/** Select the number of simulated curves based on the distance component-views */
-	void SelectSimulatedCurves(TArray<int32>& NumSimulatedCurves) const;
+	/** Select the number of dynamic curves based on the distance component-views */
+	void SelectDynamicCurves();
 
 	/** Solver settings used to control the simulation */
 	UPROPERTY(EditAnywhere, Category = "Settings")
diff --git a/Engine/Source/Runtime/Experimental/Chaos/Public/GeometryCollection/ManagedArrayTypeValues.inl b/Engine/Source/Runtime/Experimental/Chaos/Public/GeometryCollection/ManagedArrayTypeValues.inl
--- a/Engine/Source/Runtime/Experimental/Chaos/Public/GeometryCollection/ManagedArrayTypeValues.inl
+++ b/Engine/Source/Runtime/Experimental/Chaos/Public/GeometryCollection/ManagedArrayTypeValues.inl
@@ -60,6 +60,7 @@ MANAGED_ARRAY_TYPE(Chaos::PMatrix33d, PMatrix33d)
 MANAGED_ARRAY_TYPE(TArray<Chaos::PMatrix33d>, PMatrix33dArray)
 MANAGED_ARRAY_TYPE(TArray<TArray<FVector3f>>, FVector3fNestedArray)
 MANAGED_ARRAY_TYPE(FUintVector2, UintVector2)
+MANAGED_ARRAY_TYPE(TObjectPtr<UObject>, UObjectArray)
 
 // NOTE: new types must be added at the bottom to keep serialization from breaking
 
