Add vertex color support for cards/meshes groom.

Vertex color presents in the imported static meshes used for cards/meshes geometry are now preserved and accessible in shaders. This data is optional.

diff --git a/Engine/Plugins/Runtime/HairStrands/Source/HairStrandsCore/Private/GroomResources.cpp b/Engine/Plugins/Runtime/HairStrands/Source/HairStrandsCore/Private/GroomResources.cpp
--- a/Engine/Plugins/Runtime/HairStrands/Source/HairStrandsCore/Private/GroomResources.cpp
+++ b/Engine/Plugins/Runtime/HairStrands/Source/HairStrandsCore/Private/GroomResources.cpp
@@ -1008,6 +1008,7 @@ FHairCardsRestResource::FHairCardsRestResource(const FHairCardsBulkData& InBulkD
 	NormalsBuffer(),
 	UVsBuffer(),
 	MaterialsBuffer(),
+	VertexColorsBuffer(),
 	BulkData(InBulkData)
 {
 
@@ -1024,6 +1025,10 @@ void FHairCardsRestResource::InternalAllocate()
 	CreateBuffer<FHairCardsNormalFormat>(RHICmdList, BulkData.Normals, NormalsBuffer, ToHairResourceDebugName(TEXT("Hair.CardsRest_NormalBuffer"), ResourceName), OwnerName, InitialAccess);
 	CreateBuffer<FHairCardsUVFormat>(RHICmdList, BulkData.UVs, UVsBuffer, ToHairResourceDebugName(TEXT("Hair.CardsRest_UVBuffer"), ResourceName), OwnerName, InitialAccess);
 	CreateBuffer<FHairCardsMaterialFormat>(RHICmdList, BulkData.Materials, MaterialsBuffer, ToHairResourceDebugName(TEXT("Hair.CardsRest_MaterialBuffer"), ResourceName), OwnerName, InitialAccess);
+	if (BulkData.Header.bVertexColor)
+	{
+		CreateBuffer<FHairCardsMaterialFormat>(RHICmdList, BulkData.VertexColors, VertexColorsBuffer, ToHairResourceDebugName(TEXT("Hair.CardsRest_VertexColorBuffer"), ResourceName), OwnerName, InitialAccess);
+	}
 }
 
 void FHairCardsRestResource::InternalRelease()
@@ -1038,6 +1043,10 @@ void FHairCardsRestResource::InitResource(FRHICommandListBase& RHICmdList)
 	NormalsBuffer.InitResource(RHICmdList);
 	UVsBuffer.InitResource(RHICmdList);
 	MaterialsBuffer.InitResource(RHICmdList);
+	if (BulkData.Header.bVertexColor)
+	{
+		VertexColorsBuffer.InitResource(RHICmdList);
+	}
 }
 
 void FHairCardsRestResource::ReleaseResource()
@@ -1048,6 +1057,10 @@ void FHairCardsRestResource::ReleaseResource()
 	NormalsBuffer.ReleaseResource();
 	UVsBuffer.ReleaseResource();
 	MaterialsBuffer.ReleaseResource();
+	if (BulkData.Header.bVertexColor)
+	{
+		VertexColorsBuffer.ReleaseResource();
+	}
 }
 
 /////////////////////////////////////////////////////////////////////////////////////////
@@ -1098,6 +1111,7 @@ FHairMeshesRestResource::FHairMeshesRestResource(const FHairMeshesBulkData& InBu
 	IndexBuffer(InBulkData.Indices, InOwnerName),
 	NormalsBuffer(),
 	UVsBuffer(),
+	VertexColorsBuffer(),
 	BulkData(InBulkData)
 {
 	check(BulkData.GetNumVertices() > 0);
@@ -1114,6 +1128,10 @@ void FHairMeshesRestResource::InternalAllocate()
 	CreateBuffer<FHairCardsPositionFormat>(RHICmdList, BulkData.Positions, RestPositionBuffer, ToHairResourceDebugName(TEXT("Hair.MeshesRest_Positions"), ResourceName), OwnerName, InitialAccess);
 	CreateBuffer<FHairCardsNormalFormat>(RHICmdList, BulkData.Normals, NormalsBuffer, ToHairResourceDebugName(TEXT("Hair.MeshesRest_Normals"), ResourceName), OwnerName, InitialAccess);
 	CreateBuffer<FHairCardsUVFormat>(RHICmdList, BulkData.UVs, UVsBuffer, ToHairResourceDebugName(TEXT("Hair.MeshesRest_UVs"), ResourceName), OwnerName, InitialAccess);
+	if (BulkData.Header.bVertexColor)
+	{
+		CreateBuffer<FHairCardsMaterialFormat>(RHICmdList, BulkData.VertexColors, VertexColorsBuffer, ToHairResourceDebugName(TEXT("Hair.MeshesRest_VertexColors"), ResourceName), OwnerName, InitialAccess);
+	}
 }
 
 void FHairMeshesRestResource::InternalRelease()
@@ -1128,6 +1146,10 @@ void FHairMeshesRestResource::InitResource(FRHICommandListBase& RHICmdList)
 	RestPositionBuffer.InitResource(RHICmdList);
 	NormalsBuffer.InitResource(RHICmdList);
 	UVsBuffer.InitResource(RHICmdList);
+	if (BulkData.Header.bVertexColor)
+	{
+		VertexColorsBuffer.InitResource(RHICmdList);
+	}
 }
 
 void FHairMeshesRestResource::ReleaseResource()
@@ -1137,7 +1159,10 @@ void FHairMeshesRestResource::ReleaseResource()
 	RestPositionBuffer.ReleaseResource();
 	NormalsBuffer.ReleaseResource();
 	UVsBuffer.ReleaseResource();
-
+	if (BulkData.Header.bVertexColor)
+	{
+		VertexColorsBuffer.ReleaseResource();
+	}
 }
 
 /////////////////////////////////////////////////////////////////////////////////////////
diff --git a/Engine/Plugins/Runtime/HairStrands/Source/HairStrandsCore/Private/HairCardsBuilder.cpp b/Engine/Plugins/Runtime/HairStrands/Source/HairStrandsCore/Private/HairCardsBuilder.cpp
--- a/Engine/Plugins/Runtime/HairStrands/Source/HairStrandsCore/Private/HairCardsBuilder.cpp
+++ b/Engine/Plugins/Runtime/HairStrands/Source/HairStrandsCore/Private/HairCardsBuilder.cpp
@@ -66,7 +66,7 @@ namespace FHairCardsBuilder
 FString GetVersion()
 {
 	// Important to update the version when cards building or importing changes
-	return TEXT("11");
+	return TEXT("12");
 }
 
 static bool InternalCreateCardsGuides(
@@ -539,6 +539,7 @@ static bool InternalImportGeometry_WithGeneratedGuides(
 	Out.Cards.Normals.SetNum(VertexCount);
 	Out.Cards.Tangents.SetNum(VertexCount);
 	Out.Cards.UVs.SetNum(VertexCount);
+	Out.Cards.VertexColors.Empty();
 	TArray<float> TangentFrameSigns;
 	TangentFrameSigns.SetNum(VertexCount);
 
@@ -550,7 +551,13 @@ static bool InternalImportGeometry_WithGeneratedGuides(
 	const TVertexInstanceAttributesRef<const FVector3f> VertexInstanceTangents	= MeshDescription->VertexInstanceAttributes().GetAttributesRef<FVector3f>(MeshAttribute::VertexInstance::Tangent);
 	const TVertexInstanceAttributesRef<const float> VertexInstanceBinormalSigns	= MeshDescription->VertexInstanceAttributes().GetAttributesRef<float>(MeshAttribute::VertexInstance::BinormalSign);
 	const TVertexInstanceAttributesRef<const FVector2f> VertexInstanceUVs		= MeshDescription->VertexInstanceAttributes().GetAttributesRef<FVector2f>(MeshAttribute::VertexInstance::TextureCoordinate);
+	const TVertexInstanceAttributesRef<const FVector4f> VertexColors			= MeshDescription->VertexInstanceAttributes().GetAttributesRef<FVector4f>(MeshAttribute::VertexInstance::Color);
 
+	const bool bHasVertexColor = VertexColors.GetNumElements() > 0;
+	if (bHasVertexColor)
+	{
+		Out.Cards.VertexColors.SetNum(VertexCount);
+	}
 	for (const FVertexID VertexId : MeshDescription->Vertices().GetElementIDs())
 	{
 		TArrayView<const FVertexInstanceID> VertexInstanceIds = MeshDescription->GetVertexVertexInstanceIDs(VertexId);
@@ -567,12 +574,16 @@ static bool InternalImportGeometry_WithGeneratedGuides(
 		Out.Cards.UVs[VertexIndex]			= FVector4f(VertexInstanceUVs[VertexInstanceId0].Component(0), VertexInstanceUVs[VertexInstanceId0].Component(1), 0, 0); // RootUV are not set here, but will be 'patched' later once guides & interpolation data are built
 		Out.Cards.Tangents[VertexIndex]		= VertexInstanceTangents[VertexInstanceId0];
 		Out.Cards.Normals[VertexIndex]		= VertexInstanceNormals[VertexInstanceId0];
-
-		TangentFrameSigns[VertexIndex] = VertexInstanceBinormalSigns[VertexInstanceId0];
+		TangentFrameSigns[VertexIndex] 		= VertexInstanceBinormalSigns[VertexInstanceId0];
+		if (bHasVertexColor)
+		{
+			Out.Cards.VertexColors[VertexIndex] = FLinearColor(VertexColors[VertexInstanceId0]);
+		}
 
 		Out.Cards.BoundingBox += Out.Cards.Positions[VertexIndex];
 	}
-	OutBulk.BoundingBox = ToFBox3d(Out.Cards.BoundingBox);
+	OutBulk.Header.BoundingBox = ToFBox3d(Out.Cards.BoundingBox);
+	OutBulk.Header.bVertexColor = bHasVertexColor;
 
 	// Fill in render resources (do we need to keep it separated? e.g, format compression, packing)
 	const uint32 PointCount = Out.Cards.Positions.Num();
@@ -580,6 +591,11 @@ static bool InternalImportGeometry_WithGeneratedGuides(
 	OutBulk.Normals.SetNum(PointCount * FHairCardsNormalFormat::ComponentCount);
 	OutBulk.UVs.SetNum(PointCount);
 	OutBulk.Materials.SetNum(PointCount);
+	OutBulk.VertexColors.Empty();
+	if (bHasVertexColor)
+	{
+		OutBulk.VertexColors.SetNum(PointCount);
+	}
 	for (uint32 PointIt = 0; PointIt < PointCount; ++PointIt)
 	{
 		OutBulk.Positions[PointIt] = FVector4f(Out.Cards.Positions[PointIt], 0);
@@ -587,6 +603,10 @@ static bool InternalImportGeometry_WithGeneratedGuides(
 		OutBulk.Normals[PointIt * 2] = FVector4f(Out.Cards.Tangents[PointIt], 0);
 		OutBulk.Normals[PointIt * 2 + 1] = FVector4f(Out.Cards.Normals[PointIt], TangentFrameSigns[PointIt]);
 		OutBulk.Materials[PointIt] = 0;
+		if (bHasVertexColor)
+		{
+			OutBulk.VertexColors[PointIt] = Out.Cards.VertexColors[PointIt].ToFColor(true/*sRGB*/).Bits;
+		}
 	}
 
 	OutBulk.Indices.SetNum(IndexCount);
@@ -897,12 +917,20 @@ static bool InternalImportGeometry_WithImportedGuides(
 	const TVertexInstanceAttributesRef<const FVector3f> VertexInstanceTangents	= MeshDescription->VertexInstanceAttributes().GetAttributesRef<FVector3f>(MeshAttribute::VertexInstance::Tangent);
 	const TVertexInstanceAttributesRef<const float> VertexInstanceBinormalSigns	= MeshDescription->VertexInstanceAttributes().GetAttributesRef<float>(MeshAttribute::VertexInstance::BinormalSign);
 	const TVertexInstanceAttributesRef<const FVector2f> VertexInstanceUVs		= MeshDescription->VertexInstanceAttributes().GetAttributesRef<FVector2f>(MeshAttribute::VertexInstance::TextureCoordinate);
+	const TVertexInstanceAttributesRef<const FVector4f> VertexColors			= MeshDescription->VertexInstanceAttributes().GetAttributesRef<FVector4f>(MeshAttribute::VertexInstance::Color);
+
+	const bool bHasVertexColor = VertexColors.GetNumElements() > 0;
 
 	// Write out BulkData
 	OutBulk.Positions.SetNum(PointCount);
 	OutBulk.Normals.SetNum(PointCount * FHairCardsNormalFormat::ComponentCount);
 	OutBulk.UVs.SetNum(PointCount);
 	OutBulk.Materials.SetNum(PointCount);
+	OutBulk.VertexColors.Empty();
+	if (bHasVertexColor)
+	{
+		OutBulk.VertexColors.SetNum(PointCount);
+	}
 	FBox3f BoundingBox = FBox3f(EForceInit::ForceInit);
 	for (const FVertexID VertexId : MeshDescription->Vertices().GetElementIDs())
 	{
@@ -921,10 +949,15 @@ static bool InternalImportGeometry_WithImportedGuides(
 		OutBulk.Normals[VertexIndex * 2]	= FVector4f(VertexInstanceTangents[VertexInstanceId0], 0);
 		OutBulk.Normals[VertexIndex * 2 + 1]= FVector4f(VertexInstanceNormals[VertexInstanceId0], VertexInstanceBinormalSigns[VertexInstanceId0] /*TangentFrameSigns*/);
 		OutBulk.Materials[VertexIndex] 		= 0;
+		if (bHasVertexColor)
+		{
+			OutBulk.VertexColors[VertexIndex]= FLinearColor(VertexColors[VertexInstanceId0]).ToFColor(true/*sRGB*/).Bits;
+		}
 
 		BoundingBox += OutBulk.Positions[VertexIndex];
 	}
-	OutBulk.BoundingBox = ToFBox3d(BoundingBox);
+	OutBulk.Header.BoundingBox = ToFBox3d(BoundingBox);
+	OutBulk.Header.bVertexColor = bHasVertexColor;
 
 	TArray<FVector3f> CardBaseColor;
 	TArray<float> CardsRoughness;
@@ -1143,7 +1176,7 @@ namespace FHairMeshesBuilder
 FString GetVersion()
 {
 	// Important to update the version when meshes building or importing changes
-	return TEXT("3");
+	return TEXT("4");
 }
 
 void BuildGeometry(
@@ -1250,6 +1283,7 @@ void BuildGeometry(
 	OutBulk.Positions.SetNum(PointCount);
 	OutBulk.Normals.SetNum(PointCount * FHairCardsNormalFormat::ComponentCount);
 	OutBulk.UVs.SetNum(PointCount);
+	OutBulk.VertexColors.Empty();
 	for (uint32 PointIt = 0; PointIt < PointCount; ++PointIt)
 	{
 		OutBulk.Positions[PointIt] = FVector4f(Out.Meshes.Positions[PointIt], 0);
@@ -1259,8 +1293,9 @@ void BuildGeometry(
 
 		Out.Meshes.BoundingBox += FVector4f(OutBulk.Positions[PointIt]);
 	}
-	OutBulk.BoundingBox = ToFBox3d(Out.Meshes.BoundingBox);
-
+	OutBulk.Header.BoundingBox = ToFBox3d(Out.Meshes.BoundingBox);
+	OutBulk.Header.bVertexColor = false;
+	
 	const uint32 IndexCount = Out.Meshes.Indices.Num();
 	OutBulk.Indices.SetNum(IndexCount);
 	for (uint32 IndexIt = 0; IndexIt < IndexCount; ++IndexIt)
@@ -1284,12 +1319,19 @@ void ImportGeometry(
 	const FStaticMeshLODResources& LODData = StaticMesh->GetLODForExport(MeshLODIndex);
 	const uint32 VertexCount = LODData.VertexBuffers.PositionVertexBuffer.GetNumVertices();
 	const uint32 IndexCount = LODData.IndexBuffer.GetNumIndices();
+	const bool bHasVertexColor = LODData.VertexBuffers.ColorVertexBuffer.GetNumVertices() > 0;
 
 	Out.Meshes.Positions.SetNum(VertexCount);
 	Out.Meshes.Normals.SetNum(VertexCount);
 	Out.Meshes.Tangents.SetNum(VertexCount);
 	Out.Meshes.UVs.SetNum(VertexCount);
+	Out.Meshes.VertexColors.Empty();
 	Out.Meshes.Indices.SetNum(IndexCount);
+	if (bHasVertexColor)
+	{
+		Out.Meshes.VertexColors.SetNum(VertexCount);
+		LODData.VertexBuffers.ColorVertexBuffer.GetVertexColors(Out.Meshes.VertexColors);
+	}
 
 	Out.Meshes.BoundingBox.Init();
 	for (uint32 VertexIt = 0; VertexIt < VertexCount; ++VertexIt)
@@ -1308,18 +1350,28 @@ void ImportGeometry(
 	}
 
 	// Fill in render resources (do we need to keep it separated? e.g, format compression, packing)
-	OutBulk.BoundingBox = ToFBox3d(Out.Meshes.BoundingBox);
+	OutBulk.Header.BoundingBox = ToFBox3d(Out.Meshes.BoundingBox);
+	OutBulk.Header.bVertexColor = bHasVertexColor;
 
 	const uint32 PointCount = Out.Meshes.Positions.Num();
 	OutBulk.Positions.SetNum(PointCount);
 	OutBulk.Normals.SetNum(PointCount * FHairCardsNormalFormat::ComponentCount);
 	OutBulk.UVs.SetNum(PointCount);
+	OutBulk.VertexColors.Empty();
+	if (bHasVertexColor)
+	{
+		OutBulk.VertexColors.SetNum(PointCount);
+	}
 	for (uint32 PointIt = 0; PointIt < PointCount; ++PointIt)
 	{
 		OutBulk.Positions[PointIt] = FVector4f(Out.Meshes.Positions[PointIt], 0);
 		OutBulk.UVs[PointIt] = FVector4f(Out.Meshes.UVs[PointIt].X, Out.Meshes.UVs[PointIt].Y, 0, 0);
 		OutBulk.Normals[PointIt * 2] = FVector4f(Out.Meshes.Tangents[PointIt], 0);
 		OutBulk.Normals[PointIt * 2 + 1] = FVector4f(Out.Meshes.Normals[PointIt], 1);
+		if (bHasVertexColor)
+		{
+			OutBulk.VertexColors[PointIt] = Out.Meshes.VertexColors[PointIt].Bits;
+		}
 	}
 
 	OutBulk.Indices.SetNum(IndexCount);
diff --git a/Engine/Plugins/Runtime/HairStrands/Source/HairStrandsCore/Private/HairCardsDatas.cpp b/Engine/Plugins/Runtime/HairStrands/Source/HairStrandsCore/Private/HairCardsDatas.cpp
--- a/Engine/Plugins/Runtime/HairStrands/Source/HairStrandsCore/Private/HairCardsDatas.cpp
+++ b/Engine/Plugins/Runtime/HairStrands/Source/HairStrandsCore/Private/HairCardsDatas.cpp
@@ -13,8 +13,10 @@ void FHairCardsBulkData::Serialize(FArchive& Ar)
 	Ar << Normals;
 	Ar << UVs;
 	Ar << Materials;
+	Ar << VertexColors;
 	Ar << Indices;
-	Ar << BoundingBox;
+	Ar << Header.BoundingBox;
+	Ar << Header.bVertexColor;
 }
 
 FArchive& operator<<(FArchive& Ar, FHairCardsInterpolationVertex& CardInterpVertex)
@@ -40,8 +42,10 @@ void FHairMeshesBulkData::Serialize(FArchive& Ar)
 	Ar << Positions;
 	Ar << Normals;
 	Ar << UVs;
+	Ar << VertexColors;
 	Ar << Indices;
-	Ar << BoundingBox;
+	Ar << Header.BoundingBox;
+	Ar << Header.bVertexColor;
 }
 
 
diff --git a/Engine/Plugins/Runtime/HairStrands/Source/HairStrandsCore/Private/HairCardsVertexFactory.cpp b/Engine/Plugins/Runtime/HairStrands/Source/HairStrandsCore/Private/HairCardsVertexFactory.cpp
--- a/Engine/Plugins/Runtime/HairStrands/Source/HairStrandsCore/Private/HairCardsVertexFactory.cpp
+++ b/Engine/Plugins/Runtime/HairStrands/Source/HairStrandsCore/Private/HairCardsVertexFactory.cpp
@@ -28,7 +28,14 @@ template<typename T> inline void VFC_BindParam(FMeshDrawSingleShaderBindings& Sh
 
 enum class EHairCardsFactoryFlags : uint32
 {
-	InvertedUV = 0x1,
+	Texture0    = 0x1,
+	Texture1    = 0x2,
+	Texture2    = 0x4,
+	Texture3    = 0x8,
+	Texture4    = 0x10,
+	Texture5    = 0x20,
+	InvertedUV  = 0x40,
+	VertexColor = 0x80
 };
 
 static void SetTextures(
@@ -43,7 +50,7 @@ static void SetTextures(
 	{
 		if (InTextures[TextureIt] != nullptr)
 		{
-			Out.Flags |= 1u<<(TextureIt+1u);		
+			Out.Flags |= (1u<<TextureIt);
 			switch (TextureIt)
 			{
 				case 0:
@@ -102,9 +109,12 @@ FHairCardsUniformBuffer CreateHairCardsVFUniformBuffer(
 	if (GeometryType == EHairGeometryType::Cards)
 	{
 		const FHairGroupInstance::FCards::FLOD& LOD = Instance->Cards.LODs[LODIndex];
+		const bool bHasVertexColor = LOD.RestResource->BulkData.Header.bVertexColor;
 
 		// Cards atlas UV are inverted so fetching needs to be inverted on the y-axis
-		UniformParameters.Flags = LOD.RestResource->bInvertUV ? uint32(EHairCardsFactoryFlags::InvertedUV) : 0u;
+		UniformParameters.Flags = 
+			(LOD.RestResource->bInvertUV ? uint32(EHairCardsFactoryFlags::InvertedUV) : 0u) |
+			(bHasVertexColor ? uint32(EHairCardsFactoryFlags::VertexColor) : 0u);
 		UniformParameters.AttributeTextureIndex = 0; // TODO pack attribute texture/channel description. Currently use hardcoded layout in shader attribute code.
 		UniformParameters.AttributeChannelIndex = 0; // TODO pack attribute texture/channel description. Currently use hardcoded layout in shader attribute code.
 		UniformParameters.MaxVertexCount = LOD.RestResource->GetVertexCount();
@@ -128,12 +138,21 @@ FHairCardsUniformBuffer CreateHairCardsVFUniformBuffer(
 
 		UniformParameters.UVsBuffer = LOD.RestResource->UVsBuffer.ShaderResourceViewRHI.GetReference();
 		UniformParameters.MaterialsBuffer = LOD.RestResource->MaterialsBuffer.ShaderResourceViewRHI.GetReference();
+		if (bHasVertexColor)
+		{
+			UniformParameters.VertexColorsBuffer = LOD.RestResource->VertexColorsBuffer.ShaderResourceViewRHI.GetReference();
+		}
+		else
+		{
+			UniformParameters.VertexColorsBuffer = UniformParameters.NormalsBuffer; // Reuse normal buffer as a dummy input for material buffer, since it not used for meshes (material data is fetch through textures)
+		}
 
 		SetTextures(LOD.RestResource->Textures, LOD.RestResource->Samplers, LOD.RestResource->LayoutIndex, UniformParameters);
 	}
 	else if (GeometryType == EHairGeometryType::Meshes)
 	{
 		const FHairGroupInstance::FMeshes::FLOD& LOD = Instance->Meshes.LODs[LODIndex];
+		const bool bHasVertexColor = LOD.RestResource->BulkData.Header.bVertexColor;
 
 		// When the geometry is not-dynamic (no binding to skeletal mesh, no simulation), only a single vertex buffer is allocated. 
 		// In this case we force the buffer index to 0
@@ -150,7 +169,7 @@ FHairCardsUniformBuffer CreateHairCardsVFUniformBuffer(
 		}
 
 		// Meshes UV are not inverted so no need to invert the y-axis
-		UniformParameters.Flags = 0;
+		UniformParameters.Flags = (bHasVertexColor ? uint32(EHairCardsFactoryFlags::VertexColor) : 0u);
 		UniformParameters.AttributeTextureIndex = 0; // TODO pack attribute texture/channel description. Currently use hardcoded layout in shader attribute code.
 		UniformParameters.AttributeChannelIndex = 0; // TODO pack attribute texture/channel description. Currently use hardcoded layout in shader attribute code.
 		UniformParameters.MaxVertexCount = LOD.RestResource->GetVertexCount();
@@ -159,6 +178,14 @@ FHairCardsUniformBuffer CreateHairCardsVFUniformBuffer(
 		UniformParameters.NormalsBuffer = LOD.RestResource->NormalsBuffer.ShaderResourceViewRHI.GetReference();
 		UniformParameters.UVsBuffer = LOD.RestResource->UVsBuffer.ShaderResourceViewRHI.GetReference();
 		UniformParameters.MaterialsBuffer = UniformParameters.NormalsBuffer; // Reuse normal buffer as a dummy input for material buffer, since it not used for meshes (material data is fetch through textures)
+		if (bHasVertexColor)
+		{
+			UniformParameters.VertexColorsBuffer = LOD.RestResource->VertexColorsBuffer.ShaderResourceViewRHI.GetReference();
+		}
+		else
+		{
+			UniformParameters.VertexColorsBuffer = UniformParameters.NormalsBuffer; // Reuse normal buffer as a dummy input for material buffer, since it not used for meshes (material data is fetch through textures)
+		}
 
 		SetTextures(LOD.RestResource->Textures, LOD.RestResource->Samplers, LOD.RestResource->LayoutIndex, UniformParameters);
 	}
@@ -170,6 +197,7 @@ FHairCardsUniformBuffer CreateHairCardsVFUniformBuffer(
 		UniformParameters.NormalsBuffer          = GNullVertexBuffer.VertexBufferSRV;
 		UniformParameters.UVsBuffer              = GNullVertexBuffer.VertexBufferSRV;
 		UniformParameters.MaterialsBuffer        = GNullVertexBuffer.VertexBufferSRV;
+		UniformParameters.VertexColorsBuffer     = GNullVertexBuffer.VertexBufferSRV;
 	}
 
 	FRHITexture* DefaultTexture = GBlackTexture->TextureRHI;
@@ -309,7 +337,8 @@ void FHairCardsVertexFactory::GetPSOPrecacheVertexFetchElements(EVertexInputStre
 		Elements.Add(FVertexElement(1+NumStreams, NormalOffset1, FHairCardsNormalFormat::VertexElementType, 2, FHairCardsNormalFormat::SizeInByte * FHairCardsNormalFormat::ComponentCount, false));
 
 		Elements.Add(FVertexElement(2+NumStreams, 0, FHairCardsUVFormat::VertexElementType, 3, FHairCardsUVFormat::SizeInByte, false));
-		Elements.Add(FVertexElement(3+NumStreams, 0, FHairCardsMaterialFormat::VertexElementType, 4, FHairCardsMaterialFormat::SizeInByte, false));
+		Elements.Add(FVertexElement(3+NumStreams, 0, FHairCardsMaterialFormat::VertexElementType, 4, FHairCardsMaterialFormat::SizeInByte, false)); // Material
+		Elements.Add(FVertexElement(3+NumStreams, 0, FHairCardsMaterialFormat::VertexElementType, 5, FHairCardsMaterialFormat::SizeInByte, false)); // VertexColor
 	}
 }
 
@@ -371,7 +400,7 @@ void FHairCardsVertexFactory::InitResources(FRHICommandListBase& RHICmdList)
 		if (Data.GeometryType == EHairGeometryType::Cards)
 		{
 			const FHairGroupInstance::FCards::FLOD& LOD = Data.Instance->Cards.LODs[Data.LODIndex];
-			
+			const bool bHasVertexColor = LOD.RestResource->BulkData.Header.bVertexColor;
 			const bool bDynamic = LOD.DeformedResource != nullptr;
 			if (bDynamic)
 			{
@@ -400,21 +429,34 @@ void FHairCardsVertexFactory::InitResources(FRHICommandListBase& RHICmdList)
 			}
 			Elements.Add(AccessStreamComponent(FVertexStreamComponent(&LOD.RestResource->UVsBuffer,					0, 0,									FHairCardsUVFormat::SizeInByte,														FHairCardsUVFormat::VertexElementType,			EVertexStreamUsage::Default), 3));
 			Elements.Add(AccessStreamComponent(FVertexStreamComponent(&LOD.RestResource->MaterialsBuffer,			0, 0,									FHairCardsMaterialFormat::SizeInByte,												FHairCardsMaterialFormat::VertexElementType,	EVertexStreamUsage::Default), 4));
+			if (bHasVertexColor)
+			{
+				Elements.Add(AccessStreamComponent(FVertexStreamComponent(&LOD.RestResource->VertexColorsBuffer,	0, 0,									FHairCardsMaterialFormat::SizeInByte,												FHairCardsMaterialFormat::VertexElementType,	EVertexStreamUsage::Default), 5));
+			}
+			else
+			{
+				Elements.Add(AccessStreamComponent(FVertexStreamComponent(&LOD.RestResource->NormalsBuffer,			0, 0,									FHairCardsMaterialFormat::SizeInByte,												FHairCardsMaterialFormat::VertexElementType,	EVertexStreamUsage::Default), 5));
+			}
 
 			// Ensure the rest resources are in correct states for the VF
 			const FHairCardsRestResource* RestResource = LOD.RestResource;
 			ENQUEUE_RENDER_COMMAND(FHairCardsVertexFactoryTransition)(/*UE::RenderCommandPipe::Groom,*/
-			[RestResource](FRHICommandListImmediate& RHICmdImmediateList)
+			[RestResource, bHasVertexColor](FRHICommandListImmediate& RHICmdImmediateList)
 			{
 				RHICmdImmediateList.Transition(FRHITransitionInfo(RestResource->RestPositionBuffer.VertexBufferRHI,	ERHIAccess::Unknown, ERHIAccess::VertexOrIndexBuffer));
 				RHICmdImmediateList.Transition(FRHITransitionInfo(RestResource->NormalsBuffer.VertexBufferRHI, 		ERHIAccess::Unknown, ERHIAccess::VertexOrIndexBuffer));
 				RHICmdImmediateList.Transition(FRHITransitionInfo(RestResource->UVsBuffer.VertexBufferRHI, 			ERHIAccess::Unknown, ERHIAccess::VertexOrIndexBuffer));
 				RHICmdImmediateList.Transition(FRHITransitionInfo(RestResource->MaterialsBuffer.VertexBufferRHI, 	ERHIAccess::Unknown, ERHIAccess::VertexOrIndexBuffer));
+				if (bHasVertexColor)
+				{
+					RHICmdImmediateList.Transition(FRHITransitionInfo(RestResource->VertexColorsBuffer.VertexBufferRHI, ERHIAccess::Unknown, ERHIAccess::VertexOrIndexBuffer));
+				}
 			});
 		}
 		else if (Data.GeometryType == EHairGeometryType::Meshes)
 		{
 			const FHairGroupInstance::FMeshes::FLOD& LOD = Data.Instance->Meshes.LODs[Data.LODIndex];
+			const bool bHasVertexColor = LOD.RestResource->BulkData.Header.bVertexColor;
 
 			// Note: Use the 'Normal' buffer as a dummy input for 'Material' buffer, as Material data is fetched through textures for meshes
 			const bool bDynamic = LOD.DeformedResource != nullptr;
@@ -437,16 +479,28 @@ void FHairCardsVertexFactory::InitResources(FRHICommandListBase& RHICmdList)
 			Elements.Add(AccessStreamComponent(FVertexStreamComponent(&LOD.RestResource->NormalsBuffer,				0, 0,									FHairCardsNormalFormat::SizeInByte * FHairCardsNormalFormat::ComponentCount,		FHairCardsNormalFormat::VertexElementType,		EVertexStreamUsage::Default), 1));
 			Elements.Add(AccessStreamComponent(FVertexStreamComponent(&LOD.RestResource->NormalsBuffer,				0, FHairCardsNormalFormat::SizeInByte,	FHairCardsNormalFormat::SizeInByte * FHairCardsNormalFormat::ComponentCount,		FHairCardsNormalFormat::VertexElementType,		EVertexStreamUsage::Default), 2));
 			Elements.Add(AccessStreamComponent(FVertexStreamComponent(&LOD.RestResource->UVsBuffer,					0, 0,									FHairCardsUVFormat::SizeInByte,														FHairCardsUVFormat::VertexElementType,			EVertexStreamUsage::Default), 3));
-			Elements.Add(AccessStreamComponent(FVertexStreamComponent(&LOD.RestResource->NormalsBuffer,				0, 0,									FHairCardsMaterialFormat::SizeInByte,												FHairCardsMaterialFormat::VertexElementType,	EVertexStreamUsage::Default), 4)); 
+			Elements.Add(AccessStreamComponent(FVertexStreamComponent(&LOD.RestResource->NormalsBuffer,				0, 0,									FHairCardsMaterialFormat::SizeInByte,												FHairCardsMaterialFormat::VertexElementType,	EVertexStreamUsage::Default), 4)); // Dummy
+			if (bHasVertexColor)
+			{
+				Elements.Add(AccessStreamComponent(FVertexStreamComponent(&LOD.RestResource->VertexColorsBuffer,	0, 0,									FHairCardsMaterialFormat::SizeInByte,												FHairCardsMaterialFormat::VertexElementType,	EVertexStreamUsage::Default), 5));
+			}
+			else
+			{
+				Elements.Add(AccessStreamComponent(FVertexStreamComponent(&LOD.RestResource->NormalsBuffer,			0, 0,									FHairCardsMaterialFormat::SizeInByte,												FHairCardsMaterialFormat::VertexElementType,	EVertexStreamUsage::Default), 5));
+			}
 
 			// Ensure the rest resources are in correct states for the VF
 			const FHairMeshesRestResource* RestResource = LOD.RestResource;
 			ENQUEUE_RENDER_COMMAND(FHairCardsVertexFactoryTransition)(/*UE::RenderCommandPipe::Groom,*/
-			[RestResource](FRHICommandListImmediate& RHICmdImmediateList)
+			[RestResource, bHasVertexColor](FRHICommandListImmediate& RHICmdImmediateList)
 			{
 				RHICmdImmediateList.Transition(FRHITransitionInfo(RestResource->RestPositionBuffer.VertexBufferRHI,	ERHIAccess::Unknown, ERHIAccess::VertexOrIndexBuffer));
 				RHICmdImmediateList.Transition(FRHITransitionInfo(RestResource->NormalsBuffer.VertexBufferRHI, 		ERHIAccess::Unknown, ERHIAccess::VertexOrIndexBuffer));
 				RHICmdImmediateList.Transition(FRHITransitionInfo(RestResource->UVsBuffer.VertexBufferRHI, 			ERHIAccess::Unknown, ERHIAccess::VertexOrIndexBuffer));
+				if (bHasVertexColor)
+				{
+					RHICmdImmediateList.Transition(FRHITransitionInfo(RestResource->VertexColorsBuffer.VertexBufferRHI, ERHIAccess::Unknown, ERHIAccess::VertexOrIndexBuffer));
+				}
 			});
 		}
 
diff --git a/Engine/Plugins/Runtime/HairStrands/Source/HairStrandsCore/Public/GroomAsset.h b/Engine/Plugins/Runtime/HairStrands/Source/HairStrandsCore/Public/GroomAsset.h
--- a/Engine/Plugins/Runtime/HairStrands/Source/HairStrandsCore/Public/GroomAsset.h
+++ b/Engine/Plugins/Runtime/HairStrands/Source/HairStrandsCore/Public/GroomAsset.h
@@ -295,7 +295,7 @@ struct FHairGroupPlatformData
 		{
 			for (const FLOD& LOD : LODs)
 			{
-				if (LOD.IsValid()) return LOD.BulkData.BoundingBox;
+				if (LOD.IsValid()) return LOD.BulkData.Header.BoundingBox;
 			}
 			return FBox();
 		}
@@ -318,6 +318,7 @@ struct FHairGroupPlatformData
 				Total += BulkData.Positions.GetAllocatedSize();
 				Total += BulkData.Normals.GetAllocatedSize();
 				Total += BulkData.UVs.GetAllocatedSize();
+				Total += BulkData.VertexColors.GetAllocatedSize();
 				Total += BulkData.Indices.GetAllocatedSize();
 				Total += InterpolationBulkData.Interpolation.GetAllocatedSize();
 				Total += ::GetDataSize(GuideBulkData);
@@ -356,7 +357,7 @@ struct FHairGroupPlatformData
 		{
 			for (const FLOD& LOD : LODs)
 			{
-				if (LOD.IsValid()) return LOD.BulkData.BoundingBox;
+				if (LOD.IsValid()) return LOD.BulkData.Header.BoundingBox;
 			}
 			return FBox();
 		}
@@ -379,6 +380,7 @@ struct FHairGroupPlatformData
 				Total += BulkData.Positions.GetAllocatedSize();
 				Total += BulkData.Normals.GetAllocatedSize();
 				Total += BulkData.UVs.GetAllocatedSize();
+				Total += BulkData.VertexColors.GetAllocatedSize();
 				Total += BulkData.Indices.GetAllocatedSize();
 				return Total;
 			}
diff --git a/Engine/Plugins/Runtime/HairStrands/Source/HairStrandsCore/Public/GroomInstance.h b/Engine/Plugins/Runtime/HairStrands/Source/HairStrandsCore/Public/GroomInstance.h
--- a/Engine/Plugins/Runtime/HairStrands/Source/HairStrandsCore/Public/GroomInstance.h
+++ b/Engine/Plugins/Runtime/HairStrands/Source/HairStrandsCore/Public/GroomInstance.h
@@ -35,6 +35,7 @@ BEGIN_GLOBAL_SHADER_PARAMETER_STRUCT(FHairCardsVertexFactoryUniformShaderParamet
 	SHADER_PARAMETER_SRV(Buffer<float4>, NormalsBuffer)
 	SHADER_PARAMETER_SRV(Buffer<float4>, UVsBuffer)
 	SHADER_PARAMETER_SRV(Buffer<float4>, MaterialsBuffer)
+	SHADER_PARAMETER_SRV(Buffer<float4>, VertexColorsBuffer)
 	SHADER_PARAMETER_TEXTURE(Texture2D<float4>, Texture0Texture) 
 	SHADER_PARAMETER_SAMPLER(SamplerState,      Texture0Sampler)
 	SHADER_PARAMETER_TEXTURE(Texture2D<float4>, Texture1Texture)
diff --git a/Engine/Plugins/Runtime/HairStrands/Source/HairStrandsCore/Public/GroomResources.h b/Engine/Plugins/Runtime/HairStrands/Source/HairStrandsCore/Public/GroomResources.h
--- a/Engine/Plugins/Runtime/HairStrands/Source/HairStrandsCore/Public/GroomResources.h
+++ b/Engine/Plugins/Runtime/HairStrands/Source/HairStrandsCore/Public/GroomResources.h
@@ -715,6 +715,7 @@ struct FHairCardsRestResource : public FHairCommonResource
 	FHairCardsVertexBuffer NormalsBuffer;
 	FHairCardsVertexBuffer UVsBuffer;
 	FHairCardsVertexBuffer MaterialsBuffer;
+	FHairCardsVertexBuffer VertexColorsBuffer;
 
 	TArray<FSamplerStateRHIRef> Samplers;
 	TArray<FTextureReferenceRHIRef> Textures;
@@ -819,6 +820,7 @@ struct FHairMeshesRestResource : public FHairCommonResource
 		Total += BulkData.Normals.GetAllocatedSize();
 		Total += BulkData.Indices.GetAllocatedSize();
 		Total += BulkData.UVs.GetAllocatedSize();
+		Total += BulkData.VertexColors.GetAllocatedSize();
 		return Total;
 	}
 
@@ -832,6 +834,7 @@ struct FHairMeshesRestResource : public FHairCommonResource
 
 	FHairCardsVertexBuffer NormalsBuffer;
 	FHairCardsVertexBuffer UVsBuffer;
+	FHairCardsVertexBuffer VertexColorsBuffer;
 
 	TArray<FSamplerStateRHIRef> Samplers;
 	TArray<FTextureReferenceRHIRef>	Textures;
diff --git a/Engine/Plugins/Runtime/HairStrands/Source/HairStrandsCore/Public/HairCardsDatas.h b/Engine/Plugins/Runtime/HairStrands/Source/HairStrandsCore/Public/HairCardsDatas.h
--- a/Engine/Plugins/Runtime/HairStrands/Source/HairStrandsCore/Public/HairCardsDatas.h
+++ b/Engine/Plugins/Runtime/HairStrands/Source/HairStrandsCore/Public/HairCardsDatas.h
@@ -90,6 +90,7 @@ struct FHairCardsGeometry
 	TArray<uint32>   Indices;
 	TArray<float>    CoordU;  // Transient data storing [0..1] parametric value along main axis. This is used for generating guides & interpolation data
 	TArray<FVector2f>LocalUVs;// Transient data storing [0..1, 0..1] local UV of the cards. U is aligned with the card's main direction, and V is aligned perpendicularity
+	TArray<FLinearColor> VertexColors;
 
 	// Vertex offset and vertex count of each cards geometry
 	// No longer used, kept only for backward compatibility
@@ -110,6 +111,7 @@ struct FHairCardsGeometry
 		Positions.Reset();
 		Indices.Reset();
 		CoordU.Reset();
+		VertexColors.Reset();
 
 		IndexOffsets.Reset();
 		IndexCounts.Reset();
@@ -126,6 +128,7 @@ struct FHairCardsGeometry
 		Positions.Empty();
 		Indices.Empty();
 		CoordU.Reset();
+		VertexColors.Empty();
 
 		// Cards' indices offset & count
 		IndexOffsets.SetNum(Count);
@@ -169,12 +172,14 @@ struct FHairCardsBulkData
 		Normals.Empty();
 		UVs.Empty();
 		Materials.Empty();
+		VertexColors.Empty();
 		Indices.Empty();
+		Header = FHeader();
 	}
 
 	uint32 GetNumTriangles() const { return Indices.Num() / 3; }
 	uint32 GetNumVertices() const { return Positions.Num(); }
-	const FBox& GetBounds() const { return BoundingBox; }
+	const FBox& GetBounds() const { return Header.BoundingBox; }
 
 	UTexture2D* DepthTexture = nullptr;
 	UTexture2D* TangentTexture = nullptr;
@@ -185,8 +190,14 @@ struct FHairCardsBulkData
 	TArray<FHairCardsNormalFormat::Type> Normals;
 	TArray<FHairCardsUVFormat::Type> UVs;
 	TArray<FHairCardsMaterialFormat::Type> Materials;
+	TArray<FHairCardsMaterialFormat::Type> VertexColors;
 	TArray<FHairCardsIndexFormat::Type> Indices;
-	FBox BoundingBox = FBox(EForceInit::ForceInit);
+
+	struct FHeader
+	{
+		FBox BoundingBox = FBox(EForceInit::ForceInit);
+		bool bVertexColor = false;
+	} Header;
 
 	void Serialize(FArchive& Ar);
 };
@@ -230,6 +241,7 @@ struct FHairMeshes
 {
 	// Geometry
 	TArray<FVector2f> UVs;
+	TArray<FColor>    VertexColors;
 	TArray<FVector3f> Normals;
 	TArray<FVector3f> Tangents;
 	TArray<FVector3f> Positions;
@@ -241,6 +253,7 @@ struct FHairMeshes
 	{
 		// Geometry
 		UVs.Empty();
+		VertexColors.Empty();
 		Normals.Empty();
 		Tangents.Empty();
 		Positions.Empty();
@@ -266,13 +279,19 @@ struct FHairMeshesBulkData
 
 	uint32 GetNumTriangles() const { return Indices.Num() / 3; }
 	uint32 GetNumVertices() const { return Positions.Num(); }
-	const FBox& GetBounds() const { return BoundingBox; }
+	const FBox& GetBounds() const { return Header.BoundingBox; }
 
 	TArray<FHairCardsPositionFormat::Type> Positions;
 	TArray<FHairCardsNormalFormat::Type> Normals;
 	TArray<FHairCardsUVFormat::Type> UVs;
+	TArray<FHairCardsMaterialFormat::Type> VertexColors;
 	TArray<FHairCardsIndexFormat::Type> Indices;
-	FBox BoundingBox = FBox(EForceInit::ForceInit);
+
+	struct FHeader
+	{
+		FBox BoundingBox = FBox(EForceInit::ForceInit);
+		bool bVertexColor = false;
+	} Header;
 
 	void Serialize(FArchive& Ar);
 };
diff --git a/Engine/Shaders/Private/HairStrands/HairCardsAttributeCommon.ush b/Engine/Shaders/Private/HairStrands/HairCardsAttributeCommon.ush
--- a/Engine/Shaders/Private/HairStrands/HairCardsAttributeCommon.ush
+++ b/Engine/Shaders/Private/HairStrands/HairCardsAttributeCommon.ush
@@ -7,13 +7,14 @@
 ///////////////////////////////////////////////////////////////////////////////////////////////////
 
 // Vertex factory Flags
-#define EHairCardsVFlags_InvertedUV  0x1
-#define EHairCardsVFlags_Texture0    0x2
-#define EHairCardsVFlags_Texture1    0x4
-#define EHairCardsVFlags_Texture2    0x8
-#define EHairCardsVFlags_Texture3    0x10
-#define EHairCardsVFlags_Texture4    0x20
-#define EHairCardsVFlags_Texture5    0x40
+#define EHairCardsVFlags_Texture0    0x1
+#define EHairCardsVFlags_Texture1    0x2
+#define EHairCardsVFlags_Texture2    0x4
+#define EHairCardsVFlags_Texture3    0x8
+#define EHairCardsVFlags_Texture4    0x10
+#define EHairCardsVFlags_Texture5    0x20
+#define EHairCardsVFlags_InvertedUV  0x40
+#define EHairCardsVFlags_VertexColor 0x80
 
 // Attribute index
 #define EHairCardsVFAttribute_Depth        0
@@ -87,12 +88,13 @@ uint GetHairCardChannelIndex(uint InLayoutIndex, uint InAttribute)
 // uint GetHairCardChannelIndex(uint In)	{ return ((HairCards.AttributeChannelIndex) >> (In*2u) & 0x3u); }
 
 float2 GetAtlasUV(float2 In) 				{ return float2(In.x, (HairCardsVF.Flags & EHairCardsVFlags_InvertedUV) ? 1 - In.y : In.y); }
+bool HasHairCardsVertexColor() 				{ return (HairCardsVF.Flags & EHairCardsVFlags_VertexColor); }
 
 bool HasTextureBound(uint In)
 {
 	// HairCardsVF.Flags stores InverttUV on the first bit, texture bound after
 	const uint TextureIndex = GetHairCardTextureIndex(HairCardsVF.LayoutIndex, In);
-	return (TextureIndex < HairCardsVF.TextureCount) && (HairCardsVF.Flags & (1u << (TextureIndex + 1))) != 0;
+	return (TextureIndex < HairCardsVF.TextureCount) && (HairCardsVF.Flags & (1u << TextureIndex)) != 0;
 }
 
 // Generic attribute texture fetches
diff --git a/Engine/Shaders/Private/HairStrands/HairCardsVertexFactory.ush b/Engine/Shaders/Private/HairStrands/HairCardsVertexFactory.ush
--- a/Engine/Shaders/Private/HairStrands/HairCardsVertexFactory.ush
+++ b/Engine/Shaders/Private/HairStrands/HairCardsVertexFactory.ush
@@ -50,6 +50,7 @@ struct FVertexFactoryInput
   #endif
 	float4 UVs		 : ATTRIBUTE3;
 	float4 Materials : ATTRIBUTE4;
+	float4 VertexColors : ATTRIBUTE5;
 #endif
 
 	VF_GPUSCENE_DECLARE_INPUT_BLOCK(13)
@@ -81,6 +82,7 @@ struct FVertexFactoryIntermediates
 	half3x3 TangentToLocal;
 	half3x3 TangentToWorld;
 	half TangentToWorldSign;
+	half4 VertexColors;
 	float2 HairPrimitiveUV;
 	float2 HairPrimitiveRootUV;
 	float2 HairPrimitiveLength;
@@ -113,19 +115,29 @@ struct FHairCardsPoint
 	float  CardLength;
 	float  GroupIndex;
 	float4 Materials;
+	half4  VertexColors;
 };
 
 FHairCardsPoint GetVertexPosition(FVertexFactoryInput Input)
 {
+	half4 VertexColors = 0;
 #if MANUAL_VERTEX_FETCH
 	const uint VertexFetchId = GetHairVertexFetchId(Input);
 	const float4 Data = HairCardsVF.PositionBuffer[VertexFetchId];
 	const float4 UVs  = HairCardsVF.UVsBuffer[VertexFetchId];
 	const float4 Materials = HairCardsVF.MaterialsBuffer[VertexFetchId];
+	if (HasHairCardsVertexColor())
+	{
+		VertexColors = HairCardsVF.VertexColorsBuffer[VertexFetchId] FMANUALFETCH_COLOR_COMPONENT_SWIZZLE; // Swizzle vertex color.
+	}
 #else
 	const float4 Data = Input.Position;
 	const float4 UVs  = Input.UVs;
 	const float4 Materials = Input.Materials;
+	if (HasHairCardsVertexColor())
+	{
+		VertexColors = Input.VertexColors FCOLOR_COMPONENT_SWIZZLE; // Swizzle vertex color.
+	}
 #endif
 
 	const uint EncodedW = asuint(Data.w);
@@ -137,6 +149,7 @@ FHairCardsPoint GetVertexPosition(FVertexFactoryInput Input)
 	Out.CardLength	= f16tof32(EncodedW & 0xFFFF);
 	Out.GroupIndex	= (EncodedW>>16u);
 	Out.Materials	= float4(Materials.xyz*Materials.xyz, Materials.w); // Decoding (cheap) sRGB -> linear color
+	Out.VertexColors= VertexColors;
 	return Out;
 }
 
@@ -186,6 +199,9 @@ FMaterialPixelParameters GetMaterialPixelParameters(FVertexFactoryInterpolantsVS
 	Result.HairPrimitiveLength = Interpolants.HairPrimitiveLength;
 	Result.HairPrimitiveGroupIndex = Interpolants.HairPrimitiveGroupIndex;
 	Result.HairPrimitiveMaterial = Interpolants.HairPrimitiveMaterial;
+
+	Result.VertexColor = GetColor(Interpolants);
+
 #if NUM_TEX_COORD_INTERPOLATORS
 	UNROLL
 	for (int CoordinateIndex = 0; CoordinateIndex < NUM_TEX_COORD_INTERPOLATORS; CoordinateIndex++)
@@ -231,6 +247,8 @@ FMaterialVertexParameters GetMaterialVertexParameters(
 		Result.PositionInstanceSpace = VertexFactoryGetInstanceSpacePosition(Input, Intermediates);
 	}
 	Result.PositionPrimitiveSpace = Result.PositionInstanceSpace; // No support for instancing, so instance == primitive
+	
+	Result.VertexColor = Intermediates.VertexColors;
 
 	// does not handle instancing!
 	Result.TangentToWorld = Intermediates.TangentToWorld;
@@ -295,6 +313,7 @@ FVertexFactoryIntermediates GetVertexFactoryIntermediates(FVertexFactoryInput In
 	Intermediates.HairPrimitiveLength		= CardsPoint.CardLength; 
 	Intermediates.HairPrimitiveGroupIndex	= CardsPoint.GroupIndex;
 	Intermediates.HairPrimitiveMaterial		= CardsPoint.Materials;
+	Intermediates.VertexColors				= CardsPoint.VertexColors;
 
 	float TangentSign;
 	Intermediates.TangentToLocal = CalcTangentToLocal(Input, TangentSign);
@@ -347,6 +366,7 @@ FVertexFactoryInterpolantsVSToPS VertexFactoryGetInterpolantsVSToPS(FVertexFacto
 	Interpolants = (FVertexFactoryInterpolantsVSToPS)0;
 
 	SetTangents(Interpolants, Intermediates.TangentToWorld[0], Intermediates.TangentToWorld[2], Intermediates.TangentToWorldSign);
+	SetColor(Interpolants, Intermediates.VertexColors);
 
 #if INSTANCED_STEREO
 	Interpolants.EyeIndex = 0;
