Fix CIS for the groom mesh read deformer DI

diff --git a/Engine/Plugins/Runtime/HairStrands/Source/HairStrandsDeformer/Private/DeformerGroomMeshesRead.cpp b/Engine/Plugins/Runtime/HairStrands/Source/HairStrandsDeformer/Private/DeformerGroomMeshesRead.cpp
--- a/Engine/Plugins/Runtime/HairStrands/Source/HairStrandsDeformer/Private/DeformerGroomMeshesRead.cpp
+++ b/Engine/Plugins/Runtime/HairStrands/Source/HairStrandsDeformer/Private/DeformerGroomMeshesRead.cpp
@@ -128,15 +128,12 @@ FORCEINLINE void CreateBonesBuffers(const FSkeletalMeshObject* SkeletalMeshObjec
 		BoneVectors.Init(FVector4f::Zero(), NumBones * 3);
 
 		int32 VectorIndex = 0;
-		FMatrix3x4 BoneMatrix;
 		for(int32 BoneIndex = 0; BoneIndex < NumBones; ++BoneIndex)
 		{
 			const FMatrix44f LocalMatrix = BoneRefToLocals[BoneIndex] * SkeletalMeshTransform;
-			LocalMatrix.To3x4MatrixTranspose(static_cast<float*>(&BoneMatrix.M[0][0]));
-			
-			BoneVectors[VectorIndex++] = FVector4f(BoneMatrix.M[0][0], BoneMatrix.M[0][1], BoneMatrix.M[0][2], BoneMatrix.M[0][3]);
-			BoneVectors[VectorIndex++] = FVector4f(BoneMatrix.M[1][0], BoneMatrix.M[1][1], BoneMatrix.M[1][2], BoneMatrix.M[1][3]);
-			BoneVectors[VectorIndex++] = FVector4f(BoneMatrix.M[2][0], BoneMatrix.M[2][1], BoneMatrix.M[2][2], BoneMatrix.M[2][3]);
+			LocalMatrix.To3x4MatrixTranspose(static_cast<float*>(&BoneVectors[VectorIndex].X));
+
+			VectorIndex += 3;
 		}
 
 		FRDGBufferRef& BonesBuffer = BonesMatricesBuffers.Add_GetRef( GraphBuilder.CreateBuffer(
