Fix shader compilation due to incorrect vertex attribute assigned to vertex colors.

diff --git a/Engine/Plugins/Runtime/HairStrands/Source/HairStrandsCore/Private/HairCardsVertexFactory.cpp b/Engine/Plugins/Runtime/HairStrands/Source/HairStrandsCore/Private/HairCardsVertexFactory.cpp
--- a/Engine/Plugins/Runtime/HairStrands/Source/HairStrandsCore/Private/HairCardsVertexFactory.cpp
+++ b/Engine/Plugins/Runtime/HairStrands/Source/HairStrandsCore/Private/HairCardsVertexFactory.cpp
@@ -431,11 +431,11 @@ void FHairCardsVertexFactory::InitResources(FRHICommandListBase& RHICmdList)
 			Elements.Add(AccessStreamComponent(FVertexStreamComponent(&LOD.RestResource->MaterialsBuffer,			0, 0,									FHairCardsMaterialFormat::SizeInByte,												FHairCardsMaterialFormat::VertexElementType,	EVertexStreamUsage::Default), 4));
 			if (bHasVertexColor)
 			{
-				Elements.Add(AccessStreamComponent(FVertexStreamComponent(&LOD.RestResource->VertexColorsBuffer,	0, 0,									FHairCardsMaterialFormat::SizeInByte,												FHairCardsMaterialFormat::VertexElementType,	EVertexStreamUsage::Default), 5));
+				Elements.Add(AccessStreamComponent(FVertexStreamComponent(&LOD.RestResource->VertexColorsBuffer,	0, 0,									FHairCardsMaterialFormat::SizeInByte,												FHairCardsMaterialFormat::VertexElementType,	EVertexStreamUsage::Default), 6));
 			}
 			else
 			{
-				Elements.Add(AccessStreamComponent(FVertexStreamComponent(&LOD.RestResource->NormalsBuffer,			0, 0,									FHairCardsMaterialFormat::SizeInByte,												FHairCardsMaterialFormat::VertexElementType,	EVertexStreamUsage::Default), 5));
+				Elements.Add(AccessStreamComponent(FVertexStreamComponent(&LOD.RestResource->NormalsBuffer,			0, 0,									FHairCardsMaterialFormat::SizeInByte,												FHairCardsMaterialFormat::VertexElementType,	EVertexStreamUsage::Default), 6));
 			}
 
 			// Ensure the rest resources are in correct states for the VF
@@ -482,11 +482,11 @@ void FHairCardsVertexFactory::InitResources(FRHICommandListBase& RHICmdList)
 			Elements.Add(AccessStreamComponent(FVertexStreamComponent(&LOD.RestResource->NormalsBuffer,				0, 0,									FHairCardsMaterialFormat::SizeInByte,												FHairCardsMaterialFormat::VertexElementType,	EVertexStreamUsage::Default), 4)); // Dummy
 			if (bHasVertexColor)
 			{
-				Elements.Add(AccessStreamComponent(FVertexStreamComponent(&LOD.RestResource->VertexColorsBuffer,	0, 0,									FHairCardsMaterialFormat::SizeInByte,												FHairCardsMaterialFormat::VertexElementType,	EVertexStreamUsage::Default), 5));
+				Elements.Add(AccessStreamComponent(FVertexStreamComponent(&LOD.RestResource->VertexColorsBuffer,	0, 0,									FHairCardsMaterialFormat::SizeInByte,												FHairCardsMaterialFormat::VertexElementType,	EVertexStreamUsage::Default), 6));
 			}
 			else
 			{
-				Elements.Add(AccessStreamComponent(FVertexStreamComponent(&LOD.RestResource->NormalsBuffer,			0, 0,									FHairCardsMaterialFormat::SizeInByte,												FHairCardsMaterialFormat::VertexElementType,	EVertexStreamUsage::Default), 5));
+				Elements.Add(AccessStreamComponent(FVertexStreamComponent(&LOD.RestResource->NormalsBuffer,			0, 0,									FHairCardsMaterialFormat::SizeInByte,												FHairCardsMaterialFormat::VertexElementType,	EVertexStreamUsage::Default), 6));
 			}
 
 			// Ensure the rest resources are in correct states for the VF
diff --git a/Engine/Shaders/Private/HairStrands/HairCardsVertexFactory.ush b/Engine/Shaders/Private/HairStrands/HairCardsVertexFactory.ush
--- a/Engine/Shaders/Private/HairStrands/HairCardsVertexFactory.ush
+++ b/Engine/Shaders/Private/HairStrands/HairCardsVertexFactory.ush
@@ -50,7 +50,7 @@ struct FVertexFactoryInput
   #endif
 	float4 UVs		 : ATTRIBUTE3;
 	float4 Materials : ATTRIBUTE4;
-	float4 VertexColors : ATTRIBUTE5;
+	float4 VertexColors : ATTRIBUTE6;
 #endif
 
 	VF_GPUSCENE_DECLARE_INPUT_BLOCK(13)
