[Backout] - CL38920449
#rnx
Original CL Desc
-----------------------------------------------------------------
6-Groom deformers extension

#virtualized

diff --git a/Engine/Plugins/Animation/DeformerGraph/Source/OptimusCore/Private/DataInterfaces/OptimusDataInterfaceRawBuffer.cpp b/Engine/Plugins/Animation/DeformerGraph/Source/OptimusCore/Private/DataInterfaces/OptimusDataInterfaceRawBuffer.cpp
--- a/Engine/Plugins/Animation/DeformerGraph/Source/OptimusCore/Private/DataInterfaces/OptimusDataInterfaceRawBuffer.cpp
+++ b/Engine/Plugins/Animation/DeformerGraph/Source/OptimusCore/Private/DataInterfaces/OptimusDataInterfaceRawBuffer.cpp
@@ -679,9 +679,6 @@ void FOptimusPersistentBufferDataProviderProxy::GatherDispatchData(FDispatchData
 		Parameters.BufferSize = InDispatchData.bUnifiedDispatch ? TotalElementCount : InvocationElementCounts[InvocationIndex];
 		Parameters.BufferUAV = BufferUAVs[InDispatchData.GraphKernelIndex];
 
-		if(!InDispatchData.bUnifiedDispatch)
-		{
-			StartOffset += InvocationElementCounts[InvocationIndex];
-		}
+		StartOffset += InvocationElementCounts[InvocationIndex];
 	}
 }
diff --git a/Engine/Plugins/Experimental/Dataflow/Source/DataflowEditor/Private/Dataflow/DataflowSimulationScene.cpp b/Engine/Plugins/Experimental/Dataflow/Source/DataflowEditor/Private/Dataflow/DataflowSimulationScene.cpp
--- a/Engine/Plugins/Experimental/Dataflow/Source/DataflowEditor/Private/Dataflow/DataflowSimulationScene.cpp
+++ b/Engine/Plugins/Experimental/Dataflow/Source/DataflowEditor/Private/Dataflow/DataflowSimulationScene.cpp
@@ -253,7 +253,7 @@ void FDataflowSimulationScene::UpdateSimulationCache(float& DeltaSeconds)
 			DeltaSeconds = 1.0f / SceneDescription->CacheParams.FrameRate;
 			SimulationTime += DeltaSeconds;
 
-			if(SimulationTime >= TimeRange[1]+UE_SMALL_NUMBER)
+			if(SimulationTime >= TimeRange[1])
 			{
 				bIsRecordingCache = false;
 				CacheManager->EndEvaluate();
diff --git a/Engine/Plugins/Experimental/Dataflow/Source/DataflowNodes/Private/Dataflow/DataflowCollectionAddScalarVertexPropertyNode.cpp b/Engine/Plugins/Experimental/Dataflow/Source/DataflowNodes/Private/Dataflow/DataflowCollectionAddScalarVertexPropertyNode.cpp
--- a/Engine/Plugins/Experimental/Dataflow/Source/DataflowNodes/Private/Dataflow/DataflowCollectionAddScalarVertexPropertyNode.cpp
+++ b/Engine/Plugins/Experimental/Dataflow/Source/DataflowNodes/Private/Dataflow/DataflowCollectionAddScalarVertexPropertyNode.cpp
@@ -57,35 +57,6 @@ TArray<UE::Dataflow::FRenderingParameter> FDataflowAddScalarVertexPropertyCallba
 	return UniqueParameters;
 }
 
-TArray<UE::Dataflow::FRenderingParameter> FDataflowAddScalarVertexPropertyCallbackRegistry::GetRenderingParameters(const FName& TargetGroup) const
-{
-	TArray<UE::Dataflow::FRenderingParameter> UniqueParameters;
-
-	for (const TPair<FName, TUniquePtr<IDataflowAddScalarVertexPropertyCallbacks>>& CallbacksEntry : AllCallbacks)
-	{
-		const TArray<UE::Dataflow::FRenderingParameter> RenderingParameters = CallbacksEntry.Value->GetRenderingParameters();
-		const TArray<FName> TargetGroups = CallbacksEntry.Value->GetTargetGroupNames();
-		if(RenderingParameters.Num() == TargetGroups.Num())
-		{
-			for(int32 TargetIndex = 0, NumTargets = TargetGroups.Num(); TargetIndex < NumTargets; ++TargetIndex)
-			{
-				if(TargetGroups[TargetIndex] == TargetGroup)
-				{
-					UniqueParameters.AddUnique(RenderingParameters[TargetIndex]);
-				}
-			}
-		}
-		else
-		{
-			for (const UE::Dataflow::FRenderingParameter& RenderingParameter : CallbacksEntry.Value->GetRenderingParameters())
-			{
-				UniqueParameters.AddUnique(RenderingParameter);
-			}
-		}
-	}
-	return UniqueParameters;
-}
-
 
 //
 // Deprecated class
@@ -152,7 +123,7 @@ FDataflowCollectionAddScalarVertexPropertyNode::FDataflowCollectionAddScalarVert
 
 TArray<UE::Dataflow::FRenderingParameter> FDataflowCollectionAddScalarVertexPropertyNode::GetRenderParametersImpl() const
 {
-	return FDataflowAddScalarVertexPropertyCallbackRegistry::Get().GetRenderingParameters(TargetGroup.Name);
+	return FDataflowAddScalarVertexPropertyCallbackRegistry::Get().GetRenderingParameters();
 }
 
 
diff --git a/Engine/Plugins/Experimental/Dataflow/Source/DataflowNodes/Public/Dataflow/DataflowCollectionAddScalarVertexPropertyNode.h b/Engine/Plugins/Experimental/Dataflow/Source/DataflowNodes/Public/Dataflow/DataflowCollectionAddScalarVertexPropertyNode.h
--- a/Engine/Plugins/Experimental/Dataflow/Source/DataflowNodes/Public/Dataflow/DataflowCollectionAddScalarVertexPropertyNode.h
+++ b/Engine/Plugins/Experimental/Dataflow/Source/DataflowNodes/Public/Dataflow/DataflowCollectionAddScalarVertexPropertyNode.h
@@ -32,8 +32,6 @@ public:
 	DATAFLOWNODES_API TArray<FName> GetTargetGroupNames() const;
 
 	DATAFLOWNODES_API TArray<UE::Dataflow::FRenderingParameter> GetRenderingParameters() const;
-	
-	DATAFLOWNODES_API TArray<UE::Dataflow::FRenderingParameter> GetRenderingParameters(const FName& TargetGroup) const;
 
 private:
 
diff --git a/Engine/Plugins/Runtime/HairStrands/Config/BaseHairStrands.ini b/Engine/Plugins/Runtime/HairStrands/Config/BaseHairStrands.ini
--- a/Engine/Plugins/Runtime/HairStrands/Config/BaseHairStrands.ini
+++ b/Engine/Plugins/Runtime/HairStrands/Config/BaseHairStrands.ini
@@ -19,8 +19,4 @@ fx.UseShaderStages=1
 +ClassRedirects=(OldName="/Script/HairStrandsNiagara.NiagaraDataInterfacePressureGrid",NewName="/Script/HairStrandsCore.NiagaraDataInterfacePressureGrid")
 +PropertyRedirects=(OldName="GroomAsset.DeformedSkeletalMesh",NewName="GroomAsset.RiggedSkeletalMesh")
 +PropertyRedirects=(OldName="HairDeformationSettings.bEnableDeformation",NewName="HairDeformationSettings.bEnableRigging")
-+PropertyRedirects=(OldName="HairGroupsInterpolation.DeformationSettings",NewName="HairGroupsInterpolation.RiggingSettings")
-+ClassRedirects=(OldName="OptimusGroomWriteDataInterface",NewName="/Script/HairStrandsDeformer.OptimusGroomStrandsWriteDataInterface")
-+ClassRedirects=(OldName="OptimusGroomDataInterface",NewName="/Script/HairStrandsDeformer.OptimusGroomStrandsReadDataInterface")
-+ClassRedirects=(OldName="OptimusGroomGuideDataInterface",NewName="/Script/HairStrandsDeformer.OptimusGroomGuidesReadDataInterface")
-+ClassRedirects=(OldName="OptimusGroomComponentSource",NewName="/Script/HairStrandsDeformer.OptimusGroomAssetComponentSource")
\ No newline at end of file
++PropertyRedirects=(OldName="HairGroupsInterpolation.DeformationSettings",NewName="HairGroupsInterpolation.RiggingSettings")
\ No newline at end of file
diff --git a/Engine/Plugins/Runtime/HairStrands/HairStrands.uplugin b/Engine/Plugins/Runtime/HairStrands/HairStrands.uplugin
--- a/Engine/Plugins/Runtime/HairStrands/HairStrands.uplugin
+++ b/Engine/Plugins/Runtime/HairStrands/HairStrands.uplugin
@@ -71,10 +71,6 @@
 		{
 			"Name": "Dataflow",
 			"Enabled": true
-		},
-		{
-			"Name": "ComputeFramework",
-			"Enabled": true
 		}
 	]
 }
diff --git a/Engine/Plugins/Runtime/HairStrands/Shaders/Private/DeformerDataInterfaceGroom.ush b/Engine/Plugins/Runtime/HairStrands/Shaders/Private/DeformerDataInterfaceGroom.ush
new file mode 100644
--- /dev/null
+++ b/Engine/Plugins/Runtime/HairStrands/Shaders/Private/DeformerDataInterfaceGroom.ush
@@ -0,0 +1,134 @@
+// Copyright Epic Games, Inc. All Rights Reserved.
+
+#include "/Plugin/ComputeFramework/Private/ComputeKernelCommon.ush"
+#include "/Engine/Private/HairStrands/HairStrandsVertexFactoryCommon.ush"
+#include "/Engine/Private/HairStrands/HairStrandsAttributeCommon.ush"
+
+///////////////////////////////////////////////////////////////////////////////////////////////////
+// Bindings
+
+HAIR_STRANDS_INSTANCE_PARAMETERS({DataInterfaceName})
+
+#define HAIR_STRANDS_ATTRIBUTE_ACCESSORS(Name) {DataInterfaceName}_##Name
+#include "/Engine/Private/HairStrands/HairStrandsAttributeTemplate.ush"
+
+ByteAddressBuffer {DataInterfaceName}_InterpolationBuffer;
+
+///////////////////////////////////////////////////////////////////////////////////////////////////
+// Helpers
+
+FHairControlPoint ReadControlPointData(uint ControlPointIndex)
+{
+	const uint MaxControlPoints = {DataInterfaceName}_PointCount;
+	if (ControlPointIndex >= MaxControlPoints)
+	{
+		return (FHairControlPoint)0;
+	}
+
+	// Hair deformer work in local space so no need for position offset
+	return ReadHairControlPoint(
+		{DataInterfaceName}_PositionBuffer,
+		ControlPointIndex, 
+		{DataInterfaceName}_GetHairInstancePositionOffset(),
+		{DataInterfaceName}_Radius, 
+		{DataInterfaceName}_RootScale, 
+		{DataInterfaceName}_TipScale);
+}
+
+FHairCurve ReadCurveData(uint CurveIndex)
+{
+	const uint MaxCurve = {DataInterfaceName}_CurveCount;
+	if (CurveIndex >= MaxCurve)
+	{
+		return (FHairCurve)0;
+	}
+
+	return ReadHairCurve({DataInterfaceName}_CurveBuffer, CurveIndex);
+}
+
+///////////////////////////////////////////////////////////////////////////////////////////////////
+// Interface
+
+uint ReadNumControlPoints_{DataInterfaceName}()
+{
+	return {DataInterfaceName}_PointCount;
+}
+
+uint ReadNumCurves_{DataInterfaceName}()
+{
+	return {DataInterfaceName}_CurveCount;
+}
+
+float3 ReadPosition_{DataInterfaceName}(uint ControlPointIndex)
+{
+	return ReadControlPointData(ControlPointIndex).Position;
+}
+
+float ReadRadius_{DataInterfaceName}(uint ControlPointIndex)
+{
+	return ReadControlPointData(ControlPointIndex).WorldRadius;
+}
+
+float ReadCoordU_{DataInterfaceName}(uint ControlPointIndex)
+{
+	return ReadControlPointData(ControlPointIndex).UCoord;
+}
+
+float ReadLength_{DataInterfaceName}(uint ControlPointIndex)
+{
+	return GetHairStrandsDimensions(ControlPointIndex, 0).x;
+}
+
+float2 ReadRootUV_{DataInterfaceName}(uint ControlPointIndex)
+{
+	return GetHairStrandsRootUV(ControlPointIndex);
+}
+
+float ReadSeed_{DataInterfaceName}(uint ControlPointIndex)
+{
+	return GetHairStrandsSeed(ControlPointIndex);
+}
+
+uint ReadClumpId_{DataInterfaceName}(uint ControlPointIndex)
+{
+	return GetHairStrandsClumpID(ControlPointIndex).x;
+}
+
+float3 ReadColor_{DataInterfaceName}(uint ControlPointIndex)
+{
+	return GetHairStrandsColor(ControlPointIndex, 0);
+}
+
+float ReadRoughness_{DataInterfaceName}(uint ControlPointIndex)
+{
+	return GetHairStrandsRoughness(ControlPointIndex, 0);
+}
+
+float ReadAO_{DataInterfaceName}(uint ControlPointIndex)
+{
+	return GetHairStrandsAO(ControlPointIndex, 0);
+}
+
+float ReadCurveOffsetPoint_{DataInterfaceName}(uint CurveIndex)
+{
+	return ReadCurveData(CurveIndex).PointIndex;
+}
+
+float ReadCurveNumPoint_{DataInterfaceName}(uint CurveIndex)
+{
+	return ReadCurveData(CurveIndex).PointCount;
+}
+
+uint ReadGuideIndex_{DataInterfaceName}(uint ControlPointIndex)
+{
+	// Find the rendering curve index
+	const uint CurveIndex = ReadHairPointToCurveIndex({DataInterfaceName}_PointToCurveBuffer, ControlPointIndex);
+
+	// Fint the curve and point interpolation data
+	const FGuideCurveData CurveData = UnpackGuideCurveData({DataInterfaceName}_CurveInterpolationBuffer, CurveIndex, {DataInterfaceName}_bSingleGuide);
+	const FGuidePointData PointData = UnpackGuidePointData({DataInterfaceName}_PointInterpolationBuffer, ControlPointIndex, {DataInterfaceName}_bSingleGuide);
+
+	// Build the sim point index
+	const uint2 GuidePointIndices = CurveData.RootPointIndices + PointData.PointIndices;
+	return GuidePointIndices.x;
+}
diff --git a/Engine/Plugins/Runtime/HairStrands/Shaders/Private/DeformerDataInterfaceGroomExec.ush b/Engine/Plugins/Runtime/HairStrands/Shaders/Private/DeformerDataInterfaceGroomExec.ush
new file mode 100644
--- /dev/null
+++ b/Engine/Plugins/Runtime/HairStrands/Shaders/Private/DeformerDataInterfaceGroomExec.ush
@@ -0,0 +1,10 @@
+// Copyright Epic Games, Inc. All Rights Reserved.
+
+#include "/Plugin/ComputeFramework/Private/ComputeKernelCommon.ush"
+
+int3 {DataInterfaceName}_NumThreads;
+
+int3 ReadNumThreads_{DataInterfaceName}()
+{
+	return {DataInterfaceName}_NumThreads;
+}
diff --git a/Engine/Plugins/Runtime/HairStrands/Shaders/Private/DeformerDataInterfaceGroomGuide.ush b/Engine/Plugins/Runtime/HairStrands/Shaders/Private/DeformerDataInterfaceGroomGuide.ush
new file mode 100644
--- /dev/null
+++ b/Engine/Plugins/Runtime/HairStrands/Shaders/Private/DeformerDataInterfaceGroomGuide.ush
@@ -0,0 +1,70 @@
+// Copyright Epic Games, Inc. All Rights Reserved.
+
+#include "/Plugin/ComputeFramework/Private/ComputeKernelCommon.ush"
+#include "/Engine/Private/HairStrands/HairStrandsVertexFactoryCommon.ush"
+#include "/Engine/Private/HairStrands/HairStrandsAttributeCommon.ush"
+
+///////////////////////////////////////////////////////////////////////////////////////////////////
+// Bindings
+
+HAIR_STRANDS_INSTANCE_PARAMETERS({DataInterfaceName})
+
+///////////////////////////////////////////////////////////////////////////////////////////////////
+// Helpers
+
+FHairControlPoint ReadControlPointData_{DataInterfaceName}(uint ControlPointIndex)
+{
+	const uint MaxControlPoints = {DataInterfaceName}_PointCount;
+	if (ControlPointIndex >= MaxControlPoints)
+	{
+		return (FHairControlPoint)0;
+	}
+
+	// Hair deformer work in local space so no need for position offset
+	return ReadHairControlPoint(
+		{DataInterfaceName}_PositionBuffer,
+		ControlPointIndex, 
+		{DataInterfaceName}_GetHairInstancePositionOffset(),
+		{DataInterfaceName}_Radius, 
+		{DataInterfaceName}_RootScale, 
+		{DataInterfaceName}_TipScale);
+}
+
+FHairCurve ReadCurveData_{DataInterfaceName}(uint CurveIndex)
+{
+	const uint MaxCurve = {DataInterfaceName}_CurveCount;
+	if (CurveIndex >= MaxCurve)
+	{
+		return (FHairCurve)0;
+	}
+
+	return ReadHairCurve({DataInterfaceName}_CurveBuffer, CurveIndex);
+}
+
+///////////////////////////////////////////////////////////////////////////////////////////////////
+// Interface
+
+uint ReadNumGuidePoints_{DataInterfaceName}()
+{
+	return {DataInterfaceName}_PointCount;
+}
+
+uint ReadNumGuideCurves_{DataInterfaceName}()
+{
+	return {DataInterfaceName}_CurveCount;
+}
+
+float3 ReadPosition_{DataInterfaceName}(uint ControlPointIndex)
+{
+	return ReadControlPointData_{DataInterfaceName}(ControlPointIndex).Position;
+}
+
+float ReadCurveOffsetPoint_{DataInterfaceName}(uint CurveIndex)
+{
+	return ReadCurveData_{DataInterfaceName}(CurveIndex).PointIndex;
+}
+
+float ReadCurveNumPoint_{DataInterfaceName}(uint CurveIndex)
+{
+	return ReadCurveData_{DataInterfaceName}(CurveIndex).PointCount;
+}
\ No newline at end of file
diff --git a/Engine/Plugins/Runtime/HairStrands/Shaders/Private/DeformerDataInterfaceGroomWrite.ush b/Engine/Plugins/Runtime/HairStrands/Shaders/Private/DeformerDataInterfaceGroomWrite.ush
new file mode 100644
--- /dev/null
+++ b/Engine/Plugins/Runtime/HairStrands/Shaders/Private/DeformerDataInterfaceGroomWrite.ush
@@ -0,0 +1,155 @@
+// Copyright Epic Games, Inc. All Rights Reserved.
+
+#include "/Plugin/ComputeFramework/Private/ComputeKernelCommon.ush"
+#include "/Engine/Private/HairStrands/HairStrandsVertexFactoryCommon.ush"
+
+HAIR_STRANDS_INSTANCE_PARAMETERS({DataInterfaceName})
+
+// Note: resources are explicitely bound here (instance of using HAIR_STRANDS_INSTANCE_PARAMETERS's resources) 
+// beause we write out data which would be bound as SRV otherwise (position, attributes, ...)
+uint 				{DataInterfaceName}_OutputStreamStart;
+StructuredBuffer<float4> {DataInterfaceName}_PositionOffsetBufferSRV;
+ByteAddressBuffer		{DataInterfaceName}_PositionBufferSRV;
+RWByteAddressBuffer		{DataInterfaceName}_PositionBufferUAV;
+RWByteAddressBuffer 	{DataInterfaceName}_PointAttributeBufferUAV;
+RWByteAddressBuffer 	{DataInterfaceName}_CurveAttributeBufferUAV;
+
+float3 {DataInterfaceName}_GetPositionOffsetBufferSRV() 
+{ 
+	return ReadRenPositionOffset({DataInterfaceName}_PositionOffsetBufferSRV,{DataInterfaceName}_RegisteredIndex); 
+}
+
+uint ReadNumControlPoints_{DataInterfaceName}()
+{
+	return {DataInterfaceName}_PointCount;
+}
+
+uint ReadNumCurves_{DataInterfaceName}()
+{
+	return {DataInterfaceName}_CurveCount;
+}
+
+void WritePosition_{DataInterfaceName}(uint ControlPointIndex, float3 Position)
+{
+	const uint WriteIndex = {DataInterfaceName}_OutputStreamStart + ControlPointIndex;
+	if (WriteIndex < {DataInterfaceName}_PointCount)
+	{
+		const FPackedHairPosition PackedCP = ReadPackedHairPosition({DataInterfaceName}_PositionBufferSRV, WriteIndex);
+		WritePackedHairControlPointPosition({DataInterfaceName}_PositionBufferUAV, WriteIndex, PackedCP, Position, {DataInterfaceName}_GetPositionOffsetBufferSRV());
+	}
+}
+
+void WriteRadius_{DataInterfaceName}(uint ControlPointIndex, float Radius)
+{
+	const uint WriteIndex = {DataInterfaceName}_OutputStreamStart + ControlPointIndex;
+	if (WriteIndex < {DataInterfaceName}_PointCount)
+	{
+		const FPackedHairPosition PackedCP = ReadPackedHairPosition({DataInterfaceName}_PositionBufferSRV, WriteIndex);
+		const FPackedHairPosition NewPackedCP = PackHairControlPointRadius(PackedCP, Radius, {DataInterfaceName}_Radius);
+		WritePackedHairPosition({DataInterfaceName}_PositionBufferUAV, WriteIndex, NewPackedCP);
+	}
+}
+
+void WritePositionAndRadius_{DataInterfaceName}(uint ControlPointIndex, float4 PositionAndRadius)
+{
+	const uint WriteIndex = {DataInterfaceName}_OutputStreamStart + ControlPointIndex;
+	if (WriteIndex < {DataInterfaceName}_PointCount)
+	{
+		const FPackedHairPosition PackedCP = ReadPackedHairPosition({DataInterfaceName}_PositionBufferSRV, WriteIndex);
+		FPackedHairPosition NewPackedCP = 0;
+		NewPackedCP = PackHairControlPointPosition(PackedCP, PositionAndRadius.xyz, {DataInterfaceName}_GetPositionOffsetBufferSRV());
+		NewPackedCP = PackHairControlPointRadius(NewPackedCP, PositionAndRadius.w, {DataInterfaceName}_Radius);
+		WritePackedHairPosition({DataInterfaceName}_PositionBufferUAV, WriteIndex, NewPackedCP);
+	}
+}
+
+void InternalSetHairStrandsAttributes(uint Index, uint AttributeOffsetInBytes, uint AttributeStrideInBytes, RWByteAddressBuffer InAttributeBuffer, uint In, uint InIndexToChunkDivAsShift, uint InChunkElementCount, uint InChunkStrideInBytes)
+{
+	// Adapt index and offset
+	const uint ChunkIndex   = Index >> InIndexToChunkDivAsShift;
+	Index                  -= ChunkIndex * InChunkElementCount;
+	AttributeOffsetInBytes += ChunkIndex * InChunkStrideInBytes;
+
+	// Ensure the reading address is 4-byte aligned
+	const uint Address 					= AttributeOffsetInBytes + Index * AttributeStrideInBytes;
+	const uint Address_4BytesAligned 	= Address & (~0x3);
+	const uint ElementIndex 			= Address - Address_4BytesAligned;
+
+	// Create a mask where the 
+	const uint Mask 		  = (1u<<(AttributeStrideInBytes*8u))-1u;
+	const uint AttributeMask  = Mask << (ElementIndex*8);
+	const uint AttributeValue = In   << (ElementIndex*8);
+
+	// Clear previous value, and then write the value
+	InAttributeBuffer.InterlockedAnd(Address_4BytesAligned, ~AttributeMask);
+	InAttributeBuffer.InterlockedOr(Address_4BytesAligned, AttributeValue);
+}
+
+void InternalSetHairStrandsAttributes32(uint Index, uint AttributeOffsetInBytes, uint AttributeStrideInBytes, RWByteAddressBuffer InAttributeBuffer, uint In, uint InIndexToChunkDivAsShift, uint InChunkElementCount, uint InChunkStrideInBytes)
+{
+	// Adapt index and offset
+	const uint ChunkIndex   = Index >> InIndexToChunkDivAsShift;
+	Index                  -= ChunkIndex * InChunkElementCount;
+	AttributeOffsetInBytes += ChunkIndex * InChunkStrideInBytes;
+
+	// Ensure the reading address is 4-byte aligned
+	const uint Address = AttributeOffsetInBytes + Index * AttributeStrideInBytes;
+	const uint Address_4BytesAligned = Address & (~0x3);
+
+	// Clear previous value, and then write the value
+	InAttributeBuffer.Store(Address_4BytesAligned, In);
+}
+
+void WriteRootUV_{DataInterfaceName}(uint CurveIndex, float2 RootUV)
+{
+	const uint WriteIndex = {DataInterfaceName}_OutputStreamStart + CurveIndex;
+	if (WriteIndex < {DataInterfaceName}_CurveCount && HasHairStrandsRootUV())
+	{
+		InternalSetHairStrandsAttributes32(CurveIndex, HAIR_CURVE_ATTRIBUTE_OFFSET_ROOTUV({DataInterfaceName}_CurveAttributeOffsets), HAIR_CURVE_ATTRIBUTE_STRIDE_ROOTUV, {DataInterfaceName}_CurveAttributeBufferUAV, PackHairRootUV(RootUV), {DataInterfaceName}_CurveAttributeIndexToChunkDivAsShift, {DataInterfaceName}_CurveAttributeChunkElementCount), {DataInterfaceName}_CurveAttributeChunkStrideInBytes);
+	}
+}
+
+void WriteSeed_{DataInterfaceName}(uint CurveIndex, float Seed)
+{
+	const uint WriteIndex = {DataInterfaceName}_OutputStreamStart + CurveIndex;
+	if (WriteIndex < {DataInterfaceName}_CurveCount)
+	{
+		InternalSetHairStrandsAttributes(CurveIndex, HAIR_CURVE_ATTRIBUTE_OFFSET_SEED({DataInterfaceName}_CurveAttributeOffsets), HAIR_CURVE_ATTRIBUTE_STRIDE_SEED, {DataInterfaceName}_CurveAttributeBufferUAV, PackHairSeed(Seed), {DataInterfaceName}_CurveAttributeIndexToChunkDivAsShift, {DataInterfaceName}_CurveAttributeChunkElementCount), {DataInterfaceName}_CurveAttributeChunkStrideInBytes);
+	}
+}
+
+void WriteClumpId_{DataInterfaceName}(uint CurveIndex, uint ClumpId)
+{
+	const uint WriteIndex = {DataInterfaceName}_OutputStreamStart + CurveIndex;
+	if (WriteIndex < {DataInterfaceName}_CurveCount && HasHairStrandsClumpId())
+	{
+		InternalSetHairStrandsAttributes(CurveIndex, HAIR_CURVE_ATTRIBUTE_OFFSET_CLUMPID({DataInterfaceName}_CurveAttributeOffsets), HAIR_CURVE_ATTRIBUTE_STRIDE_CLUMPID, {DataInterfaceName}_CurveAttributeBufferUAV, PackHairClumpID(ClumpId), {DataInterfaceName}_CurveAttributeIndexToChunkDivAsShift, {DataInterfaceName}_CurveAttributeChunkElementCount), {DataInterfaceName}_CurveAttributeChunkStrideInBytes);
+	}
+}
+
+void WriteColor_{DataInterfaceName}(uint PointIndex, float3 Color)
+{
+	const uint WriteIndex = {DataInterfaceName}_OutputStreamStart + PointIndex;
+	if (WriteIndex < {DataInterfaceName}_PointCount && HasHairStrandsColor())
+	{
+		InternalSetHairStrandsAttributes32(PointIndex, HAIR_POINT_ATTRIBUTE_OFFSET_COLOR({DataInterfaceName}_PointAttributeOffsets), HAIR_POINT_ATTRIBUTE_STRIDE_COLOR, {DataInterfaceName}_PointAttributeBufferUAV, PackHairColor(Color), {DataInterfaceName}_PointAttributeIndexToChunkDivAsShift, {DataInterfaceName}_PointAttributeChunkElementCount), {DataInterfaceName}_PointAttributeChunkStrideInBytes);
+	}
+}
+
+void WriteRoughness_{DataInterfaceName}(uint PointIndex, float Roughness)
+{
+	const uint WriteIndex = {DataInterfaceName}_OutputStreamStart + PointIndex;
+	if (WriteIndex < {DataInterfaceName}_PointCount && HasHairStrandsRoughness())
+	{
+		InternalSetHairStrandsAttributes(PointIndex, HAIR_POINT_ATTRIBUTE_OFFSET_ROUGHNESS({DataInterfaceName}_PointAttributeOffsets), HAIR_POINT_ATTRIBUTE_STRIDE_ROUGHNESS, {DataInterfaceName}_PointAttributeBufferUAV, PackHairRoughness(Roughness), {DataInterfaceName}_PointAttributeIndexToChunkDivAsShift, {DataInterfaceName}_PointAttributeChunkElementCount), {DataInterfaceName}_PointAttributeChunkStrideInBytes);
+	}
+}
+
+void WriteAO_{DataInterfaceName}(uint PointIndex, float Ao)
+{
+	const uint WriteIndex = {DataInterfaceName}_OutputStreamStart + PointIndex;
+	if (WriteIndex < {DataInterfaceName}_PointCount && HasHairStrandsAO())
+	{
+		InternalSetHairStrandsAttributes(PointIndex, HAIR_POINT_ATTRIBUTE_OFFSET_AO({DataInterfaceName}_PointAttributeOffsets), HAIR_POINT_ATTRIBUTE_STRIDE_AO, {DataInterfaceName}_PointAttributeBufferUAV, PackHairAO(Ao), {DataInterfaceName}_PointAttributeIndexToChunkDivAsShift, {DataInterfaceName}_PointAttributeChunkElementCount), {DataInterfaceName}_PointAttributeChunkStrideInBytes);
+	}
+}
\ No newline at end of file
diff --git a/Engine/Plugins/Runtime/HairStrands/Shaders/Private/Deformers/DeformerGroomAttributeRead.ush b/Engine/Plugins/Runtime/HairStrands/Shaders/Private/Deformers/DeformerGroomAttributeRead.ush
deleted file mode 100644
diff --git a/Engine/Plugins/Runtime/HairStrands/Shaders/Private/Deformers/DeformerGroomDomainsExec.ush b/Engine/Plugins/Runtime/HairStrands/Shaders/Private/Deformers/DeformerGroomDomainsExec.ush
deleted file mode 100644
diff --git a/Engine/Plugins/Runtime/HairStrands/Shaders/Private/Deformers/DeformerGroomGuidesRead.ush b/Engine/Plugins/Runtime/HairStrands/Shaders/Private/Deformers/DeformerGroomGuidesRead.ush
deleted file mode 100644
diff --git a/Engine/Plugins/Runtime/HairStrands/Shaders/Private/Deformers/DeformerGroomGuidesWrite.ush b/Engine/Plugins/Runtime/HairStrands/Shaders/Private/Deformers/DeformerGroomGuidesWrite.ush
deleted file mode 100644
diff --git a/Engine/Plugins/Runtime/HairStrands/Shaders/Private/Deformers/DeformerGroomMeshesRead.ush b/Engine/Plugins/Runtime/HairStrands/Shaders/Private/Deformers/DeformerGroomMeshesRead.ush
deleted file mode 100644
diff --git a/Engine/Plugins/Runtime/HairStrands/Shaders/Private/Deformers/DeformerGroomSolverRead.ush b/Engine/Plugins/Runtime/HairStrands/Shaders/Private/Deformers/DeformerGroomSolverRead.ush
deleted file mode 100644
diff --git a/Engine/Plugins/Runtime/HairStrands/Shaders/Private/Deformers/DeformerGroomStrandsRead.ush b/Engine/Plugins/Runtime/HairStrands/Shaders/Private/Deformers/DeformerGroomStrandsRead.ush
deleted file mode 100644
diff --git a/Engine/Plugins/Runtime/HairStrands/Shaders/Private/Deformers/DeformerGroomStrandsWrite.ush b/Engine/Plugins/Runtime/HairStrands/Shaders/Private/Deformers/DeformerGroomStrandsWrite.ush
deleted file mode 100644
diff --git a/Engine/Plugins/Runtime/HairStrands/Source/HairStrandsCore/Private/GroomBuilder.cpp b/Engine/Plugins/Runtime/HairStrands/Source/HairStrandsCore/Private/GroomBuilder.cpp
--- a/Engine/Plugins/Runtime/HairStrands/Source/HairStrandsCore/Private/GroomBuilder.cpp
+++ b/Engine/Plugins/Runtime/HairStrands/Source/HairStrandsCore/Private/GroomBuilder.cpp
@@ -2153,7 +2153,7 @@ bool FGroomBuilder::BuildHairDescriptionGroups(const FHairDescription& HairDescr
 			continue;
 		}
 
-		const bool bAddEndingControlPoint = NeedsToAddEndingControlPoint(CurveNumVertices) && bAllowAddEndControlPoint && !bIsGuide;
+		const bool bAddEndingControlPoint = NeedsToAddEndingControlPoint(CurveNumVertices) && bAllowAddEndControlPoint;
 		CurrentHairStrandsDatas->StrandsCurves.CurvesCount.Add(FMath::Min(uint32(CurveNumVertices + (bAddEndingControlPoint ? 1u : 0u)), HAIR_MAX_NUM_POINT_PER_CURVE));
 
 		if (bCanUseClosestGuidesAndWeights)
diff --git a/Engine/Plugins/Runtime/HairStrands/Source/HairStrandsCore/Private/GroomComponent.cpp b/Engine/Plugins/Runtime/HairStrands/Source/HairStrandsCore/Private/GroomComponent.cpp
--- a/Engine/Plugins/Runtime/HairStrands/Source/HairStrandsCore/Private/GroomComponent.cpp
+++ b/Engine/Plugins/Runtime/HairStrands/Source/HairStrandsCore/Private/GroomComponent.cpp
@@ -2848,7 +2848,6 @@ void UGroomComponent::InitResources(bool bIsBindingReloading)
 		HairGroupInstance->BindingType = EHairBindingType::NoneBinding;
 
 		HairGroupInstance->Debug.SkinningCurrentLocalToWorld	= RegisteredMeshComponent ? RegisteredMeshComponent->GetComponentTransform() : FTransform::Identity;
-		HairGroupInstance->Debug.SkinningRestLocalToWorld		= HairGroupInstance->Debug.SkinningCurrentLocalToWorld;
 		HairGroupInstance->Debug.SkinningPreviousLocalToWorld	= HairGroupInstance->Debug.SkinningCurrentLocalToWorld;
 		HairGroupInstance->Debug.RigidCurrentLocalToWorld		= GetComponentTransform();
 		HairGroupInstance->Debug.RigidRestLocalToWorld			= HairGroupInstance->Debug.RigidCurrentLocalToWorld;
diff --git a/Engine/Plugins/Runtime/HairStrands/Source/HairStrandsCore/Public/GroomAssetDataflow.h b/Engine/Plugins/Runtime/HairStrands/Source/HairStrandsCore/Public/GroomAssetDataflow.h
--- a/Engine/Plugins/Runtime/HairStrands/Source/HairStrandsCore/Public/GroomAssetDataflow.h
+++ b/Engine/Plugins/Runtime/HairStrands/Source/HairStrandsCore/Public/GroomAssetDataflow.h
@@ -63,10 +63,7 @@ struct HAIRSTRANDSCORE_API FGroomDataflowSettings
 	const FString& GetDataflowTerminal() const { return DataflowTerminal; }
 
 	/** Return the skeletal mesh associated to this groom asset if any. */
-	USkeletalMesh* GetSkeletalMesh(const int32 GroupIndex) const { return !SkeletalMeshes.IsValidIndex(GroupIndex) ? nullptr : SkeletalMeshes[GroupIndex]; }
-
-	/** Return the mesh LOD used to transfer the skinning */
-	int32 GetMeshLOD(const int32 GroupIndex) const { return !MeshLODs.IsValidIndex(GroupIndex) ? INDEX_NONE : MeshLODs[GroupIndex]; }
+	USkeletalMesh* GetSkeletalMesh() const { return SkeletalMesh; }
 	
 	/** Return the groom asset dataflow settings rest collection . */
 	const FManagedArrayCollection* GetRestCollection() const { return RestCollection.Get(); }
@@ -76,19 +73,9 @@ struct HAIRSTRANDSCORE_API FGroomDataflowSettings
 	
 	/** Set the Dataflow terminal associated to this groom asset if any. */
 	void SetDataflowTerminal(const FString& InDataflowTerminal) { DataflowTerminal = InDataflowTerminal; }
-	
-	/** Init the skeletal mesh associated to this groom asset if any. */
-	void InitSkeletalMeshes(const int32 NumGroups) { SkeletalMeshes.Init(nullptr, NumGroups); MeshLODs.Init(INDEX_NONE, NumGroups); }
 
-	/** Set a group skeletal mesh and its associated LOD to this groom asset */
-	void SetSkeletalMesh(const int32 GroupIndex, USkeletalMesh* SkeletalMesh, const int32 MeshLOD)
-	{
-		if(SkeletalMeshes.IsValidIndex(GroupIndex) && MeshLODs.IsValidIndex(GroupIndex))
-		{
-			SkeletalMeshes[GroupIndex] = SkeletalMesh;
-			MeshLODs[GroupIndex] = MeshLOD;
-		}
-	}
+	/** Set the skeletal mesh associated to this groom asset if any. */
+	void SetSkeletalMesh(USkeletalMesh* InSkeletalMesh) { SkeletalMesh = InSkeletalMesh; }
 	
 	/** Set the rest collection onto the groom asset dataflow settings */
 	void SetRestCollection(FManagedArrayCollection* InRestCollection) { RestCollection = TSharedPtr<FManagedArrayCollection, ESPMode::ThreadSafe>(InRestCollection); }
@@ -124,11 +111,7 @@ private :
 
     /** Hair geometry that could describe external cards and meshes. */
     UPROPERTY()
-    TArray<TObjectPtr<USkeletalMesh>> SkeletalMeshes;
-
-	/** LOD indices of the hair geometry used to transfer the skinnign weights */
-	UPROPERTY()
-	TArray<int32> MeshLODs;
+    TObjectPtr<USkeletalMesh> SkeletalMesh;
 
 #if WITH_EDITORONLY_DATA
 	
diff --git a/Engine/Plugins/Runtime/HairStrands/Source/HairStrandsCore/Public/GroomInstance.h b/Engine/Plugins/Runtime/HairStrands/Source/HairStrandsCore/Public/GroomInstance.h
--- a/Engine/Plugins/Runtime/HairStrands/Source/HairStrandsCore/Public/GroomInstance.h
+++ b/Engine/Plugins/Runtime/HairStrands/Source/HairStrandsCore/Public/GroomInstance.h
@@ -265,7 +265,6 @@ struct HAIRSTRANDSCORE_API FHairGroupInstance : public FHairStrandsInstance
 		const UGroomComponent*	GroomComponentForDebug = nullptr; // For debug only, shouldn't be deferred on the rendering thread
 		FTransform				RigidRestLocalToWorld = FTransform::Identity;
 		FTransform				RigidCurrentLocalToWorld = FTransform::Identity;
-		FTransform				SkinningRestLocalToWorld = FTransform::Identity;
 		FTransform				SkinningCurrentLocalToWorld = FTransform::Identity;
 		FTransform				RigidPreviousLocalToWorld = FTransform::Identity;
 		FTransform				SkinningPreviousLocalToWorld = FTransform::Identity;
@@ -300,13 +299,6 @@ struct HAIRSTRANDSCORE_API FHairGroupInstance : public FHairStrandsInstance
 	virtual const FHairGroupPublicData* GetHairData() const { return HairGroupPublicData; }
 	virtual const EHairGeometryType GetHairGeometry() const { return GeometryType; }
 
-	/** Get the rest local to world transform according to the internal binding type */
-	FORCEINLINE const FTransform& GetRestLocalToWorld() const
-	{
-		return (BindingType == EHairBindingType::Skinning) ? Debug.SkinningRestLocalToWorld: 
-															 Debug.RigidRestLocalToWorld;
-	}
-
 	/** Get the current local to world transform according to the internal binding type */
 	FORCEINLINE const FTransform& GetCurrentLocalToWorld() const
 	{
diff --git a/Engine/Plugins/Runtime/HairStrands/Source/HairStrandsDataflow/Private/AttachGuidesRootsNode.cpp b/Engine/Plugins/Runtime/HairStrands/Source/HairStrandsDataflow/Private/AttachGuidesRootsNode.cpp
deleted file mode 100644
diff --git a/Engine/Plugins/Runtime/HairStrands/Source/HairStrandsDataflow/Private/BuildGuidesLODsNode.cpp b/Engine/Plugins/Runtime/HairStrands/Source/HairStrandsDataflow/Private/BuildGuidesLODsNode.cpp
deleted file mode 100644
diff --git a/Engine/Plugins/Runtime/HairStrands/Source/HairStrandsDataflow/Private/BuildGuidesSkinningNode.cpp b/Engine/Plugins/Runtime/HairStrands/Source/HairStrandsDataflow/Private/BuildGuidesSkinningNode.cpp
--- a/Engine/Plugins/Runtime/HairStrands/Source/HairStrandsDataflow/Private/BuildGuidesSkinningNode.cpp
+++ b/Engine/Plugins/Runtime/HairStrands/Source/HairStrandsDataflow/Private/BuildGuidesSkinningNode.cpp
@@ -62,39 +62,32 @@ void FBuildGuidesSkinningDataflowNode::Evaluate(UE::Dataflow::FContext& Context,
 	{
 		FManagedArrayCollection GroomCollection = GetValue<FManagedArrayCollection>(Context, &Collection);
 		
-		UE::Groom::FGroomGuidesFacade GuidesFacade(GroomCollection);
+		UE::Groom::FGroomGuidesFacade CurvesFacade(GroomCollection);
 
-		if(GuidesFacade.IsValid())
+		if(CurvesFacade.IsValid())
 		{
-			TManagedArray<TObjectPtr<UObject>>& ObjectSkeletalMeshes =
-				GroomCollection.AddAttribute<TObjectPtr<UObject>>(UE::Groom::FGroomGuidesFacade::ObjectSkeletalMeshesAttribute, UE::Groom::FGroomGuidesFacade::ObjectsGroup);
-
-			TManagedArray<int32>& ObjectMeshLODs =
-				GroomCollection.AddAttribute<int32>(UE::Groom::FGroomGuidesFacade::ObjectMeshLODsAttribute, UE::Groom::FGroomGuidesFacade::ObjectsGroup);
-			
-			if(GroupIndex == INDEX_NONE)
+			const USkeletalMesh* LocalSkelmesh = SkeletalMesh;
+			if (!LocalSkelmesh)
 			{
-				for(int32 LocalIndex = 0; LocalIndex < ObjectSkeletalMeshes.Num(); ++LocalIndex)
+				if (const UE::Dataflow::FEngineContext* EngineContext = Context.AsType<UE::Dataflow::FEngineContext>())
 				{
-					ObjectSkeletalMeshes[LocalIndex] = SkeletalMesh;
-					ObjectMeshLODs[LocalIndex] = LODIndex;
+					if(UGroomAsset* GroomAsset = Cast<UGroomAsset>(EngineContext->Owner))
+					{
+						LocalSkelmesh = GroomAsset->GetDataflowSettings().GetSkeletalMesh();
+					}
 				}
 			}
-			else if(ObjectSkeletalMeshes.IsValidIndex(GroupIndex) && ObjectMeshLODs.IsValidIndex(GroupIndex))
-			{
-				ObjectSkeletalMeshes[GroupIndex] = SkeletalMesh;
-				ObjectMeshLODs[GroupIndex] = LODIndex;
-			}
-			if(SkeletalMesh && SkeletalMesh->IsValidLODIndex(LODIndex))
+			if(LocalSkelmesh && LocalSkelmesh->IsValidLODIndex(LODIndex))
 			{
 				TMap<FName, FBoneIndexType> TargetBoneToIndex;
-				TargetBoneToIndex.Reserve(SkeletalMesh->GetRefSkeleton().GetRawBoneNum());
-				for (int32 BoneIdx = 0; BoneIdx < SkeletalMesh->GetRefSkeleton().GetRawBoneNum(); ++BoneIdx)
+				TargetBoneToIndex.Reserve(LocalSkelmesh->GetRefSkeleton().GetRawBoneNum());
+				for (int32 BoneIdx = 0; BoneIdx < LocalSkelmesh->GetRefSkeleton().GetRawBoneNum(); ++BoneIdx)
 				{
-					TargetBoneToIndex.Add(SkeletalMesh->GetRefSkeleton().GetRawRefBoneInfo()[BoneIdx].Name, BoneIdx);
+					TargetBoneToIndex.Add(LocalSkelmesh->GetRefSkeleton().GetRawRefBoneInfo()[BoneIdx].Name, BoneIdx);
 				}
+				
 				UE::Geometry::FDynamicMesh3 DynamicMesh;
-				if (UE::Groom::Private::SkeletalMeshToDynamicMesh(SkeletalMesh, LODIndex, DynamicMesh))
+				if (UE::Groom::Private::SkeletalMeshToDynamicMesh(LocalSkelmesh, LODIndex, DynamicMesh))
 				{
 					MeshTransforms::ApplyTransform(DynamicMesh, RelativeTransform, true);
 
@@ -105,30 +98,15 @@ void FBuildGuidesSkinningDataflowNode::Evaluate(UE::Dataflow::FContext& Context,
 
 					if (TransferBoneWeights.Validate() == UE::Geometry::EOperationValidationResult::Ok)
 					{
-						const int32 LocalIndex = GroupIndex;
-						ParallelFor(GuidesFacade.GetNumPoints(), [&TransferBoneWeights, &TargetBoneToIndex, &GuidesFacade, LocalIndex](int32 PointIndex)
+						ParallelFor(CurvesFacade.GetNumPoints(), [&TransferBoneWeights, &TargetBoneToIndex, &CurvesFacade](int32 PointIndex)
 						{
-							const int32 CurveIndex = GuidesFacade.GetPointCurveIndices()[PointIndex];
-							const int32 ObjectIndex = GuidesFacade.GetCurveObjectIndices()[CurveIndex];
-
-							if(LocalIndex == INDEX_NONE || ObjectIndex == LocalIndex)
-							{
-								TArray<int32> BoneIndices;
-								TArray<float> BoneWeights;
-								TransferBoneWeights.TransferWeightsToPoint(BoneIndices, BoneWeights,
-									GuidesFacade.GetPointRestPositions()[PointIndex],&TargetBoneToIndex);
-
-								FIntVector4 PointBoneIndices(INDEX_NONE);
-								FVector4f PointBoneWeights(0.0f);
-
-								for(int32 BoneIdx = 0; BoneIdx < BoneIndices.Num(); ++BoneIdx)
-								{
-									PointBoneIndices[BoneIdx] = BoneIndices[BoneIdx];
-									PointBoneWeights[BoneIdx] = BoneWeights[BoneIdx];
-								}
-								GuidesFacade.SetPointBoneIndices(PointIndex, PointBoneIndices);
-								GuidesFacade.SetPointBoneWeights(PointIndex, PointBoneWeights);
-							}
+							TArray<int32> BoneIndices;
+							TArray<float> BoneWeights;
+							TransferBoneWeights.TransferWeightsToPoint(BoneIndices, BoneWeights,
+								CurvesFacade.GetPointRestPositions()[PointIndex],&TargetBoneToIndex);
+
+							CurvesFacade.SetPointBoneIndices(PointIndex, BoneIndices);
+							CurvesFacade.SetPointBoneWeights(PointIndex, BoneWeights);
 
 						}, TransferBoneWeights.bUseParallel  ? EParallelForFlags::None : EParallelForFlags::ForceSingleThread);
 					}
diff --git a/Engine/Plugins/Runtime/HairStrands/Source/HairStrandsDataflow/Private/GroomAssetTerminalNode.cpp b/Engine/Plugins/Runtime/HairStrands/Source/HairStrandsDataflow/Private/GroomAssetTerminalNode.cpp
--- a/Engine/Plugins/Runtime/HairStrands/Source/HairStrandsDataflow/Private/GroomAssetTerminalNode.cpp
+++ b/Engine/Plugins/Runtime/HairStrands/Source/HairStrandsDataflow/Private/GroomAssetTerminalNode.cpp
@@ -2,7 +2,6 @@
 
 #include "GroomAssetTerminalNode.h"
 
-#include "AssetCompilingManager.h"
 #include "GroomCollectionFacades.h"
 #include "GroomEdit.h"
 #include "Dataflow/DataflowObjectInterface.h"
@@ -81,14 +80,6 @@ namespace UE::Groom::Private
 			if(InCollection.GetAttributeType(AttributeName, FacadeType::VerticesGroup) == EManagedArrayType::FFloatType)
 			{
 				const TManagedArray<float>& VerticesAttribute = InCollection.GetAttribute<float>(AttributeName, FacadeType::VerticesGroup);
-				if(OutCollection->NumElements(FacadeType::PointsGroup) != NumPoints)
-				{
-					if(OutCollection->NumElements(FacadeType::PointsGroup) > 0)
-					{
-						OutCollection->EmptyGroup(FacadeType::PointsGroup);
-					}
-					OutCollection->AddElements(NumPoints, FacadeType::PointsGroup);
-				}
 				TManagedArray<float>& PointsAttribute = OutCollection->AddAttribute<float>(AttributeName, FacadeType::PointsGroup);
 				
 				for(int32 PointIndex = 0; PointIndex < NumPoints; ++PointIndex)
@@ -106,28 +97,22 @@ void FGroomAssetTerminalDataflowNode::SetAssetValue(TObjectPtr<UObject> Asset, U
 	{
 		const FManagedArrayCollection& InCollection = GetValue<FManagedArrayCollection>(Context, &Collection);
 		UE::Groom::FGroomGuidesFacade GuidesFacade(InCollection);
-		UE::Groom::FGroomStrandsFacade StrandsFacade(InCollection);
 
-		FManagedArrayCollection* OutCollection = (GuidesFacade.IsValid() || StrandsFacade.IsValid()) ? new FManagedArrayCollection() : nullptr;
+		FManagedArrayCollection* OutCollection = (GuidesFacade.IsValid()) ? new FManagedArrayCollection() : nullptr;
 
-		if(StrandsFacade.IsValid())
-		{
-			// Transfer vertices weight maps onto the points to be stored onto the rest collection
-			UE::Groom::Private::TransferVerticesAttributes<UE::Groom::FGroomStrandsFacade>(InCollection, OutCollection, StrandsFacade.GetNumPoints());
-		}
 		if(GuidesFacade.IsValid())
 		{
 			// Only copy the attributes not stored already in the asset
 			const TArray<TTuple<FName, FName>> AttributesToCopy = {
 				{UE::Groom::FGroomGuidesFacade::PointsGroup, UE::Groom::FGroomGuidesFacade::PointBoneIndicesAttribute},
 				{UE::Groom::FGroomGuidesFacade::PointsGroup, UE::Groom::FGroomGuidesFacade::PointBoneWeightsAttribute},
-				{UE::Groom::FGroomGuidesFacade::CurvesGroup, UE::Groom::FGroomGuidesFacade::CurveParentIndicesAttribute},
-				{UE::Groom::FGroomGuidesFacade::CurvesGroup, UE::Groom::FGroomGuidesFacade::CurveLodIndicesAttribute}};
+				{UE::Groom::FGroomGuidesFacade::ObjectsGroup, UE::Groom::FGroomGuidesFacade::ObjectCurveOffsetsAttribute},
+				{UE::Groom::FGroomGuidesFacade::CurvesGroup, UE::Groom::FGroomGuidesFacade::CurvePointOffsetsAttribute}};
 
 			// Copy attributes from input collection
 			UE::Groom::Private::CopyCollectionAttributes(&InCollection, OutCollection, AttributesToCopy);
 			
-			// Transfer vertices weight maps onto the points to be stored onto the rest collection (VertexKinematicWeights should be included here)
+			// Transfer vertices weight maps onto the points to be stored onto the rest collection
 			UE::Groom::Private::TransferVerticesAttributes<UE::Groom::FGroomGuidesFacade>(InCollection, OutCollection, GuidesFacade.GetNumPoints());
 
 			// Build an editable groom asset for the strands
@@ -139,9 +124,6 @@ void FGroomAssetTerminalDataflowNode::SetAssetValue(TObjectPtr<UObject> Asset, U
 				// Build the editable guides
 				UE::Groom::Private::BuildEditableGuides(GuidesFacade,EditGroom);
 
-				// Ensure compilation dependent assets is done
-				FAssetCompilingManager::Get().FinishCompilationForObjects({ GroomAsset });
-
 				// Convert to groom asset
 				ConvertToGroomAsset(const_cast<UGroomAsset*>(GroomAsset), &EditGroom, EEditableGroomOperations::ControlPoints_Modified);
 			}
@@ -150,22 +132,6 @@ void FGroomAssetTerminalDataflowNode::SetAssetValue(TObjectPtr<UObject> Asset, U
 		if(OutCollection)
 		{
 			GroomAsset->GetDataflowSettings().SetRestCollection(OutCollection);
-			GroomAsset->GetDataflowSettings().InitSkeletalMeshes(GuidesFacade.GetNumObjects());
-
-			if(InCollection.HasAttribute(UE::Groom::FGroomGuidesFacade::ObjectSkeletalMeshesAttribute, UE::Groom::FGroomGuidesFacade::ObjectsGroup))
-			{
-				const TManagedArray<TObjectPtr<UObject>>& ObjectSkeletalMeshes =
-					InCollection.GetAttribute<TObjectPtr<UObject>>(UE::Groom::FGroomGuidesFacade::ObjectSkeletalMeshesAttribute, UE::Groom::FGroomGuidesFacade::ObjectsGroup);
-
-				const TManagedArray<int32>& ObjectMeshLODs =
-					InCollection.GetAttribute<int32>(UE::Groom::FGroomGuidesFacade::ObjectMeshLODsAttribute, UE::Groom::FGroomGuidesFacade::ObjectsGroup);
-				
-				for(int32 GroupIndex = 0; GroupIndex < GuidesFacade.GetNumObjects(); ++GroupIndex)
-				{
-					GroomAsset->GetDataflowSettings().SetSkeletalMesh(GroupIndex,
-						Cast<USkeletalMesh>(ObjectSkeletalMeshes[GroupIndex]), ObjectMeshLODs[GroupIndex]);
-				}
-			}
 		}
 	}
 }
diff --git a/Engine/Plugins/Runtime/HairStrands/Source/HairStrandsDataflow/Private/GroomCollectionFacades.cpp b/Engine/Plugins/Runtime/HairStrands/Source/HairStrandsDataflow/Private/GroomCollectionFacades.cpp
--- a/Engine/Plugins/Runtime/HairStrands/Source/HairStrandsDataflow/Private/GroomCollectionFacades.cpp
+++ b/Engine/Plugins/Runtime/HairStrands/Source/HairStrandsDataflow/Private/GroomCollectionFacades.cpp
@@ -14,8 +14,6 @@ namespace UE::Groom
 	const FName FGroomCollectionFacade<DerivedType>::ObjectsGroup(DerivedType::GroupPrefix.ToString() + FString("Objects"));
 	template<typename DerivedType>
 	const FName FGroomCollectionFacade<DerivedType>::PointsGroup(DerivedType::GroupPrefix.ToString() + FString("Points"));
-	template<typename DerivedType>
-	const FName FGroomCollectionFacade<DerivedType>::VerticesGroup(DerivedType::GroupPrefix.ToString() + FString("Vertices"));
 	
 	// Strands/Guides Attributes
 	template<typename DerivedType>
@@ -30,23 +28,24 @@ namespace UE::Groom
 	const FName FGroomCollectionFacade<DerivedType>::EdgeRestOrientationsAttribute("EdgeRestOrientations");
 	template<typename DerivedType>
 	const FName FGroomCollectionFacade<DerivedType>::PointRestPositionsAttribute("PointRestPositions");
-	template<typename DerivedType>
-	const FName FGroomCollectionFacade<DerivedType>::VertexLinearColorsAttribute("VertexLinearColors");
 
 	// Strands/Guides Prefix
 	const FName FGroomStrandsFacade::GroupPrefix("Strands");
 	const FName FGroomGuidesFacade::GroupPrefix("Guides");
 
+	// Guides Groups
+	const FName FGroomGuidesFacade::VerticesGroup(FGroomGuidesFacade::GroupPrefix.ToString() + FString("Vertices"));
+	//const FName FGroomGuidesFacade::VerticesGroup(FString("GuidesVertices"));
+
 	// Guides Attributes
-	const FName FGroomGuidesFacade::ObjectMeshLODsAttribute("ObjectMeshLODs");
-	const FName FGroomGuidesFacade::ObjectSkeletalMeshesAttribute("ObjectSkeletalMeshes");
 	const FName FGroomGuidesFacade::VertexKinematicWeightsAttribute("VertexKinematicWeights");
 	const FName FGroomGuidesFacade::PointBoneIndicesAttribute("PointBoneIndices");
 	const FName FGroomGuidesFacade::PointBoneWeightsAttribute("PointBoneWeights");
 	const FName FGroomGuidesFacade::ObjectPointSamplesAttribute("ObjectPointSamples");
 	const FName FGroomGuidesFacade::CurveStrandIndicesAttribute("CurveStrandIndices");
-	const FName FGroomGuidesFacade::CurveParentIndicesAttribute("CurveParentIndices");
-	const FName FGroomGuidesFacade::CurveLodIndicesAttribute("CurveLodIndices");
+	const FName FGroomGuidesFacade::CurveLevelIndicesAttribute("CurveLevelIndices");
+	const FName FGroomGuidesFacade::PointInterpIndicesAttribute("PointInterpIndices");
+	const FName FGroomGuidesFacade::PointInterpWeightsAttribute("PointInterpWeights");
 	
 	template<typename DerivedType>
 	FGroomCollectionFacade<DerivedType>::FGroomCollectionFacade(FManagedArrayCollection& InCollection) :
@@ -56,10 +55,8 @@ namespace UE::Groom
 		CurvePointOffsets(InCollection, CurvePointOffsetsAttribute, CurvesGroup),
 		ObjectCurveOffsets(InCollection, ObjectCurveOffsetsAttribute, ObjectsGroup),
 		PointCurveIndices(InCollection, PointCurveIndicesAttribute, PointsGroup),
-		CurveObjectIndices(InCollection, CurveObjectIndicesAttribute, CurvesGroup),
-		VertexLinearColors(InCollection, VertexLinearColorsAttribute, VerticesGroup)
-	{
-	}
+		CurveObjectIndices(InCollection, CurveObjectIndicesAttribute, CurvesGroup)
+	{}
 
 	template<typename DerivedType>
 	FGroomCollectionFacade<DerivedType>::FGroomCollectionFacade(const FManagedArrayCollection& InCollection) :
@@ -69,8 +66,7 @@ namespace UE::Groom
 		CurvePointOffsets(InCollection, CurvePointOffsetsAttribute, CurvesGroup),
 		ObjectCurveOffsets(InCollection, ObjectCurveOffsetsAttribute, ObjectsGroup),
 		PointCurveIndices(InCollection, PointCurveIndicesAttribute, PointsGroup),
-		CurveObjectIndices(InCollection, CurveObjectIndicesAttribute, CurvesGroup),
-		VertexLinearColors(InCollection, VertexLinearColorsAttribute, VerticesGroup)
+		CurveObjectIndices(InCollection, CurveObjectIndicesAttribute, CurvesGroup)
 	{
 	}
 
@@ -79,7 +75,7 @@ namespace UE::Groom
 	{
 		return EdgeRestOrientations.IsValid() && PointRestPositions.IsValid() && 
 			CurvePointOffsets.IsValid() && ObjectCurveOffsets.IsValid() && PointCurveIndices.IsValid() &&
-					CurveObjectIndices.IsValid() && VertexLinearColors.IsValid() && static_cast<const DerivedType*>(this)->IsFacadeValid();
+					CurveObjectIndices.IsValid() && static_cast<const DerivedType*>(this)->IsFacadeValid();
 	}
 
 	template<typename DerivedType>
@@ -92,7 +88,6 @@ namespace UE::Groom
 		ObjectCurveOffsets.Add();
 		PointCurveIndices.Add();
 		CurveObjectIndices.Add();
-		VertexLinearColors.Add();
 		
 		static_cast<DerivedType*>(this)->DefineFacadeSchema();
 	}
@@ -119,10 +114,6 @@ namespace UE::Groom
 			Collection->EmptyGroup(EdgesGroup);
 			Collection->AddElements(InPointRestPositions.Num()-InCurvePointOffsets.Num(), EdgesGroup);
 
-			// Vertices groups
-			Collection->EmptyGroup(VerticesGroup);
-			Collection->AddElements(InPointRestPositions.Num() * 2, VerticesGroup);
-
 			// Fill attributes
 			SetObjectCurveOffsets(InObjectCurveOffsets);
 			SetCurvePointOffsets(InCurvePointOffsets);
@@ -216,9 +207,10 @@ namespace UE::Groom
 		PointBoneIndices(InCollection, PointBoneIndicesAttribute, PointsGroup),
 		PointBoneWeights(InCollection, PointBoneWeightsAttribute, PointsGroup),
 		ObjectPointSamples(InCollection, ObjectPointSamplesAttribute, ObjectsGroup),
-		CurveStrandIndices(InCollection, CurveStrandIndicesAttribute, CurvesGroup),
-		CurveParentIndices(InCollection, CurveParentIndicesAttribute, CurvesGroup),
-		CurveLodIndices(InCollection, CurveLodIndicesAttribute, CurvesGroup)
+		CurveStrandIndices(InCollection, CurveStrandIndicesAttribute, CurvesGroup)
+		// CurveLevelIndices(InCollection, CurveLevelIndicesAttribute, CurvesGroup),
+		// PointInterpIndices(InCollection, PointInterpIndicesAttribute, PointsGroup),
+		// PointInterpWeights(InCollection, PointInterpWeightsAttribute, PointsGroup)
 	{
 		DefineSchema();
 	}
@@ -229,16 +221,21 @@ namespace UE::Groom
 		PointBoneIndices(InCollection, PointBoneIndicesAttribute, PointsGroup),
 		PointBoneWeights(InCollection, PointBoneWeightsAttribute, PointsGroup),
 		ObjectPointSamples(InCollection, ObjectPointSamplesAttribute, ObjectsGroup),
-		CurveStrandIndices(InCollection, CurveStrandIndicesAttribute, CurvesGroup),
-		CurveParentIndices(InCollection, CurveParentIndicesAttribute, CurvesGroup),
-		CurveLodIndices(InCollection, CurveLodIndicesAttribute, CurvesGroup)
+		CurveStrandIndices(InCollection, CurveStrandIndicesAttribute, CurvesGroup)
+		// CurveLevelIndices(InCollection, CurveLevelIndicesAttribute, CurvesGroup),
+		// PointInterpIndices(InCollection, PointInterpIndicesAttribute, PointsGroup),
+		// PointInterpWeights(InCollection, PointInterpWeightsAttribute, PointsGroup)
 	{
 	}
 
 	bool FGroomGuidesFacade::IsFacadeValid() const
 	{
-		return VertexKinematicWeights.IsValid() && PointBoneIndices.IsValid() && PointBoneWeights.IsValid() && ObjectPointSamples.IsValid() &&
-			CurveStrandIndices.IsValid() && CurveLodIndices.IsValid() && CurveParentIndices.IsValid();
+		return VertexKinematicWeights.IsValid() && PointBoneIndices.IsValid() && PointBoneWeights.IsValid() && ObjectPointSamples.IsValid();
+	}
+
+	bool FGroomGuidesFacade::HasSkinningAttributes() const
+	{
+		return PointBoneIndices.IsValid() && PointBoneWeights.IsValid() && ObjectCurveOffsets.IsValid() && CurvePointOffsets.IsValid();
 	}
 
 	void FGroomGuidesFacade::DefineFacadeSchema()
@@ -249,20 +246,26 @@ namespace UE::Groom
 		PointBoneIndices.Add();
 		ObjectPointSamples.Add();
 		CurveStrandIndices.Add();
-		CurveParentIndices.Add();
-		CurveLodIndices.Add();
+		// CurveLevelIndices.Add();
+		// PointInterpIndices.Add();
+		// PointInterpWeights.Add();
 	}
 
 	void FGroomGuidesFacade::InitFacadeCollection()
 	{
+		// Reset groups
+		Collection->EmptyGroup(VerticesGroup);
+		Collection->AddElements(GetNumPoints() * 2, VerticesGroup);
+
 		// Reset attributes
 		VertexKinematicWeights.Modify().Fill(0.0f);
-		PointBoneIndices.Modify().Fill(FIntVector4::ZeroValue);
-		PointBoneWeights.Modify().Fill(FVector4f::Zero());
+		PointBoneIndices.Modify().Fill(TArray<int32>());
+		PointBoneWeights.Modify().Fill(TArray<float>());
 		ObjectPointSamples.Modify().Fill(4);
 		CurveStrandIndices.Modify().Fill(INDEX_NONE);
-		CurveParentIndices.Modify().Fill(INDEX_NONE);
-		CurveLodIndices.Modify().Fill(INDEX_NONE);
+		// CurveLevelIndices.Modify().Fill(0);
+		// PointInterpIndices.Modify().Fill(FVector4f(INDEX_NONE));
+		// PointInterpWeights.Modify().Fill(FVector4f(0.0f));
 	}
 
 	void FGroomGuidesFacade::ResizePointsGroups(const int32 NumPoints) const 
diff --git a/Engine/Plugins/Runtime/HairStrands/Source/HairStrandsDataflow/Private/GroomDataflowRendering.cpp b/Engine/Plugins/Runtime/HairStrands/Source/HairStrandsDataflow/Private/GroomDataflowRendering.cpp
--- a/Engine/Plugins/Runtime/HairStrands/Source/HairStrandsDataflow/Private/GroomDataflowRendering.cpp
+++ b/Engine/Plugins/Runtime/HairStrands/Source/HairStrandsDataflow/Private/GroomDataflowRendering.cpp
@@ -2,7 +2,6 @@
 
 #include "GroomDataflowRendering.h"
 #include "GroomCollectionFacades.h"
-#include "Engine/SkeletalMesh.h"
 #include "GeometryCollection/GeometryCollection.h"
 
 namespace UE::Groom
@@ -11,21 +10,14 @@ namespace UE::Groom
 	static FAutoConsoleVariableRef CVarGroomDataflowDebugDraw(
 		 TEXT("p.Groom.Dataflow.DebugDraw"),
 		 GGroomDataflowDebugDraw,
-		 TEXT("Type of information we want to draw for debug (0-DefaultColor, 1-SkinWeights, 2-GuidesLods)")
+		 TEXT("Type of information we want to draw for debug (0-DefaultColor, 1-SkinWeights)")
 	 );
 
-	FString GGroomDataflowBoneName = TEXT("");
-	static FAutoConsoleVariableRef CVarGroomDataflowBoneName(
-		 TEXT("p.Groom.Dataflow.BoneName"),
-		 GGroomDataflowBoneName,
-		 TEXT("Bone name we want to visualize the skin weights")
-	 );
-
-	int32 GGroomDataflowGuidesLod = 0;
-	static FAutoConsoleVariableRef CVarGroomDataflowGuidesLod(
-		 TEXT("p.Groom.Dataflow.GuidesLod"),
-		 GGroomDataflowGuidesLod,
-		 TEXT("Guides Lod we want to display")
+	int32 GGroomDataflowBoneIndex = 0;
+	static FAutoConsoleVariableRef CVarGroomDataflowBoneIndex(
+		 TEXT("p.Groom.Dataflow.BoneIndex"),
+		 GGroomDataflowBoneIndex,
+		 TEXT("Bone index we want to visualize the skin weights")
 	 );
 
 	float GGroomDataflowRenderingThickness = 0.5f;
@@ -198,7 +190,7 @@ namespace UE::Groom
 				}
 
 				TArray<FLinearColor> Colors;
-				Colors.Init(FLinearColor(0.2,0.6,1.0), VertexPositions.Num());
+				Colors.Init(FLinearColor::Blue, VertexPositions.Num());
 				ColorLambda(GroomFacade, Colors);
  
 				const int32 GeometryIndex = RenderCollection.StartGeometryGroup(State.GetGuid().ToString());
@@ -209,95 +201,29 @@ namespace UE::Groom
 	}
 	
 	Dataflow::FRenderKey FGroomGuidesRenderingCallbacks::RenderKey = { "GuidesRender", FName("FGroomCollection") };
-
-	void RenderBoneWeights(const FGroomGuidesFacade& GuidesFacade, TArray<FLinearColor>& VertexColors)
+	
+	void FGroomGuidesRenderingCallbacks::Render(GeometryCollection::Facades::FRenderingFacade& RenderCollection, const Dataflow::FGraphRenderingState& State)
 	{
-		if(GuidesFacade.GetManagedArrayCollection().HasAttribute(
-						UE::Groom::FGroomGuidesFacade::ObjectSkeletalMeshesAttribute, UE::Groom::FGroomGuidesFacade::ObjectsGroup))
-		{
-			const TManagedArray<TObjectPtr<UObject>>& ObjectSkeletalMeshes = GuidesFacade.GetManagedArrayCollection().GetAttribute<TObjectPtr<UObject>>(
-				UE::Groom::FGroomGuidesFacade::ObjectSkeletalMeshesAttribute, UE::Groom::FGroomGuidesFacade::ObjectsGroup);
-				
-			for(int32 PointIndex = 0, PointEnd = GuidesFacade.GetNumPoints(); PointIndex < PointEnd; ++PointIndex)
+		RenderGroomCollection<FGroomGuidesFacade>(RenderCollection, State,
+			[](const FGroomGuidesFacade& GuidesFacade, TArray<FLinearColor>& VertexColors)
 			{
-				const FIntVector4& BoneIndices = GuidesFacade.GetPointBoneIndices(PointIndex);
-				const FVector4f& BoneWeights = GuidesFacade.GetPointBoneWeights(PointIndex);
-
-				const int32 CurveIndex = GuidesFacade.GetPointCurveIndices()[PointIndex];
-				const int32 ObjectIndex = GuidesFacade.GetCurveObjectIndices()[CurveIndex];
-
-				if(USkeletalMesh* SkeletalMesh = Cast<USkeletalMesh>(ObjectSkeletalMeshes[ObjectIndex]))
+				if(GGroomDataflowDebugDraw == 1)
 				{
-					for(int32 BoneIndex = 0, BoneEnd = BoneIndices.Num(); BoneIndex < BoneEnd; ++BoneIndex)
+					for(int32 PointIndex = 0, PointEnd = GuidesFacade.GetNumPoints(); PointIndex < PointEnd; ++PointIndex)
 					{
-						const int32 SkeletonBone = BoneIndices[BoneIndex];
-						if(SkeletonBone != INDEX_NONE)
+						const TArray<int32>& BoneIndices = GuidesFacade.GetPointBoneIndices(PointIndex);
+						const TArray<float>& BoneWeights = GuidesFacade.GetPointBoneWeights(PointIndex);
+						for(int32 BoneIndex = 0, BoneEnd = BoneIndices.Num(); BoneIndex < BoneEnd; ++BoneIndex)
 						{
-							const FName BoneName = SkeletalMesh->GetRefSkeleton().GetRawRefBoneInfo()[SkeletonBone].Name;
-							if(BoneName == GGroomDataflowBoneName)
+							if(BoneIndices[BoneIndex] == GGroomDataflowBoneIndex)
 							{
 								VertexColors[2*PointIndex] = FLinearColor::LerpUsingHSV(
 									FLinearColor::Black, FLinearColor::Yellow, BoneWeights[BoneIndex]);
 								VertexColors[2*PointIndex+1] = VertexColors[2*PointIndex];
 							}
-						}
+						}	
 					}
 				}
-			}
-		}
-	}
-	
-	void RenderGuidesLods(const FGroomGuidesFacade& GuidesFacade, TArray<FLinearColor>& VertexColors)
-	{
-		const TArray<int32>& ParentIndices = GuidesFacade.GetCurveParentIndices();
-		const TArray<int32>& LodIndices = GuidesFacade.GetCurveLodIndices();
-
-		TArray<FLinearColor> LodColors;
-		LodColors.Init(FLinearColor::Black, GuidesFacade.GetNumCurves());
-		for(int32 CurveIndex = 0, CurveEnd = GuidesFacade.GetNumCurves(); CurveIndex < CurveEnd; ++CurveIndex)
-		{
-			if(LodIndices[CurveIndex] >= GGroomDataflowGuidesLod)
-			{
-				LodColors[CurveIndex] = FLinearColor::MakeRandomColor();
-			}
-		}
-		for(int32 CurveIndex = 0, CurveEnd = GuidesFacade.GetNumCurves(); CurveIndex < CurveEnd; ++CurveIndex)
-		{
-			if((LodIndices[CurveIndex] != INDEX_NONE) && (LodIndices[CurveIndex] < GGroomDataflowGuidesLod))
-			{
-				int32 ParentIndex = CurveIndex;
-				int32 ParentLod = LodIndices[ParentIndex];
-				while((ParentLod < GGroomDataflowGuidesLod) && ParentIndex != INDEX_NONE)
-				{
-					ParentIndex = ParentIndices[ParentIndex];
-					ParentLod = LodIndices[ParentIndex];
-				}
-				LodColors[CurveIndex] = LodColors[ParentIndex];
-			}
-		}
-
-		for(int32 PointIndex = 0, PointEnd = GuidesFacade.GetNumPoints(); PointIndex < PointEnd; ++PointIndex)
-		{
-			const int32 CurveIndex = GuidesFacade.GetPointCurveIndices()[PointIndex];
-
-			VertexColors[2*PointIndex] = LodColors[CurveIndex];
-			VertexColors[2*PointIndex+1] = LodColors[CurveIndex];
-		}
-	}
-	
-	void FGroomGuidesRenderingCallbacks::Render(GeometryCollection::Facades::FRenderingFacade& RenderCollection, const Dataflow::FGraphRenderingState& State)
-	{
-		RenderGroomCollection<FGroomGuidesFacade>(RenderCollection, State,
-			[](const FGroomGuidesFacade& GuidesFacade, TArray<FLinearColor>& VertexColors)
-			{
-				if(GGroomDataflowDebugDraw == 1)
-				{
-					RenderBoneWeights(GuidesFacade, VertexColors);
-				}
-				if(GGroomDataflowDebugDraw == 2)
-				{
-					RenderGuidesLods(GuidesFacade, VertexColors);
-				}
 			});
 	}
 
diff --git a/Engine/Plugins/Runtime/HairStrands/Source/HairStrandsDataflow/Private/HairStrandsDataflowModule.cpp b/Engine/Plugins/Runtime/HairStrands/Source/HairStrandsDataflow/Private/HairStrandsDataflowModule.cpp
--- a/Engine/Plugins/Runtime/HairStrands/Source/HairStrandsDataflow/Private/HairStrandsDataflowModule.cpp
+++ b/Engine/Plugins/Runtime/HairStrands/Source/HairStrandsDataflow/Private/HairStrandsDataflowModule.cpp
@@ -1,9 +1,6 @@
 ﻿// Copyright Epic Games, Inc. All Rights Reserved.
 
 #include "HairStrandsDataflowModule.h"
-
-#include "AttachGuidesRootsNode.h"
-#include "BuildGuidesLODsNode.h"
 #include "GetGroomAssetNode.h"
 #include "GroomCollectionFacades.h"
 #include "GroomAssetTerminalNode.h"
@@ -34,13 +31,12 @@ public:
 
 	virtual TArray<FName> GetTargetGroupNames() const override
 	{
-		return {UE::Groom::FGroomGuidesFacade::VerticesGroup, UE::Groom::FGroomStrandsFacade::VerticesGroup};
+		return {UE::Groom::FGroomGuidesFacade::VerticesGroup};
 	}
 
 	virtual TArray<UE::Dataflow::FRenderingParameter> GetRenderingParameters() const override
 	{
-		return { {TEXT("GuidesRender"), FName("FGroomCollection"), {TEXT("Collection")}},
-			{TEXT("StrandsRender"), FName("FGroomCollection"), {TEXT("Collection")}} };
+		return { {TEXT("GuidesRender"), FName("FGroomCollection"), {TEXT("Collection")}}};
 	}
 };
 
@@ -55,8 +51,6 @@ void FHairStrandsDataflowModule::StartupModule()
 	DATAFLOW_NODE_REGISTER_CREATION_FACTORY(FGenerateGuidesCurvesDataflowNode);
 	DATAFLOW_NODE_REGISTER_CREATION_FACTORY(FGetGuidesKinematicWeightsDataflowNode);
 	DATAFLOW_NODE_REGISTER_CREATION_FACTORY(FSmoothGuidesCurvesDataflowNode);
-	DATAFLOW_NODE_REGISTER_CREATION_FACTORY(FAttachGuidesRootsDataflowNode);
-	DATAFLOW_NODE_REGISTER_CREATION_FACTORY(FBuildGuidesLODsDataflowNode);
 	
 	UE::Groom::RegisterRenderingCallbacks();
 	
diff --git a/Engine/Plugins/Runtime/HairStrands/Source/HairStrandsDataflow/Public/AttachGuidesRootsNode.h b/Engine/Plugins/Runtime/HairStrands/Source/HairStrandsDataflow/Public/AttachGuidesRootsNode.h
deleted file mode 100644
diff --git a/Engine/Plugins/Runtime/HairStrands/Source/HairStrandsDataflow/Public/BuildGuidesLODsNode.h b/Engine/Plugins/Runtime/HairStrands/Source/HairStrandsDataflow/Public/BuildGuidesLODsNode.h
deleted file mode 100644
diff --git a/Engine/Plugins/Runtime/HairStrands/Source/HairStrandsDataflow/Public/BuildGuidesSkinningNode.h b/Engine/Plugins/Runtime/HairStrands/Source/HairStrandsDataflow/Public/BuildGuidesSkinningNode.h
--- a/Engine/Plugins/Runtime/HairStrands/Source/HairStrandsDataflow/Public/BuildGuidesSkinningNode.h
+++ b/Engine/Plugins/Runtime/HairStrands/Source/HairStrandsDataflow/Public/BuildGuidesSkinningNode.h
@@ -35,17 +35,13 @@ public:
 	UPROPERTY(meta = (DataflowInput, DataflowOutput, DisplayName = "Collection", DataflowPassthrough  = "Collection"))
 	FManagedArrayCollection Collection;
 
-	/** SkeletalMesh used to transfer the skinning weights. Will be stored onto the groom asset */
+	/** Managed array collection to be used to store datas */
 	UPROPERTY(EditAnywhere, Category="Skinning")
 	TObjectPtr<USkeletalMesh> SkeletalMesh = nullptr;
 
-	/** LOD used to transfer the weights */
+	/** Managed array collection to be used to store datas */
 	UPROPERTY(EditAnywhere, Category="Skinning", meta = (DisplayName = "LOD Index"))
 	int32 LODIndex = 0;
-	
-	/** Group index on which the dats will be transfered. -1 will transfer on all the groups */
-	UPROPERTY(EditAnywhere, Category="Groups", meta = (DisplayName = "Group Index"))
-	int32 GroupIndex = INDEX_NONE;
 
 	/** The relative transform between the skeletal mesh and the groom asset. */
 	UPROPERTY(EditAnywhere, Category = "Skinning")
diff --git a/Engine/Plugins/Runtime/HairStrands/Source/HairStrandsDataflow/Public/GroomCollectionFacades.h b/Engine/Plugins/Runtime/HairStrands/Source/HairStrandsDataflow/Public/GroomCollectionFacades.h
--- a/Engine/Plugins/Runtime/HairStrands/Source/HairStrandsDataflow/Public/GroomCollectionFacades.h
+++ b/Engine/Plugins/Runtime/HairStrands/Source/HairStrandsDataflow/Public/GroomCollectionFacades.h
@@ -41,7 +41,6 @@ namespace UE::Groom
 		static const FName EdgesGroup;
 		static const FName ObjectsGroup;
 		static const FName PointsGroup;
-		static const FName VerticesGroup;
 		
 		/** Groom collection attribute names */
 		static const FName CurvePointOffsetsAttribute;
@@ -50,7 +49,6 @@ namespace UE::Groom
 		static const FName PointRestPositionsAttribute;
 		static const FName PointCurveIndicesAttribute;
 		static const FName CurveObjectIndicesAttribute;
-		static const FName VertexLinearColorsAttribute;
 
 		FGroomCollectionFacade(FManagedArrayCollection& InCollection);
 		FGroomCollectionFacade(const FManagedArrayCollection& InCollection);
@@ -75,9 +73,6 @@ namespace UE::Groom
 
 		/** Get the number of edges */
 		int32 GetNumEdges() const { return EdgeRestOrientations.Num(); }
-
-		/** Get the number of vertices */
-		int32 GetNumVertices() const { return VertexLinearColors.Num(); }
 		
 		/** Get the point rest positions */
 		const TArray<FVector3f>& GetPointRestPositions() const { return PointRestPositions.Get().GetConstArray(); }
@@ -96,9 +91,6 @@ namespace UE::Groom
 
 		/** Get the curve object indices */
 		const TArray<int32>& GetCurveObjectIndices() const { return CurveObjectIndices.Get().GetConstArray(); }
-
-		/** Get the curve object indices */
-		const TArray<FLinearColor>& GetVertexLinearColors() const { return VertexLinearColors.Get().GetConstArray(); }
 		
 		/** Set the point rest positions */
 		void SetPointRestPositions(const TArray<FVector3f>& InPointRestPositions) { PointRestPositions.Modify() = InPointRestPositions; UpdateEdgeRestOrientations(); }	
@@ -109,15 +101,9 @@ namespace UE::Groom
 		/** Set the object curve offsets */
 		void SetObjectCurveOffsets(const TArray<int32>& InObjectCurveOffsets) { ObjectCurveOffsets.Modify() = InObjectCurveOffsets; UpdateCurveObjectIndices(); }
 		
-		/** Set the vertex linear colors */
-		void SetVertexLinearColors(const TArray<FLinearColor>& InVertexLinearColors) { VertexLinearColors.Modify() = InVertexLinearColors; }
-		
 		/** Initialize the curve vertices from point positions and objects offsets */
 		void InitGroomCollection(const TArray<FVector3f>& InPointRestPositions,
 			const TArray<int32>& InCurvePointOffsets, const TArray<int32>& InObjectCurveOffsets);
-
-		/** Get the managed array collection */
-		const FManagedArrayCollection& GetManagedArrayCollection() const {return ConstCollection;}
 		
 	protected :
 		
@@ -153,9 +139,6 @@ namespace UE::Groom
 
 		/** Groom curves object index */
 		TManagedArrayAccessor<int32> CurveObjectIndices;
-
-		/** Groom vertices linear color*/
-		TManagedArrayAccessor<FLinearColor> VertexLinearColors;
 	};
 	
 	/**
@@ -201,16 +184,17 @@ namespace UE::Groom
 
 		FGroomGuidesFacade(FManagedArrayCollection& InCollection);
 		FGroomGuidesFacade(const FManagedArrayCollection& InCollection);
+		
+		static const FName VerticesGroup;
 
-		static const FName ObjectMeshLODsAttribute;
-		static const FName ObjectSkeletalMeshesAttribute;
 		static const FName VertexKinematicWeightsAttribute;
 		static const FName PointBoneIndicesAttribute;
 		static const FName PointBoneWeightsAttribute;
 		static const FName ObjectPointSamplesAttribute;
 		static const FName CurveStrandIndicesAttribute;
-		static const FName CurveParentIndicesAttribute;
-		static const FName CurveLodIndicesAttribute;
+		static const FName CurveLevelIndicesAttribute;
+		static const FName PointInterpIndicesAttribute;
+		static const FName PointInterpWeightsAttribute;
 		
 		/** Create the facade attributes. */
 		void DefineFacadeSchema();
@@ -218,9 +202,15 @@ namespace UE::Groom
 		/** Is the Facade defined on the collection? */
 		bool IsFacadeValid() const;
 
+		/** Check if the facade have skinning attributes*/
+		bool HasSkinningAttributes() const;
+
 		/** Init facade collection attributes */
 		void InitFacadeCollection();
 
+		/** Get the number of Vertices */
+		int32 GetNumVertices() const { return VertexKinematicWeights.Num(); }
+
 		/** Get the editable groom types */
 		static const TArray<FEditableType>& GetEditableGroom(const FEditableGroomGroup& GroomGroup)
 		{
@@ -231,10 +221,10 @@ namespace UE::Groom
 		const TArray<float>& GetVertexKinematicWeights() const { return VertexKinematicWeights.Get().GetConstArray(); }
 		
 		/** Get the point bone indices */
-		const FIntVector4& GetPointBoneIndices(const int32 PointIndex) const { return PointBoneIndices.Get()[PointIndex]; }
+		const TArray<int32>& GetPointBoneIndices(const int32 PointIndex) const { return PointBoneIndices.Get()[PointIndex]; }
 
 		/** Get the point bone weights*/
-		const FVector4f& GetPointBoneWeights(const int32 PointIndex) const { return PointBoneWeights.Get()[PointIndex]; }
+		const TArray<float>& GetPointBoneWeights(const int32 PointIndex) const { return PointBoneWeights.Get()[PointIndex]; }
 
 		/** Get the object point samples */
 		const TArray<int32>& GetObjectPointSamples() const { return ObjectPointSamples.Get().GetConstArray(); }
@@ -242,20 +232,14 @@ namespace UE::Groom
 		/** Get the guide strand indices */
 		const TArray<int32>& GetCurveStrandIndices() const { return CurveStrandIndices.Get().GetConstArray(); }
 
-		/** Get the guide parent indices */
-		const TArray<int32>& GetCurveParentIndices() const { return CurveParentIndices.Get().GetConstArray(); }
-		
-		/** Get the guide lod indices */
-		const TArray<int32>& GetCurveLodIndices() const { return CurveLodIndices.Get().GetConstArray(); }
-
 		/** Set the vertex kinematic weights */
 		void SetVertexKinematicWeights(const TArray<float>& InKinematicWeights) { VertexKinematicWeights.Modify() = InKinematicWeights; }
 
 		/** Set the point bone indices */
-		void SetPointBoneIndices(const int32 PointIndex, const FIntVector4& InBoneIndices) { return PointBoneIndices.ModifyAt(PointIndex,InBoneIndices); }
+		void SetPointBoneIndices(const int32 PointIndex, const TArray<int32>& InBoneIndices) { return PointBoneIndices.ModifyAt(PointIndex,InBoneIndices); }
 
 		/** Set the point bone weights*/
-		void SetPointBoneWeights(const int32 PointIndex, const FVector4f& InBoneWeights) { return PointBoneWeights.ModifyAt(PointIndex, InBoneWeights); }
+		void SetPointBoneWeights(const int32 PointIndex, const TArray<float>& InBoneWeights) { return PointBoneWeights.ModifyAt(PointIndex, InBoneWeights); }
 		
 		/** Set the object point samples */
 		void SetObjectPointSamples(const TArray<int32>& NumPointSamples) { ObjectPointSamples.Modify() = NumPointSamples; }
@@ -263,12 +247,6 @@ namespace UE::Groom
 		/** Set the guide strand indices */
         void SetCurveStrandIndices(const TArray<int32>& StrandIndices) { CurveStrandIndices.Modify() = StrandIndices; }
 
-		/** Set the guide parent indices */
-		void SetCurveParentIndices(const TArray<int32>& ParentIndices) { CurveParentIndices.Modify() = ParentIndices; }
-
-		/** Set the guide lod indices */
-		void SetCurveLodIndices(const TArray<int32>& LodIndices) { CurveLodIndices.Modify() = LodIndices; }
-
 		/** Resize the geometry groups (vertex, points, edges) based on a new total number of points */
 		void ResizePointsGroups(const int32 NumPoints) const;
 
@@ -278,10 +256,10 @@ namespace UE::Groom
 		TManagedArrayAccessor<float> VertexKinematicWeights;
 
 		/** Point bone indices */
-		TManagedArrayAccessor<FIntVector4> PointBoneIndices;
+		TManagedArrayAccessor<TArray<int32>> PointBoneIndices;
 
 		/** Point bone weights */
-		TManagedArrayAccessor<FVector4f> PointBoneWeights;
+		TManagedArrayAccessor<TArray<float>> PointBoneWeights;
 
 		/** Object point samples */
 		TManagedArrayAccessor<int32> ObjectPointSamples;
@@ -289,10 +267,13 @@ namespace UE::Groom
 		/** Strand index from which the guide has been generated */
 		TManagedArrayAccessor<int32> CurveStrandIndices;
 
-		/** Parent guide indices */
-		TManagedArrayAccessor<int32> CurveParentIndices;
-
-		/** Lod guide indices */
-		TManagedArrayAccessor<int32> CurveLodIndices;
+		// /** Level index for the curves hierarchy */
+		// TManagedArrayAccessor<int32> CurveLevelIndices;
+		//
+		// /** Point (Level L) to point (Level L+1) interpolation indices+coords */
+		// TManagedArrayAccessor<FVector4f> PointInterpIndices;
+		//
+		// /** Point (Level L) to point (Level L+1) interpolation weights */
+		// TManagedArrayAccessor<FVector4f> PointInterpWeights;
 	};
 }
diff --git a/Engine/Plugins/Runtime/HairStrands/Source/HairStrandsDeformer/HairStrandsDeformer.Build.cs b/Engine/Plugins/Runtime/HairStrands/Source/HairStrandsDeformer/HairStrandsDeformer.Build.cs
--- a/Engine/Plugins/Runtime/HairStrands/Source/HairStrandsDeformer/HairStrandsDeformer.Build.cs
+++ b/Engine/Plugins/Runtime/HairStrands/Source/HairStrandsDeformer/HairStrandsDeformer.Build.cs
@@ -19,9 +19,7 @@ namespace UnrealBuildTool.Rules
 					"RHI",
 					"ComputeFramework",
 					"OptimusCore",
-					"HairStrandsCore",
-					"HairStrandsSolver", 
-					"Chaos",
+					"HairStrandsCore"
 				});
 
 			PrivateIncludePathModuleNames.AddRange(
diff --git a/Engine/Plugins/Runtime/HairStrands/Source/HairStrandsDeformer/Private/DeformerDataInterfaceGroom.cpp b/Engine/Plugins/Runtime/HairStrands/Source/HairStrandsDeformer/Private/DeformerDataInterfaceGroom.cpp
new file mode 100644
--- /dev/null
+++ b/Engine/Plugins/Runtime/HairStrands/Source/HairStrandsDeformer/Private/DeformerDataInterfaceGroom.cpp
@@ -0,0 +1,279 @@
+// Copyright Epic Games, Inc. All Rights Reserved.
+
+#include "DeformerDataInterfaceGroom.h"
+
+#include "ComputeFramework/ShaderParamTypeDefinition.h"
+#include "ComputeFramework/ComputeKernelPermutationVector.h"
+#include "OptimusDataDomain.h"
+#include "ShaderCompilerCore.h"
+#include "ShaderParameterMetadataBuilder.h"
+#include "GroomComponent.h"
+#include "GroomInstance.h"
+#include "RenderGraphBuilder.h"
+#include "GlobalRenderResources.h"
+#include "HairStrandsDefinitions.h"
+#include "DeformerGroomComponentSource.h"
+#include "HairStrandsInterpolation.h"
+#include "SystemTextures.h"
+
+FString UOptimusGroomDataInterface::GetDisplayName() const
+{
+	return TEXT("Groom");
+}
+
+TArray<FOptimusCDIPinDefinition> UOptimusGroomDataInterface::GetPinDefinitions() const
+{
+	FName ControlPoint(UOptimusGroomComponentSource::Domains::ControlPoint);
+	FName Curve(UOptimusGroomComponentSource::Domains::Curve);
+
+	TArray<FOptimusCDIPinDefinition> Defs;
+	Defs.Add({ "NumControlPoints", "ReadNumControlPoints" });
+	Defs.Add({ "NumCurves",        "ReadNumCurves" });
+	Defs.Add({ "Position",         "ReadPosition",          ControlPoint,   "ReadPosition" });
+	Defs.Add({ "Radius",           "ReadRadius",            ControlPoint,   "ReadRadius" });
+	Defs.Add({ "CoordU",           "ReadCoordU",            ControlPoint,   "ReadCoordU" });
+	Defs.Add({ "Length",           "ReadLength",            ControlPoint,   "ReadLength" });
+	Defs.Add({ "RootUV",           "ReadRootUV",            ControlPoint,   "ReadRootUV" });
+	Defs.Add({ "Seed",             "ReadSeed",              ControlPoint,   "ReadSeed" });
+	Defs.Add({ "ClumpId",          "ReadClumpId",           ControlPoint,   "ReadClumpId" });
+	Defs.Add({ "Color",            "ReadColor",             ControlPoint,   "ReadColor" });
+	Defs.Add({ "Roughness",        "ReadRoughness",         ControlPoint,   "ReadRoughness" });
+	Defs.Add({ "AO",               "ReadAO",                ControlPoint,   "ReadAO" });
+	Defs.Add({ "CurveOffsetPoint", "ReadCurveOffsetPoint",  Curve,          "ReadCurveOffsetPoint" });
+	Defs.Add({ "CurveNumPoint",    "ReadCurveNumPoint",     Curve,          "ReadCurveNumPoint" });
+	Defs.Add({ "GuideIndex",       "ReadGuideIndex",        ControlPoint,   "ReadGuideIndex" });
+	return Defs;
+}
+ 
+TSubclassOf<UActorComponent> UOptimusGroomDataInterface::GetRequiredComponentClass() const
+{
+	return UGroomComponent::StaticClass();
+}
+
+void UOptimusGroomDataInterface::GetSupportedInputs(TArray<FShaderFunctionDefinition>& OutFunctions) const
+{
+	OutFunctions.AddDefaulted_GetRef()
+		.SetName(TEXT("ReadNumControlPoints"))
+		.AddReturnType(EShaderFundamentalType::Uint);
+
+	OutFunctions.AddDefaulted_GetRef()
+		.SetName(TEXT("ReadNumCurves"))
+		.AddReturnType(EShaderFundamentalType::Uint);
+
+	OutFunctions.AddDefaulted_GetRef()
+		.SetName(TEXT("ReadPosition"))
+		.AddReturnType(EShaderFundamentalType::Float, 3)
+		.AddParam(EShaderFundamentalType::Uint);
+
+	OutFunctions.AddDefaulted_GetRef()
+		.SetName(TEXT("ReadRadius"))
+		.AddReturnType(EShaderFundamentalType::Float)
+		.AddParam(EShaderFundamentalType::Uint);
+
+	OutFunctions.AddDefaulted_GetRef()
+		.SetName(TEXT("ReadCoordU"))
+		.AddReturnType(EShaderFundamentalType::Float)
+		.AddParam(EShaderFundamentalType::Uint);
+
+	OutFunctions.AddDefaulted_GetRef()
+		.SetName(TEXT("ReadLength"))
+		.AddReturnType(EShaderFundamentalType::Float)
+		.AddParam(EShaderFundamentalType::Uint);
+
+	OutFunctions.AddDefaulted_GetRef()
+		.SetName(TEXT("ReadRootUV"))
+		.AddReturnType(EShaderFundamentalType::Float, 2)
+		.AddParam(EShaderFundamentalType::Uint);
+
+	OutFunctions.AddDefaulted_GetRef()
+		.SetName(TEXT("ReadSeed"))
+		.AddReturnType(EShaderFundamentalType::Float)
+		.AddParam(EShaderFundamentalType::Uint);
+
+	OutFunctions.AddDefaulted_GetRef()
+		.SetName(TEXT("ReadClumpId"))
+		.AddReturnType(EShaderFundamentalType::Uint)
+		.AddParam(EShaderFundamentalType::Uint);
+
+	OutFunctions.AddDefaulted_GetRef()
+		.SetName(TEXT("ReadColor"))
+		.AddReturnType(EShaderFundamentalType::Float, 3)
+		.AddParam(EShaderFundamentalType::Uint);
+
+	OutFunctions.AddDefaulted_GetRef()
+		.SetName(TEXT("ReadRoughnes"))
+		.AddReturnType(EShaderFundamentalType::Float)
+		.AddParam(EShaderFundamentalType::Uint);
+
+	OutFunctions.AddDefaulted_GetRef()
+		.SetName(TEXT("ReadAO"))
+		.AddReturnType(EShaderFundamentalType::Float)
+		.AddParam(EShaderFundamentalType::Uint);
+
+	OutFunctions.AddDefaulted_GetRef()
+		.SetName(TEXT("ReadCurveOffsetPoint"))
+		.AddReturnType(EShaderFundamentalType::Uint)
+		.AddParam(EShaderFundamentalType::Uint);
+
+	OutFunctions.AddDefaulted_GetRef()
+		.SetName(TEXT("ReadCurveNumPoint"))
+		.AddReturnType(EShaderFundamentalType::Uint)
+		.AddParam(EShaderFundamentalType::Uint);
+
+	OutFunctions.AddDefaulted_GetRef()
+		.SetName(TEXT("ReadGuideIndex"))
+		.AddReturnType(EShaderFundamentalType::Uint)
+		.AddParam(EShaderFundamentalType::Uint);
+}
+
+BEGIN_SHADER_PARAMETER_STRUCT(FGroomDataInterfaceParameters, )
+	SHADER_PARAMETER_STRUCT_INCLUDE(FHairStrandsInstanceCommonParameters, Common)
+	SHADER_PARAMETER_STRUCT_INCLUDE(FHairStrandsInstanceResourceParameters, Resources)
+	SHADER_PARAMETER_STRUCT_INCLUDE(FHairStrandsInstanceInterpolationParameters, Interpolation)
+END_SHADER_PARAMETER_STRUCT()
+
+void UOptimusGroomDataInterface::GetShaderParameters(TCHAR const* UID, FShaderParametersMetadataBuilder& InOutBuilder, FShaderParametersMetadataAllocations& InOutAllocations) const
+{
+	InOutBuilder.AddNestedStruct<FGroomDataInterfaceParameters>(UID);
+}
+
+TCHAR const* UOptimusGroomDataInterface::TemplateFilePath = TEXT("/Plugin/Runtime/HairStrands/Private/DeformerDataInterfaceGroom.ush");
+
+TCHAR const* UOptimusGroomDataInterface::GetShaderVirtualPath() const
+{
+	return TemplateFilePath;
+}
+
+void UOptimusGroomDataInterface::GetShaderHash(FString& InOutKey) const
+{
+	GetShaderFileHash(TemplateFilePath, EShaderPlatform::SP_PCD3D_SM5).AppendString(InOutKey);
+}
+
+void UOptimusGroomDataInterface::GetHLSL(FString& OutHLSL, FString const& InDataInterfaceName) const
+{
+	TMap<FString, FStringFormatArg> TemplateArgs =
+	{
+		{ TEXT("DataInterfaceName"), InDataInterfaceName },
+	};
+
+	FString TemplateFile;
+	LoadShaderSourceFile(TemplateFilePath, EShaderPlatform::SP_PCD3D_SM5, &TemplateFile, nullptr);
+	OutHLSL += FString::Format(*TemplateFile, TemplateArgs);
+}
+
+UComputeDataProvider* UOptimusGroomDataInterface::CreateDataProvider(TObjectPtr<UObject> InBinding, uint64 InInputMask, uint64 InOutputMask) const
+{
+	UOptimusGroomDataProvider* Provider = NewObject<UOptimusGroomDataProvider>();
+	Provider->Groom = Cast<UGroomComponent>(InBinding);
+	return Provider;
+}
+
+
+FComputeDataProviderRenderProxy* UOptimusGroomDataProvider::GetRenderProxy()
+{
+	return new FOptimusGroomDataProviderProxy(Groom);
+}
+
+
+FOptimusGroomDataProviderProxy::FOptimusGroomDataProviderProxy(UGroomComponent* InGroomComponent)
+{
+	GroomComponent = InGroomComponent;
+}
+
+bool FOptimusGroomDataProviderProxy::IsValid(FValidationData const& InValidationData) const
+{
+	if (InValidationData.ParameterStructSize != sizeof(FParameters))
+	{
+		return false;
+	}
+	if (GroomComponent == nullptr)
+	{
+		return false;
+	}
+	const uint32 GroupCount = GroomComponent->GetGroupCount();
+	if (InValidationData.NumInvocations != GroupCount)
+	{
+		return false;
+	}
+	// Earlier invalidation in case one of the instance is not ready/valid
+	for (uint32 GroupIndex = 0; GroupIndex < GroupCount; ++GroupIndex)
+	{
+		if (GroomComponent->GetGroupInstance(GroupIndex) == nullptr)
+		{
+			return false;
+		}
+	}
+	
+	return true;
+}
+
+void FOptimusGroomDataProviderProxy::AllocateResources(FRDGBuilder& GraphBuilder)
+{
+	Resources.Empty();
+	FallbackByteAddressSRV = nullptr;
+	FallbackStructuredSRV = nullptr;
+	const uint32 InstanceCount = GroomComponent ? GroomComponent->GetGroupCount() : 0;
+	for (uint32 Index =0; Index <InstanceCount;++Index)
+	{
+		if (FHairGroupInstance* Instance = GroomComponent->GetGroupInstance(Index))
+		{
+			{
+				FHairStrandsInstanceResourceParameters& R = Resources.AddDefaulted_GetRef();
+				R.PositionBuffer		= RegisterAsSRV(GraphBuilder, Instance->Strands.RestResource->PositionBuffer);
+				R.PositionOffsetBuffer 	= RegisterAsSRV(GraphBuilder, Instance->Strands.RestResource->PositionOffsetBuffer);
+				R.CurveBuffer			= RegisterAsSRV(GraphBuilder, Instance->Strands.RestResource->CurveBuffer);
+				R.PointToCurveBuffer	= RegisterAsSRV(GraphBuilder, Instance->Strands.RestResource->PointToCurveBuffer);
+				R.CurveAttributeBuffer	= RegisterAsSRV(GraphBuilder, Instance->Strands.RestResource->CurveAttributeBuffer);
+				R.PointAttributeBuffer	= RegisterAsSRV(GraphBuilder, Instance->Strands.RestResource->PointAttributeBuffer);
+			}
+
+			{
+				FHairStrandsInstanceInterpolationParameters& R = Interpolations.AddDefaulted_GetRef();
+				R.CurveInterpolationBuffer = Instance->Strands.InterpolationResource ? RegisterAsSRV(GraphBuilder, Instance->Strands.InterpolationResource->CurveInterpolationBuffer) : nullptr;
+				R.PointInterpolationBuffer = Instance->Strands.InterpolationResource ? RegisterAsSRV(GraphBuilder, Instance->Strands.InterpolationResource->PointInterpolationBuffer) : nullptr;
+			}
+
+			if (!FallbackByteAddressSRV){ FallbackByteAddressSRV = GraphBuilder.CreateSRV(GSystemTextures.GetDefaultByteAddressBuffer(GraphBuilder, 16u)); }
+			if (!FallbackStructuredSRV) { FallbackStructuredSRV  = GraphBuilder.CreateSRV(GSystemTextures.GetDefaultStructuredBuffer(GraphBuilder, 16u)); }
+		}
+	}
+}
+
+void FOptimusGroomDataProviderProxy::GatherDispatchData(FDispatchData const& InDispatchData)
+{
+	const uint32 InstanceCount = GroomComponent ? GroomComponent->GetGroupCount() : 0;
+	check(InDispatchData.NumInvocations == InstanceCount);
+	check(uint32(Resources.Num()) == InstanceCount);
+
+	const TStridedView<FParameters> ParameterArray = MakeStridedParameterView<FParameters>(InDispatchData);
+	for (int32 InvocationIndex = 0; InvocationIndex < ParameterArray.Num(); ++InvocationIndex)
+	{
+		// Default valid in case the instance is not valid
+		FParameters& Parameters = ParameterArray[InvocationIndex];
+		Parameters.Common.PointCount = 0;
+		Parameters.Common.CurveCount = 0;
+		Parameters.Resources.PositionBuffer			= FallbackByteAddressSRV;
+		Parameters.Resources.PositionOffsetBuffer 	= FallbackStructuredSRV;
+		Parameters.Resources.CurveAttributeBuffer	= FallbackByteAddressSRV;
+		Parameters.Resources.PointAttributeBuffer	= FallbackByteAddressSRV;
+		Parameters.Resources.PointToCurveBuffer		= FallbackByteAddressSRV;
+		Parameters.Resources.CurveBuffer			= FallbackByteAddressSRV;
+		Parameters.Interpolation.CurveInterpolationBuffer = FallbackByteAddressSRV;
+		Parameters.Interpolation.PointInterpolationBuffer = FallbackByteAddressSRV;
+
+		if (FHairGroupInstance* Instance = GroomComponent->GetGroupInstance(InvocationIndex))
+		{
+			const FHairGroupPublicData::FVertexFactoryInput VFInput = ComputeHairStrandsVertexInputData(Instance, EGroomViewMode::None);
+			Parameters.Common = VFInput.Strands.Common;
+
+			if (Resources[InvocationIndex].PositionBuffer != nullptr)
+			{
+				Parameters.Resources = Resources[InvocationIndex];
+			}
+			if (Interpolations[InvocationIndex].CurveInterpolationBuffer != nullptr)
+			{
+				Parameters.Interpolation = Interpolations[InvocationIndex];
+			}
+		}
+	}
+}
\ No newline at end of file
diff --git a/Engine/Plugins/Runtime/HairStrands/Source/HairStrandsDeformer/Private/DeformerDataInterfaceGroom.h b/Engine/Plugins/Runtime/HairStrands/Source/HairStrandsDeformer/Private/DeformerDataInterfaceGroom.h
new file mode 100644
--- /dev/null
+++ b/Engine/Plugins/Runtime/HairStrands/Source/HairStrandsDeformer/Private/DeformerDataInterfaceGroom.h
@@ -0,0 +1,75 @@
+// Copyright Epic Games, Inc. All Rights Reserved.
+
+#pragma once
+
+#include "RenderGraphDefinitions.h"
+#include "OptimusComputeDataInterface.h"
+#include "HairStrandsInterface.h"
+#include "ComputeFramework/ComputeDataProvider.h"
+#include "DeformerDataInterfaceGroom.generated.h"
+
+class FGroomDataInterfaceParameters;
+class UGroomComponent;
+
+/** Compute Framework Data Interface for reading groom. */
+UCLASS(Category = ComputeFramework)
+class UOptimusGroomDataInterface : public UOptimusComputeDataInterface
+{
+	GENERATED_BODY()
+
+public:
+	//~ Begin UOptimusComputeDataInterface Interface
+	FString GetDisplayName() const override;
+	TArray<FOptimusCDIPinDefinition> GetPinDefinitions() const override;
+	TSubclassOf<UActorComponent> GetRequiredComponentClass() const override;
+	//~ End UOptimusComputeDataInterface Interface
+	
+	//~ Begin UComputeDataInterface Interface
+	TCHAR const* GetClassName() const override { return TEXT("Groom"); }
+	void GetSupportedInputs(TArray<FShaderFunctionDefinition>& OutFunctions) const override;
+	void GetShaderParameters(TCHAR const* UID, FShaderParametersMetadataBuilder& InOutBuilder, FShaderParametersMetadataAllocations& InOutAllocations) const override;
+	TCHAR const* GetShaderVirtualPath() const override;
+	void GetShaderHash(FString& InOutKey) const override;
+	void GetHLSL(FString& OutHLSL, FString const& InDataInterfaceName) const override;
+	UComputeDataProvider* CreateDataProvider(TObjectPtr<UObject> InBinding, uint64 InInputMask, uint64 InOutputMask) const override;
+	//~ End UComputeDataInterface Interface
+
+private:
+	static TCHAR const* TemplateFilePath;
+};
+
+/** Compute Framework Data Provider for reading groom. */
+UCLASS(BlueprintType, editinlinenew, Category = ComputeFramework)
+class UOptimusGroomDataProvider : public UComputeDataProvider
+{
+	GENERATED_BODY()
+
+public:
+	UPROPERTY(EditAnywhere, BlueprintReadWrite, Category = Binding)
+	TObjectPtr<UGroomComponent> Groom = nullptr;
+
+	//~ Begin UComputeDataProvider Interface
+	FComputeDataProviderRenderProxy* GetRenderProxy() override;
+	//~ End UComputeDataProvider Interface
+};
+
+class FOptimusGroomDataProviderProxy : public FComputeDataProviderRenderProxy
+{
+public:
+	FOptimusGroomDataProviderProxy(UGroomComponent* InGroomComponent);
+
+	//~ Begin FComputeDataProviderRenderProxy Interface
+	bool IsValid(FValidationData const& InValidationData) const override;
+	void AllocateResources(FRDGBuilder& GraphBuilder) override;
+	void GatherDispatchData(FDispatchData const& InDispatchData) override;
+	//~ End FComputeDataProviderRenderProxy Interface
+
+private:
+	using FParameters = FGroomDataInterfaceParameters;
+
+	UGroomComponent* GroomComponent = nullptr; // Should it be HairInstance instead?
+	TArray<FHairStrandsInstanceResourceParameters> Resources;
+	TArray<FHairStrandsInstanceInterpolationParameters> Interpolations;
+	FRDGBufferSRVRef FallbackByteAddressSRV;
+	FRDGBufferSRVRef FallbackStructuredSRV;
+};
diff --git a/Engine/Plugins/Runtime/HairStrands/Source/HairStrandsDeformer/Private/DeformerDataInterfaceGroomExec.cpp b/Engine/Plugins/Runtime/HairStrands/Source/HairStrandsDeformer/Private/DeformerDataInterfaceGroomExec.cpp
new file mode 100644
--- /dev/null
+++ b/Engine/Plugins/Runtime/HairStrands/Source/HairStrandsDeformer/Private/DeformerDataInterfaceGroomExec.cpp
@@ -0,0 +1,174 @@
+// Copyright Epic Games, Inc. All Rights Reserved.
+
+#include "DeformerDataInterfaceGroomExec.h"
+
+#include "DeformerGroomComponentSource.h"
+#include "GroomComponent.h"
+#include "ComputeFramework/ShaderParamTypeDefinition.h"
+#include "ShaderCompilerCore.h"
+#include "ShaderParameterMetadataBuilder.h"
+#include "SkeletalRenderPublic.h"
+
+FString UOptimusGroomExecDataInterface::GetDisplayName() const
+{
+	return TEXT("Execute Groom");
+}
+
+FName UOptimusGroomExecDataInterface::GetCategory() const
+{
+	return CategoryName::ExecutionDataInterfaces;
+}
+
+TArray<FOptimusCDIPinDefinition> UOptimusGroomExecDataInterface::GetPinDefinitions() const
+{
+	TArray<FOptimusCDIPinDefinition> Defs;
+	Defs.Add({ "NumThreads", "ReadNumThreads" });
+	return Defs;
+}
+
+
+TSubclassOf<UActorComponent> UOptimusGroomExecDataInterface::GetRequiredComponentClass() const
+{
+	return UGroomComponent::StaticClass();
+}
+
+
+void UOptimusGroomExecDataInterface::GetSupportedInputs(TArray<FShaderFunctionDefinition>& OutFunctions) const
+{
+	OutFunctions.AddDefaulted_GetRef()
+		.SetName(TEXT("ReadNumThreads"))
+		.AddReturnType(EShaderFundamentalType::Int, 3);
+}
+
+
+BEGIN_SHADER_PARAMETER_STRUCT(FGroomExecDataInterfaceParameters, )
+	SHADER_PARAMETER(FIntVector, NumThreads)
+END_SHADER_PARAMETER_STRUCT()
+
+void UOptimusGroomExecDataInterface::GetShaderParameters(TCHAR const* UID, FShaderParametersMetadataBuilder& InOutBuilder, FShaderParametersMetadataAllocations& InOutAllocations) const
+{
+	InOutBuilder.AddNestedStruct<FGroomExecDataInterfaceParameters>(UID);
+}
+
+TCHAR const* UOptimusGroomExecDataInterface::TemplateFilePath = TEXT("/Plugin/Runtime/HairStrands/Private/DeformerDataInterfaceGroomExec.ush");
+
+TCHAR const* UOptimusGroomExecDataInterface::GetShaderVirtualPath() const
+{
+	return TemplateFilePath;
+}
+
+void UOptimusGroomExecDataInterface::GetShaderHash(FString& InOutKey) const
+{
+	GetShaderFileHash(TemplateFilePath, EShaderPlatform::SP_PCD3D_SM5).AppendString(InOutKey);
+}
+
+void UOptimusGroomExecDataInterface::GetHLSL(FString& OutHLSL, FString const& InDataInterfaceName) const
+{
+	TMap<FString, FStringFormatArg> TemplateArgs =
+	{
+		{ TEXT("DataInterfaceName"), InDataInterfaceName },
+	};
+
+	FString TemplateFile;
+	LoadShaderSourceFile(TemplateFilePath, EShaderPlatform::SP_PCD3D_SM5, &TemplateFile, nullptr);
+	OutHLSL += FString::Format(*TemplateFile, TemplateArgs);
+}
+
+UComputeDataProvider* UOptimusGroomExecDataInterface::CreateDataProvider(TObjectPtr<UObject> InBinding, uint64 InInputMask, uint64 InOutputMask) const
+{
+	UOptimusGroomExecDataProvider* Provider = NewObject<UOptimusGroomExecDataProvider>();
+	Provider->GroomComponent = Cast<UGroomComponent>(InBinding);
+	Provider->Domain = Domain;
+	return Provider;
+}
+
+FName UOptimusGroomExecDataInterface::GetSelectedExecutionDomainName() const
+{
+	if (Domain == EOptimusGroomExecDomain::ControlPoint)
+	{
+		return UOptimusGroomComponentSource::Domains::ControlPoint;
+	}
+
+	if (Domain == EOptimusGroomExecDomain::Curve)
+	{
+		return UOptimusGroomComponentSource::Domains::Curve;
+	}
+
+	return NAME_None;
+}
+
+
+FComputeDataProviderRenderProxy* UOptimusGroomExecDataProvider::GetRenderProxy()
+{
+	return new FOptimusGroomExecDataProviderProxy(GroomComponent, Domain);
+}
+
+
+FOptimusGroomExecDataProviderProxy::FOptimusGroomExecDataProviderProxy(UGroomComponent* InGroomComponent, EOptimusGroomExecDomain InDomain)
+{
+	GroomComponent = InGroomComponent;
+	Domain = InDomain;
+	Instances.SetNum(0);
+
+	if (GroomComponent)
+	{
+		const uint32 InstanceCount = GroomComponent->GetGroupCount();
+		Instances.Reserve(InstanceCount);
+		for (uint32 InstanceIt = 0; InstanceIt < InstanceCount; ++InstanceIt)
+		{		
+			if (FHairGroupInstance* Instance = GroomComponent->GetGroupInstance(InstanceIt))
+			{
+				Instances.Add(Instance);
+			}
+		}
+	}
+}
+
+int32 FOptimusGroomExecDataProviderProxy::GetDispatchThreadCount(TArray<FIntVector>& ThreadCounts) const
+{
+	ThreadCounts.Reset();
+	ThreadCounts.Reserve(Instances.Num());
+
+	uint32 NumInvocations = 0;
+	for (const FHairGroupInstance* Instance : Instances)
+	{	
+		if (Instance && Instance->Strands.IsValid())
+		{
+			const int32 NumControlPoints = Instance->Strands.GetData().GetNumPoints();
+			const int32 NumCurves = Instance->Strands.GetData().GetNumCurves();
+			const int32 NumThreads = Domain == EOptimusGroomExecDomain::ControlPoint ? NumControlPoints : NumCurves;
+			ThreadCounts.Add(FIntVector(NumThreads, 1, 1));
+			++NumInvocations;
+		}
+	}
+	return NumInvocations;
+}
+
+bool FOptimusGroomExecDataProviderProxy::IsValid(FValidationData const& InValidationData) const
+{
+	if (InValidationData.ParameterStructSize != sizeof(FParameters))
+	{
+		return false;
+	}
+	if (GroomComponent == nullptr || Instances.Num() == 0)
+	{
+		return false;
+	}
+
+	return true;
+}
+
+void FOptimusGroomExecDataProviderProxy::GatherDispatchData(FDispatchData const& InDispatchData)
+{
+	const TStridedView<FParameters> ParameterArray = MakeStridedParameterView<FParameters>(InDispatchData);
+	for (int32 InvocationIndex = 0; InvocationIndex < ParameterArray.Num(); ++InvocationIndex)
+	{
+		FHairGroupInstance* Instance = Instances[InvocationIndex];
+		const int32 NumControlPoints = Instance->Strands.GetData().GetNumPoints();
+		const int32 NumCurves = Instance->Strands.GetData().GetNumCurves();
+		const int32 NumThreads = Domain == EOptimusGroomExecDomain::ControlPoint ? NumControlPoints : NumCurves;
+
+		FParameters& Parameters = ParameterArray[InvocationIndex];
+		Parameters.NumThreads = FIntVector(NumThreads, 1, 1);
+	}
+}
diff --git a/Engine/Plugins/Runtime/HairStrands/Source/HairStrandsDeformer/Private/DeformerDataInterfaceGroomExec.h b/Engine/Plugins/Runtime/HairStrands/Source/HairStrandsDeformer/Private/DeformerDataInterfaceGroomExec.h
new file mode 100644
--- /dev/null
+++ b/Engine/Plugins/Runtime/HairStrands/Source/HairStrandsDeformer/Private/DeformerDataInterfaceGroomExec.h
@@ -0,0 +1,98 @@
+// Copyright Epic Games, Inc. All Rights Reserved.
+
+#pragma once
+
+#include "IOptimusDeprecatedExecutionDataInterface.h"
+#include "OptimusComputeDataInterface.h"
+#include "ComputeFramework/ComputeDataProvider.h"
+#include "DeformerDataInterfaceGroomExec.generated.h"
+
+class FGroomExecDataInterfaceParameters;
+class UGroomComponent;
+struct FHairGroupInstance;
+
+UENUM()
+enum class EOptimusGroomExecDomain : uint8
+{
+	None = 0 UMETA(Hidden),
+	/** Run kernel with one thread per control point. */
+	ControlPoint = 1,
+	/** Run kernel with one thread per curve. */
+	Curve,
+};
+
+/** Compute Framework Data Interface for executing kernels over a skinned mesh domain. */
+UCLASS(Category = ComputeFramework)
+class UOptimusGroomExecDataInterface :
+	public UOptimusComputeDataInterface,
+	public IOptimusDeprecatedExecutionDataInterface
+{
+	GENERATED_BODY()
+
+public:
+	//~ Begin UOptimusComputeDataInterface Interface
+	FString GetDisplayName() const override;
+	FName GetCategory() const override;
+	bool IsVisible() const override {return false;};
+	TArray<FOptimusCDIPinDefinition> GetPinDefinitions() const override;
+	TSubclassOf<UActorComponent> GetRequiredComponentClass() const override;
+	//~ End UOptimusComputeDataInterface Interface
+	
+	//~ Begin UComputeDataInterface Interface
+	TCHAR const* GetClassName() const override { return TEXT("GroomExec"); }
+	bool IsExecutionInterface() const override { return true; }
+	void GetSupportedInputs(TArray<FShaderFunctionDefinition>& OutFunctions) const override;
+	void GetShaderParameters(TCHAR const* UID, FShaderParametersMetadataBuilder& InOutBuilder, FShaderParametersMetadataAllocations& InOutAllocations) const override;
+	TCHAR const* GetShaderVirtualPath() const override;
+	void GetShaderHash(FString& InOutKey) const override;
+	void GetHLSL(FString& OutHLSL, FString const& InDataInterfaceName) const override;
+	UComputeDataProvider* CreateDataProvider(TObjectPtr<UObject> InBinding, uint64 InInputMask, uint64 InOutputMask) const override;
+	//~ End UComputeDataInterface Interface
+
+	//~ Begin IOptimusDeprecatedExecutionDataInterface Interface
+	FName GetSelectedExecutionDomainName() const override;
+	//~ End IOptimusDeprecatedExecutionDataInterface Interface
+	
+	UPROPERTY(EditAnywhere, Category = Execution)
+	EOptimusGroomExecDomain Domain = EOptimusGroomExecDomain::ControlPoint;
+
+private:
+	static TCHAR const* TemplateFilePath;
+};
+
+/** Compute Framework Data Provider for executing kernels over a skinned mesh domain. */
+UCLASS(BlueprintType, editinlinenew, Category = ComputeFramework)
+class UOptimusGroomExecDataProvider : public UComputeDataProvider
+{
+	GENERATED_BODY()
+
+public:
+	UPROPERTY()
+	TObjectPtr<UGroomComponent> GroomComponent = nullptr;
+
+	UPROPERTY()
+	EOptimusGroomExecDomain Domain = EOptimusGroomExecDomain::ControlPoint;
+
+	//~ Begin UComputeDataProvider Interface
+	FComputeDataProviderRenderProxy* GetRenderProxy() override;
+	//~ End UComputeDataProvider Interface
+};
+
+class FOptimusGroomExecDataProviderProxy : public FComputeDataProviderRenderProxy
+{
+public:
+	FOptimusGroomExecDataProviderProxy(UGroomComponent* InGroomComponent, EOptimusGroomExecDomain InDomain);
+
+	//~ Begin FComputeDataProviderRenderProxy Interface
+	int32 GetDispatchThreadCount(TArray<FIntVector>& ThreadCounts) const override;
+	bool IsValid(FValidationData const& InValidationData) const override;
+	void GatherDispatchData(FDispatchData const& InDispatchData) override;
+	//~ End FComputeDataProviderRenderProxy Interface
+
+private:
+	using FParameters = FGroomExecDataInterfaceParameters;
+
+	UGroomComponent* GroomComponent = nullptr;
+	TArray<FHairGroupInstance*> Instances;
+	EOptimusGroomExecDomain Domain = EOptimusGroomExecDomain::ControlPoint;
+};
diff --git a/Engine/Plugins/Runtime/HairStrands/Source/HairStrandsDeformer/Private/DeformerDataInterfaceGroomGuide.cpp b/Engine/Plugins/Runtime/HairStrands/Source/HairStrandsDeformer/Private/DeformerDataInterfaceGroomGuide.cpp
new file mode 100644
--- /dev/null
+++ b/Engine/Plugins/Runtime/HairStrands/Source/HairStrandsDeformer/Private/DeformerDataInterfaceGroomGuide.cpp
@@ -0,0 +1,206 @@
+// Copyright Epic Games, Inc. All Rights Reserved.
+
+#include "DeformerDataInterfaceGroomGuide.h"
+
+#include "ComputeFramework/ShaderParamTypeDefinition.h"
+#include "ComputeFramework/ComputeKernelPermutationVector.h"
+#include "OptimusDataDomain.h"
+#include "ShaderCompilerCore.h"
+#include "ShaderParameterMetadataBuilder.h"
+#include "GroomComponent.h"
+#include "GroomInstance.h"
+#include "RenderGraphBuilder.h"
+#include "GlobalRenderResources.h"
+#include "HairStrandsDefinitions.h"
+#include "DeformerGroomComponentSource.h"
+#include "HairStrandsInterpolation.h"
+#include "SystemTextures.h"
+
+FString UOptimusGroomGuideDataInterface::GetDisplayName() const
+{
+	return TEXT("Guides");
+}
+
+TArray<FOptimusCDIPinDefinition> UOptimusGroomGuideDataInterface::GetPinDefinitions() const
+{
+	FName ControlPoint(UOptimusGroomComponentSource::Domains::ControlPoint);
+	FName Curve(UOptimusGroomComponentSource::Domains::Curve);
+
+	TArray<FOptimusCDIPinDefinition> Defs;
+	Defs.Add({ "NumGuidePoints",   "ReadNumGuidePoints" });
+	Defs.Add({ "NumGuideCurves",   "ReadNumGuideCurves" });
+	Defs.Add({ "Position",         "ReadPosition",          ControlPoint,   "ReadPosition" });
+	Defs.Add({ "CurveOffsetPoint", "ReadCurveOffsetPoint",  Curve,          "ReadCurveOffsetPoint" });
+	Defs.Add({ "CurveNumPoint",    "ReadCurveNumPoint",     Curve,          "ReadCurveNumPoint" });
+	return Defs;
+}
+ 
+TSubclassOf<UActorComponent> UOptimusGroomGuideDataInterface::GetRequiredComponentClass() const
+{
+	return UGroomComponent::StaticClass();
+}
+
+void UOptimusGroomGuideDataInterface::GetSupportedInputs(TArray<FShaderFunctionDefinition>& OutFunctions) const
+{
+	OutFunctions.AddDefaulted_GetRef()
+		.SetName(TEXT("ReadNumGuidePoints"))
+		.AddReturnType(EShaderFundamentalType::Uint);
+
+	OutFunctions.AddDefaulted_GetRef()
+		.SetName(TEXT("ReadNumGuideCurves"))
+		.AddReturnType(EShaderFundamentalType::Uint);
+
+	OutFunctions.AddDefaulted_GetRef()
+		.SetName(TEXT("ReadPosition"))
+		.AddReturnType(EShaderFundamentalType::Float, 3)
+		.AddParam(EShaderFundamentalType::Uint);
+
+	OutFunctions.AddDefaulted_GetRef()
+		.SetName(TEXT("ReadCurveOffsetPoint"))
+		.AddReturnType(EShaderFundamentalType::Uint)
+		.AddParam(EShaderFundamentalType::Uint);
+
+	OutFunctions.AddDefaulted_GetRef()
+		.SetName(TEXT("ReadCurveNumPoint"))
+		.AddReturnType(EShaderFundamentalType::Uint)
+		.AddParam(EShaderFundamentalType::Uint);
+}
+
+BEGIN_SHADER_PARAMETER_STRUCT(FGroomGuideDataInterfaceParameters, )
+	SHADER_PARAMETER_STRUCT_INCLUDE(FHairStrandsInstanceCommonParameters, Common)
+	SHADER_PARAMETER_STRUCT_INCLUDE(FHairStrandsInstanceResourceParameters, Resources)
+END_SHADER_PARAMETER_STRUCT()
+
+void UOptimusGroomGuideDataInterface::GetShaderParameters(TCHAR const* UID, FShaderParametersMetadataBuilder& InOutBuilder, FShaderParametersMetadataAllocations& InOutAllocations) const
+{
+	InOutBuilder.AddNestedStruct<FGroomGuideDataInterfaceParameters>(UID);
+}
+
+TCHAR const* UOptimusGroomGuideDataInterface::TemplateFilePath = TEXT("/Plugin/Runtime/HairStrands/Private/DeformerDataInterfaceGroomGuide.ush");
+
+TCHAR const* UOptimusGroomGuideDataInterface::GetShaderVirtualPath() const
+{
+	return TemplateFilePath;
+}
+
+void UOptimusGroomGuideDataInterface::GetShaderHash(FString& InOutKey) const
+{
+	GetShaderFileHash(TemplateFilePath, EShaderPlatform::SP_PCD3D_SM5).AppendString(InOutKey);
+}
+
+void UOptimusGroomGuideDataInterface::GetHLSL(FString& OutHLSL, FString const& InDataInterfaceName) const
+{
+	TMap<FString, FStringFormatArg> TemplateArgs =
+	{
+		{ TEXT("DataInterfaceName"), InDataInterfaceName },
+	};
+
+	FString TemplateFile;
+	LoadShaderSourceFile(TemplateFilePath, EShaderPlatform::SP_PCD3D_SM5, &TemplateFile, nullptr);
+	OutHLSL += FString::Format(*TemplateFile, TemplateArgs);
+}
+
+UComputeDataProvider* UOptimusGroomGuideDataInterface::CreateDataProvider(TObjectPtr<UObject> InBinding, uint64 InInputMask, uint64 InOutputMask) const
+{
+	UOptimusGroomGuideDataProvider* Provider = NewObject<UOptimusGroomGuideDataProvider>();
+	Provider->Groom = Cast<UGroomComponent>(InBinding);
+	return Provider;
+}
+
+
+FComputeDataProviderRenderProxy* UOptimusGroomGuideDataProvider::GetRenderProxy()
+{
+	return new FOptimusGroomGuideDataProviderProxy(Groom);
+}
+
+
+FOptimusGroomGuideDataProviderProxy::FOptimusGroomGuideDataProviderProxy(UGroomComponent* InGroomComponent)
+{
+	GroomComponent = InGroomComponent;
+}
+
+bool FOptimusGroomGuideDataProviderProxy::IsValid(FValidationData const& InValidationData) const
+{
+	if (InValidationData.ParameterStructSize != sizeof(FParameters))
+	{
+		return false;
+	}
+	if (GroomComponent == nullptr)
+	{
+		return false;
+	}
+	const uint32 GroupCount = GroomComponent->GetGroupCount();
+	if (InValidationData.NumInvocations != GroupCount)
+	{
+		return false;
+	}
+	// Earlier invalidation in case one of the instance is not ready/valid
+	for (uint32 GroupIndex = 0; GroupIndex < GroupCount; ++GroupIndex)
+	{
+		if (GroomComponent->GetGroupInstance(GroupIndex) == nullptr)
+		{
+			return false;
+		}
+	}
+	
+	return true;
+}
+
+void FOptimusGroomGuideDataProviderProxy::AllocateResources(FRDGBuilder& GraphBuilder)
+{
+	Resources.Empty();
+	FallbackByteAddressSRV = nullptr;
+	FallbackStructuredSRV = nullptr;
+	const uint32 InstanceCount = GroomComponent ? GroomComponent->GetGroupCount() : 0;
+	for (uint32 Index =0; Index <InstanceCount;++Index)
+	{
+		if (FHairGroupInstance* Instance = GroomComponent->GetGroupInstance(Index))
+		{
+			FHairStrandsInstanceResourceParameters& R = Resources.AddDefaulted_GetRef();
+			if (Instance->Guides.RestResource)
+			{
+				R.PositionBuffer		= RegisterAsSRV(GraphBuilder, Instance->Guides.RestResource->PositionBuffer);
+				R.PositionOffsetBuffer 	= RegisterAsSRV(GraphBuilder, Instance->Guides.RestResource->PositionOffsetBuffer);
+				R.CurveBuffer			= RegisterAsSRV(GraphBuilder, Instance->Guides.RestResource->CurveBuffer);
+				R.PointToCurveBuffer	= RegisterAsSRV(GraphBuilder, Instance->Guides.RestResource->PointToCurveBuffer);
+			}
+
+			if (!FallbackByteAddressSRV){ FallbackByteAddressSRV = GraphBuilder.CreateSRV(GSystemTextures.GetDefaultByteAddressBuffer(GraphBuilder, 16u)); }
+			if (!FallbackStructuredSRV) { FallbackStructuredSRV  = GraphBuilder.CreateSRV(GSystemTextures.GetDefaultStructuredBuffer(GraphBuilder, 16u)); }
+		}
+	}
+}
+
+void FOptimusGroomGuideDataProviderProxy::GatherDispatchData(FDispatchData const& InDispatchData)
+{
+	const uint32 InstanceCount = GroomComponent ? GroomComponent->GetGroupCount() : 0;
+	check(InDispatchData.NumInvocations == InstanceCount);
+	check(uint32(Resources.Num()) == InstanceCount);
+
+	const TStridedView<FParameters> ParameterArray = MakeStridedParameterView<FParameters>(InDispatchData);
+	for (int32 InvocationIndex = 0; InvocationIndex < ParameterArray.Num(); ++InvocationIndex)
+	{
+		// Default valid in case the instance is not valid
+		FParameters& Parameters = ParameterArray[InvocationIndex];
+		Parameters.Common.PointCount 				= 0;
+		Parameters.Common.CurveCount 				= 0;
+		Parameters.Resources.PositionBuffer			= FallbackByteAddressSRV;
+		Parameters.Resources.PositionOffsetBuffer 	= FallbackStructuredSRV;
+		Parameters.Resources.PointToCurveBuffer		= FallbackByteAddressSRV;
+		Parameters.Resources.CurveBuffer			= FallbackByteAddressSRV;
+
+		if (FHairGroupInstance* Instance = GroomComponent->GetGroupInstance(InvocationIndex))
+		{
+			const bool bIsSRVValid = Resources[InvocationIndex].PositionBuffer != nullptr;
+			if (bIsSRVValid)
+			{
+				// Reuse the strands properties for guides, and apply guides values
+				const FHairGroupPublicData::FVertexFactoryInput VFInput = ComputeHairStrandsVertexInputData(Instance, EGroomViewMode::None);
+				Parameters.Common 				= VFInput.Strands.Common;
+				Parameters.Common.PointCount 	= Instance->Guides.IsValid() ? Instance->Guides.GetData().Header.PointCount : 0;
+				Parameters.Common.CurveCount 	= Instance->Guides.IsValid() ? Instance->Guides.GetData().Header.CurveCount : 0;
+				Parameters.Resources 			= Resources[InvocationIndex];
+			}
+		}
+	}
+}
\ No newline at end of file
diff --git a/Engine/Plugins/Runtime/HairStrands/Source/HairStrandsDeformer/Private/DeformerDataInterfaceGroomGuide.h b/Engine/Plugins/Runtime/HairStrands/Source/HairStrandsDeformer/Private/DeformerDataInterfaceGroomGuide.h
new file mode 100644
--- /dev/null
+++ b/Engine/Plugins/Runtime/HairStrands/Source/HairStrandsDeformer/Private/DeformerDataInterfaceGroomGuide.h
@@ -0,0 +1,74 @@
+// Copyright Epic Games, Inc. All Rights Reserved.
+
+#pragma once
+
+#include "RenderGraphDefinitions.h"
+#include "OptimusComputeDataInterface.h"
+#include "HairStrandsInterface.h"
+#include "ComputeFramework/ComputeDataProvider.h"
+#include "DeformerDataInterfaceGroomGuide.generated.h"
+
+class FGroomGuideDataInterfaceParameters;
+class UGroomComponent;
+
+/** Compute Framework Data Interface for reading groom guides. */
+UCLASS(Category = ComputeFramework)
+class UOptimusGroomGuideDataInterface : public UOptimusComputeDataInterface
+{
+	GENERATED_BODY()
+
+public:
+	//~ Begin UOptimusComputeDataInterface Interface
+	FString GetDisplayName() const override;
+	TArray<FOptimusCDIPinDefinition> GetPinDefinitions() const override;
+	TSubclassOf<UActorComponent> GetRequiredComponentClass() const override;
+	//~ End UOptimusComputeDataInterface Interface
+	
+	//~ Begin UComputeDataInterface Interface
+	TCHAR const* GetClassName() const override { return TEXT("Groom"); }
+	void GetSupportedInputs(TArray<FShaderFunctionDefinition>& OutFunctions) const override;
+	void GetShaderParameters(TCHAR const* UID, FShaderParametersMetadataBuilder& InOutBuilder, FShaderParametersMetadataAllocations& InOutAllocations) const override;
+	TCHAR const* GetShaderVirtualPath() const override;
+	void GetShaderHash(FString& InOutKey) const override;
+	void GetHLSL(FString& OutHLSL, FString const& InDataInterfaceName) const override;
+	UComputeDataProvider* CreateDataProvider(TObjectPtr<UObject> InBinding, uint64 InInputMask, uint64 InOutputMask) const override;
+	//~ End UComputeDataInterface Interface
+
+private:
+	static TCHAR const* TemplateFilePath;
+};
+
+/** Compute Framework Data Provider for reading groom. */
+UCLASS(BlueprintType, editinlinenew, Category = ComputeFramework)
+class UOptimusGroomGuideDataProvider : public UComputeDataProvider
+{
+	GENERATED_BODY()
+
+public:
+	UPROPERTY(EditAnywhere, BlueprintReadWrite, Category = Binding)
+	TObjectPtr<UGroomComponent> Groom = nullptr;
+
+	//~ Begin UComputeDataProvider Interface
+	FComputeDataProviderRenderProxy* GetRenderProxy() override;
+	//~ End UComputeDataProvider Interface
+};
+
+class FOptimusGroomGuideDataProviderProxy : public FComputeDataProviderRenderProxy
+{
+public:
+	FOptimusGroomGuideDataProviderProxy(UGroomComponent* InGroomComponent);
+
+	//~ Begin FComputeDataProviderRenderProxy Interface
+	bool IsValid(FValidationData const& InValidationData) const override;
+	void AllocateResources(FRDGBuilder& GraphBuilder) override;
+	void GatherDispatchData(FDispatchData const& InDispatchData) override;
+	//~ End FComputeDataProviderRenderProxy Interface
+
+private:
+	using FParameters = FGroomGuideDataInterfaceParameters;
+
+	UGroomComponent* GroomComponent = nullptr;
+	TArray<FHairStrandsInstanceResourceParameters> Resources;
+	FRDGBufferSRVRef FallbackByteAddressSRV;
+	FRDGBufferSRVRef FallbackStructuredSRV;
+};
diff --git a/Engine/Plugins/Runtime/HairStrands/Source/HairStrandsDeformer/Private/DeformerDataInterfaceGroomWrite.cpp b/Engine/Plugins/Runtime/HairStrands/Source/HairStrandsDeformer/Private/DeformerDataInterfaceGroomWrite.cpp
new file mode 100644
--- /dev/null
+++ b/Engine/Plugins/Runtime/HairStrands/Source/HairStrandsDeformer/Private/DeformerDataInterfaceGroomWrite.cpp
@@ -0,0 +1,286 @@
+// Copyright Epic Games, Inc. All Rights Reserved.
+
+#include "DeformerDataInterfaceGroomWrite.h"
+
+#include "ComputeFramework/ShaderParamTypeDefinition.h"
+#include "OptimusDataDomain.h"
+#include "RenderGraphBuilder.h"
+#include "Rendering/SkeletalMeshLODRenderData.h"
+#include "Rendering/SkeletalMeshRenderData.h"
+#include "ShaderParameterMetadataBuilder.h"
+#include "SkeletalMeshDeformerHelpers.h"
+#include "SkeletalRenderPublic.h"
+#include "GroomComponent.h"
+#include "GroomInstance.h"
+#include "DeformerGroomComponentSource.h"
+#include "RenderGraphUtils.h"
+#include "SystemTextures.h"
+#include "HairStrandsInterpolation.h"
+
+FString UOptimusGroomWriteDataInterface::GetDisplayName() const
+{
+	return TEXT("Write Groom");
+}
+
+FName UOptimusGroomWriteDataInterface::GetCategory() const
+{
+	return CategoryName::OutputDataInterfaces;
+}
+
+TArray<FOptimusCDIPinDefinition> UOptimusGroomWriteDataInterface::GetPinDefinitions() const
+{
+	FName ControlPoint(UOptimusGroomComponentSource::Domains::ControlPoint);
+	FName Curve(UOptimusGroomComponentSource::Domains::Curve);
+
+	TArray<FOptimusCDIPinDefinition> Defs;
+	Defs.Add({ "Position",         "WritePosition",         	ControlPoint, "ReadNumControlPoints" });
+	Defs.Add({ "Radius",           "WriteRadius",           	ControlPoint, "ReadNumControlPoints" });
+	Defs.Add({ "PositionAndRadius","WritePositionAndRadius",	ControlPoint, "ReadNumControlPoints" });
+
+	Defs.Add({ "RootUV",			"WriteRootUV",   			Curve, 		  "ReadNumCurves" });
+	Defs.Add({ "Seed",				"WriteSeed",   				Curve, 		  "ReadNumCurves" });
+	Defs.Add({ "ClumpId",			"WriteClumpId",   			Curve, 		  "ReadNumCurves" });
+	Defs.Add({ "Color",				"WriteColor",   			ControlPoint, "ReadNumControlPoints" });
+	Defs.Add({ "Roughness",			"WriteRoughness",   		ControlPoint, "ReadNumControlPoints" });
+	Defs.Add({ "AO",				"WriteAO",   				ControlPoint, "ReadNumControlPoints" });
+
+	return Defs;
+}
+
+TSubclassOf<UActorComponent> UOptimusGroomWriteDataInterface::GetRequiredComponentClass() const
+{
+	return UGroomComponent::StaticClass();
+}
+
+void UOptimusGroomWriteDataInterface::GetSupportedInputs(TArray<FShaderFunctionDefinition>& OutFunctions) const
+{
+	OutFunctions.AddDefaulted_GetRef()
+		.SetName(TEXT("ReadNumControlPoints"))
+		.AddReturnType(EShaderFundamentalType::Uint);
+
+	OutFunctions.AddDefaulted_GetRef()
+		.SetName(TEXT("ReadNumCurves"))
+		.AddReturnType(EShaderFundamentalType::Uint);
+}
+
+void UOptimusGroomWriteDataInterface::GetSupportedOutputs(TArray<FShaderFunctionDefinition>& OutFunctions) const
+{
+	OutFunctions.AddDefaulted_GetRef()
+		.SetName(TEXT("WritePosition"))
+		.AddParam(EShaderFundamentalType::Uint)
+		.AddParam(EShaderFundamentalType::Float, 3);
+
+	OutFunctions.AddDefaulted_GetRef()
+		.SetName(TEXT("WriteRadius"))
+		.AddParam(EShaderFundamentalType::Uint)
+		.AddParam(EShaderFundamentalType::Float);
+
+	OutFunctions.AddDefaulted_GetRef()
+		.SetName(TEXT("WritePositionAndRadius"))
+		.AddParam(EShaderFundamentalType::Uint)
+		.AddParam(EShaderFundamentalType::Float, 4);
+
+	OutFunctions.AddDefaulted_GetRef()
+		.SetName(TEXT("WriteRootUV"))
+		.AddParam(EShaderFundamentalType::Uint)
+		.AddParam(EShaderFundamentalType::Float, 2);
+
+	OutFunctions.AddDefaulted_GetRef()
+		.SetName(TEXT("WriteSeed"))
+		.AddParam(EShaderFundamentalType::Uint)
+		.AddParam(EShaderFundamentalType::Float);
+
+	OutFunctions.AddDefaulted_GetRef()
+		.SetName(TEXT("WriteClumpId"))
+		.AddParam(EShaderFundamentalType::Uint)
+		.AddParam(EShaderFundamentalType::Float);
+
+	OutFunctions.AddDefaulted_GetRef()
+		.SetName(TEXT("WriteColor"))
+		.AddParam(EShaderFundamentalType::Uint)
+		.AddParam(EShaderFundamentalType::Float, 3);
+
+	OutFunctions.AddDefaulted_GetRef()
+		.SetName(TEXT("WriteRoughness"))
+		.AddParam(EShaderFundamentalType::Uint)
+		.AddParam(EShaderFundamentalType::Float);
+
+	OutFunctions.AddDefaulted_GetRef()
+		.SetName(TEXT("WriteAO"))
+		.AddParam(EShaderFundamentalType::Uint)
+		.AddParam(EShaderFundamentalType::Float);
+}
+
+BEGIN_SHADER_PARAMETER_STRUCT(FGroomWriteDataInterfaceParameters, )
+	SHADER_PARAMETER_STRUCT_INCLUDE(FHairStrandsInstanceCommonParameters, Common)
+	SHADER_PARAMETER(uint32, OutputStreamStart)
+	SHADER_PARAMETER_RDG_BUFFER_SRV(StructuredBuffer, PositionOffsetBufferSRV)
+	SHADER_PARAMETER_RDG_BUFFER_SRV(ByteAddressBuffer, PositionBufferSRV)
+	SHADER_PARAMETER_RDG_BUFFER_UAV(RWByteAddressBuffer, PositionBufferUAV)
+	SHADER_PARAMETER_RDG_BUFFER_UAV(RWByteAddressBuffer, CurveAttributeBufferUAV)
+	SHADER_PARAMETER_RDG_BUFFER_UAV(RWByteAddressBuffer, PointAttributeBufferUAV)
+END_SHADER_PARAMETER_STRUCT()
+
+void UOptimusGroomWriteDataInterface::GetShaderParameters(TCHAR const* UID, FShaderParametersMetadataBuilder& InOutBuilder, FShaderParametersMetadataAllocations& InOutAllocations) const
+{
+	InOutBuilder.AddNestedStruct<FGroomWriteDataInterfaceParameters>(UID);
+}
+
+TCHAR const* UOptimusGroomWriteDataInterface::TemplateFilePath = TEXT("/Plugin/Runtime/HairStrands/Private/DeformerDataInterfaceGroomWrite.ush");
+
+TCHAR const* UOptimusGroomWriteDataInterface::GetShaderVirtualPath() const
+{
+	return TemplateFilePath;
+}
+
+void UOptimusGroomWriteDataInterface::GetShaderHash(FString& InOutKey) const
+{
+	GetShaderFileHash(TemplateFilePath, EShaderPlatform::SP_PCD3D_SM5).AppendString(InOutKey);
+}
+
+void UOptimusGroomWriteDataInterface::GetHLSL(FString& OutHLSL, FString const& InDataInterfaceName) const
+{
+	TMap<FString, FStringFormatArg> TemplateArgs =
+	{
+		{ TEXT("DataInterfaceName"), InDataInterfaceName },
+	};
+
+	FString TemplateFile;
+	LoadShaderSourceFile(TemplateFilePath, EShaderPlatform::SP_PCD3D_SM5, &TemplateFile, nullptr);
+	OutHLSL += FString::Format(*TemplateFile, TemplateArgs);
+}
+
+UComputeDataProvider* UOptimusGroomWriteDataInterface::CreateDataProvider(TObjectPtr<UObject> InBinding, uint64 InInputMask, uint64 InOutputMask) const
+{
+	UOptimusGroomWriteDataProvider* Provider = NewObject<UOptimusGroomWriteDataProvider>();
+	Provider->GroomComponent = Cast<UGroomComponent>(InBinding);
+	Provider->OutputMask = InOutputMask;
+	return Provider;
+}
+
+
+FComputeDataProviderRenderProxy* UOptimusGroomWriteDataProvider::GetRenderProxy()
+{
+	return new FOptimusGroomWriteDataProviderProxy(GroomComponent, OutputMask);
+}
+
+
+FOptimusGroomWriteDataProviderProxy::FOptimusGroomWriteDataProviderProxy(UGroomComponent* InGroomComponent, uint64 InOutputMask)
+{
+	OutputMask = InOutputMask;
+	const uint32 InstanceCount = InGroomComponent ? InGroomComponent->GetGroupCount() : 0;
+	for (uint32 Index = 0; Index < InstanceCount; ++Index)
+	{
+		Instances.Add(InGroomComponent->GetGroupInstance(Index));
+	}
+}
+
+bool FOptimusGroomWriteDataProviderProxy::IsValid(FValidationData const& InValidationData) const
+{
+	if (InValidationData.ParameterStructSize != sizeof(FParameters))
+	{
+		return false;
+	}
+	if (InValidationData.NumInvocations != Instances.Num())
+	{
+		return false;
+	}
+
+	return true;
+}
+
+void FOptimusGroomWriteDataProviderProxy::AllocateResources(FRDGBuilder& GraphBuilder)
+{
+	for (FHairGroupInstance* GroomInstance : Instances)
+	{
+		if (GroomInstance)
+		{
+			// Allocate required buffers
+			const int32 NumControlPoints = GroomInstance->Strands.GetData().GetNumPoints();
+			const int32 NumCurves = GroomInstance->Strands.GetData().GetNumCurves();
+
+			FResources& R = Resources.AddDefaulted_GetRef();
+
+			R.PositionBufferSRV_fallback = GraphBuilder.CreateSRV(GSystemTextures.GetDefaultByteAddressBuffer(GraphBuilder, 16u));
+			R.PositionBufferUAV_fallback = GraphBuilder.CreateUAV(GraphBuilder.CreateBuffer(FRDGBufferDesc::CreateByteAddressDesc(16u), TEXT("Groom.Deformer.FallbackDeformedPositionBuffer")), ERDGUnorderedAccessViewFlags::SkipBarrier);
+			R.AttributeBufferUAV_fallback = GraphBuilder.CreateUAV(GraphBuilder.CreateBuffer(FRDGBufferDesc::CreateByteAddressDesc(16u), TEXT("Groom.Deformer.FallbackDeformedAttributeBuffer")), ERDGUnorderedAccessViewFlags::SkipBarrier);
+
+			// Positions / Radius
+			if (OutputMask & 0x7)
+			{
+				R.PositionOffsetBufferSRV 	= Register(GraphBuilder, GroomInstance->Strands.RestResource->PositionOffsetBuffer, ERDGImportedBufferFlags::CreateSRV).SRV;
+				R.PositionBufferSRV 		= Register(GraphBuilder, GroomInstance->Strands.RestResource->PositionBuffer, ERDGImportedBufferFlags::CreateSRV).SRV;
+				R.PositionBufferUAV 		= Register(GraphBuilder, GroomInstance->Strands.DeformedResource->GetDeformerBuffer(GraphBuilder), ERDGImportedBufferFlags::CreateUAV).UAV;
+			}
+			else
+			{
+				R.PositionOffsetBufferSRV 	= GraphBuilder.CreateSRV(GSystemTextures.GetDefaultStructuredBuffer(GraphBuilder, 16u), PF_A32B32G32R32F);
+				R.PositionBufferSRV 		= GraphBuilder.CreateSRV(GSystemTextures.GetDefaultByteAddressBuffer(GraphBuilder, 16u));
+				R.PositionBufferUAV 		= GraphBuilder.CreateUAV(GraphBuilder.CreateBuffer(FRDGBufferDesc::CreateByteAddressDesc(16u), TEXT("Hair.Deformer.DummyBuffer")));
+			}
+			
+			// Curve Attributes
+			if (OutputMask & 0x38)
+			{
+				FRDGExternalBuffer CurveAttributeBufferExt = GroomInstance->Strands.DeformedResource->GetDeformerCurveAttributeBuffer(GraphBuilder);
+				if (CurveAttributeBufferExt.Buffer)
+				{
+					// Always copy the attributes from the rest asset, so that if deformer write different attribute at different tick, it all remains consistent with the source data
+					FRDGImportedBuffer DstCurveAttributeBuffer = Register(GraphBuilder, CurveAttributeBufferExt, ERDGImportedBufferFlags::CreateUAV);
+					FRDGImportedBuffer SrcCurveAttributeBuffer = Register(GraphBuilder, GroomInstance->Strands.RestResource->CurveAttributeBuffer, ERDGImportedBufferFlags::None);
+					if (SrcCurveAttributeBuffer.Buffer)
+					{
+						AddCopyBufferPass(GraphBuilder, DstCurveAttributeBuffer.Buffer, SrcCurveAttributeBuffer.Buffer);
+					}
+					R.CurveAttributeBufferUAV = DstCurveAttributeBuffer.UAV;
+				}
+			}
+			if (R.CurveAttributeBufferUAV == nullptr)
+			{
+				R.CurveAttributeBufferUAV = R.AttributeBufferUAV_fallback;
+			}
+
+			// Point Attributes
+			if (OutputMask & 0x1c0)
+			{
+				FRDGExternalBuffer PointAttributeBufferExt = GroomInstance->Strands.DeformedResource->GetDeformerPointAttributeBuffer(GraphBuilder);
+				if (PointAttributeBufferExt.Buffer)
+				{
+					FRDGImportedBuffer PointAttributeBuffer = Register(GraphBuilder, PointAttributeBufferExt, ERDGImportedBufferFlags::CreateUAV);
+					AddCopyBufferPass(GraphBuilder, PointAttributeBuffer.Buffer, Register(GraphBuilder, GroomInstance->Strands.RestResource->PointAttributeBuffer, ERDGImportedBufferFlags::None).Buffer);
+					R.PointAttributeBufferUAV = PointAttributeBuffer.UAV;
+				}
+			}
+			if (R.PointAttributeBufferUAV == nullptr)
+			{
+				R.PointAttributeBufferUAV = R.AttributeBufferUAV_fallback;
+				//R.PointAttributeBufferUAV = GraphBuilder.CreateUAV(GraphBuilder.RegisterExternalBuffer(GWhiteVertexBufferWithRDG->Buffer), PF_R32_UINT);
+			}
+		}	
+	}
+}
+
+void FOptimusGroomWriteDataProviderProxy::GatherDispatchData(FDispatchData const& InDispatchData)
+{
+	const TStridedView<FParameters> ParameterArray = MakeStridedParameterView<FParameters>(InDispatchData);
+	for (int32 InvocationIndex = 0; InvocationIndex < ParameterArray.Num(); ++InvocationIndex)
+	{
+		if (FHairGroupInstance* GroomInstance = Instances[InvocationIndex])
+		{
+			FResources& Resource = Resources[InvocationIndex];
+
+			const bool bValid = Resource.PositionBufferUAV != nullptr && Resource.PositionBufferSRV != nullptr && Resource.PositionOffsetBufferSRV != nullptr;
+
+			const FHairGroupPublicData::FVertexFactoryInput VFInput = ComputeHairStrandsVertexInputData(GroomInstance, EGroomViewMode::None);		
+
+			FParameters& Parameters = ParameterArray[InvocationIndex];
+			Parameters.Common 					= VFInput.Strands.Common;
+			Parameters.OutputStreamStart 		= 0;
+			Parameters.PositionOffsetBufferSRV 	= bValid ? Resource.PositionOffsetBufferSRV : Resource.PositionBufferSRV_fallback;
+			Parameters.PositionBufferSRV 		= bValid ? Resource.PositionBufferSRV : Resource.PositionBufferSRV_fallback;
+			Parameters.PositionBufferUAV 		= bValid ? Resource.PositionBufferUAV : Resource.PositionBufferUAV_fallback;
+			Parameters.PointAttributeBufferUAV 	= bValid ? Resource.PointAttributeBufferUAV : Resource.AttributeBufferUAV_fallback;
+			Parameters.CurveAttributeBufferUAV 	= bValid ? Resource.CurveAttributeBufferUAV : Resource.AttributeBufferUAV_fallback;
+		}
+	}
+}
\ No newline at end of file
diff --git a/Engine/Plugins/Runtime/HairStrands/Source/HairStrandsDeformer/Private/DeformerDataInterfaceGroomWrite.h b/Engine/Plugins/Runtime/HairStrands/Source/HairStrandsDeformer/Private/DeformerDataInterfaceGroomWrite.h
new file mode 100644
--- /dev/null
+++ b/Engine/Plugins/Runtime/HairStrands/Source/HairStrandsDeformer/Private/DeformerDataInterfaceGroomWrite.h
@@ -0,0 +1,92 @@
+// Copyright Epic Games, Inc. All Rights Reserved.
+
+#pragma once
+
+#include "OptimusComputeDataInterface.h"
+#include "ComputeFramework/ComputeDataProvider.h"
+#include "DeformerDataInterfaceGroomWrite.generated.h"
+
+class FGroomWriteDataInterfaceParameters;
+struct FHairGroupInstance;
+class FRDGBuffer;
+class FRDGBufferSRV;
+class FRDGBufferUAV;
+class UGroomComponent;
+
+/** Compute Framework Data Interface for writing skinned mesh. */
+UCLASS(Category = ComputeFramework)
+class UOptimusGroomWriteDataInterface : public UOptimusComputeDataInterface
+{
+	GENERATED_BODY()
+
+public:
+	//~ Begin UOptimusComputeDataInterface Interface
+	FString GetDisplayName() const override;
+	FName GetCategory() const override;
+	TArray<FOptimusCDIPinDefinition> GetPinDefinitions() const override;
+	TSubclassOf<UActorComponent> GetRequiredComponentClass() const override;
+	//~ End UOptimusComputeDataInterface Interface
+	
+	//~ Begin UComputeDataInterface Interface
+	TCHAR const* GetClassName() const override { return TEXT("GroomWrite"); }
+	void GetSupportedInputs(TArray<FShaderFunctionDefinition>& OutFunctions) const override;
+	void GetSupportedOutputs(TArray<FShaderFunctionDefinition>& OutFunctions) const override;
+	void GetShaderParameters(TCHAR const* UID, FShaderParametersMetadataBuilder& InOutBuilder, FShaderParametersMetadataAllocations& InOutAllocations) const override;
+	TCHAR const* GetShaderVirtualPath() const override;
+	void GetShaderHash(FString& InOutKey) const override;
+	void GetHLSL(FString& OutHLSL, FString const& InDataInterfaceName) const override;
+	UComputeDataProvider* CreateDataProvider(TObjectPtr<UObject> InBinding, uint64 InInputMask, uint64 InOutputMask) const override;
+	//~ End UComputeDataInterface Interface
+
+private:
+	static TCHAR const* TemplateFilePath;
+};
+
+/** Compute Framework Data Provider for writing skinned mesh. */
+UCLASS(BlueprintType, editinlinenew, Category = ComputeFramework)
+class UOptimusGroomWriteDataProvider : public UComputeDataProvider
+{
+	GENERATED_BODY()
+
+public:
+	UPROPERTY(EditAnywhere, BlueprintReadWrite, Category = Binding)
+	TObjectPtr<UGroomComponent> GroomComponent = nullptr;
+
+	uint64 OutputMask = 0;
+
+	//~ Begin UComputeDataProvider Interface
+	FComputeDataProviderRenderProxy* GetRenderProxy() override;
+	//~ End UComputeDataProvider Interface
+};
+
+class FOptimusGroomWriteDataProviderProxy : public FComputeDataProviderRenderProxy
+{
+public:
+	FOptimusGroomWriteDataProviderProxy(UGroomComponent* InGroomComponent, uint64 InOutputMask);
+
+	//~ Begin FComputeDataProviderRenderProxy Interface
+	bool IsValid(FValidationData const& InValidationData) const override;
+	void AllocateResources(FRDGBuilder& GraphBuilder) override;
+	void GatherDispatchData(FDispatchData const& InDispatchData) override;
+	//~ End FComputeDataProviderRenderProxy Interface
+
+private:
+	using FParameters = FGroomWriteDataInterfaceParameters;
+
+	TArray<FHairGroupInstance*> Instances;
+	uint64 OutputMask = 0;
+
+	struct FResources
+	{
+		FRDGBufferSRV* PositionOffsetBufferSRV = nullptr;
+		FRDGBufferSRV* PositionBufferSRV = nullptr;
+		FRDGBufferUAV* PositionBufferUAV = nullptr;
+		FRDGBufferSRV* PositionBufferSRV_fallback = nullptr;
+		FRDGBufferUAV* PositionBufferUAV_fallback = nullptr;
+
+		FRDGBufferUAV* PointAttributeBufferUAV = nullptr;
+		FRDGBufferUAV* CurveAttributeBufferUAV = nullptr;
+		FRDGBufferUAV* AttributeBufferUAV_fallback = nullptr;
+	};
+	TArray<FResources> Resources;
+};
\ No newline at end of file
diff --git a/Engine/Plugins/Runtime/HairStrands/Source/HairStrandsDeformer/Private/DeformerGroomAttributeRead.cpp b/Engine/Plugins/Runtime/HairStrands/Source/HairStrandsDeformer/Private/DeformerGroomAttributeRead.cpp
deleted file mode 100644
diff --git a/Engine/Plugins/Runtime/HairStrands/Source/HairStrandsDeformer/Private/DeformerGroomAttributeRead.h b/Engine/Plugins/Runtime/HairStrands/Source/HairStrandsDeformer/Private/DeformerGroomAttributeRead.h
deleted file mode 100644
diff --git a/Engine/Plugins/Runtime/HairStrands/Source/HairStrandsDeformer/Private/DeformerGroomComponentSource.cpp b/Engine/Plugins/Runtime/HairStrands/Source/HairStrandsDeformer/Private/DeformerGroomComponentSource.cpp
new file mode 100644
--- /dev/null
+++ b/Engine/Plugins/Runtime/HairStrands/Source/HairStrandsDeformer/Private/DeformerGroomComponentSource.cpp
@@ -0,0 +1,99 @@
+// Copyright Epic Games, Inc. All Rights Reserved.
+
+#include "DeformerGroomComponentSource.h"
+
+#include "GroomComponent.h"
+
+#define LOCTEXT_NAMESPACE "DeformersGroomComponentSource"
+
+FName UOptimusGroomComponentSource::Domains::ControlPoint("ControlPoint");
+FName UOptimusGroomComponentSource::Domains::Curve("Curve");
+
+FText UOptimusGroomComponentSource::GetDisplayName() const
+{
+	return LOCTEXT("GroomComponent", "Groom Component");
+}
+
+TSubclassOf<UActorComponent> UOptimusGroomComponentSource::GetComponentClass() const
+{
+	return UGroomComponent::StaticClass();
+}
+
+bool UOptimusGroomComponentSource::IsUsableAsPrimarySource() const
+{
+	return GetComponentClass()->IsChildOf<UGroomComponent>();
+}
+
+TArray<FName> UOptimusGroomComponentSource::GetExecutionDomains() const
+{
+	return { Domains::ControlPoint, Domains::Curve };
+}
+
+int32 UOptimusGroomComponentSource::GetLodIndex(const UActorComponent* InComponent) const
+{
+	return 0;
+}
+
+uint32 UOptimusGroomComponentSource::GetDefaultNumInvocations(
+	const UActorComponent* InComponent,
+	int32 InLod
+	) const
+{
+	const UGroomComponent* GroomComponent = Cast<UGroomComponent>(InComponent);
+	if (!GroomComponent)
+	{
+		return 0;
+	}
+
+	const uint32 GroupCount = GroomComponent->GetGroupCount();
+
+	return GroupCount;
+}
+
+bool UOptimusGroomComponentSource::GetComponentElementCountsForExecutionDomain(
+	FName InDomainName,
+	const UActorComponent* InComponent,
+	int32 InLodIndex,
+	TArray<int32>& OutInvocationElementCounts
+) const
+{
+	const UGroomComponent* GroomComponent = Cast<UGroomComponent>(InComponent);
+	if (!GroomComponent)
+	{
+		return false;
+	}
+
+	// Ensure all groups are valid
+	const uint32 GroupCount = GroomComponent->GetGroupCount();
+	for (uint32 GroupIt=0;GroupIt<GroupCount;++GroupIt)
+	{
+		if (!GroomComponent->GetGroupInstance(GroupIt))
+		{
+			return 0;
+		}
+	}
+
+	OutInvocationElementCounts.Reset();
+
+	if (InDomainName == Domains::ControlPoint || InDomainName == Domains::Curve)
+	{
+		const int32 NumInvocations = GroupCount;
+		OutInvocationElementCounts.Reserve(NumInvocations);
+		for (int32 InvocationIndex = 0; InvocationIndex < NumInvocations; ++InvocationIndex)
+		{
+			if (const FHairGroupInstance* Instance = GroomComponent->GetGroupInstance(InvocationIndex))
+			{
+				const int32 NumControlPoints = Instance->Strands.GetData().GetNumPoints();
+				const int32 NumCurves = Instance->Strands.GetData().GetNumCurves();
+				const int32 NumThreads = InDomainName == Domains::ControlPoint ? NumControlPoints : NumCurves;
+				OutInvocationElementCounts.Add(NumThreads);
+			}
+		}
+
+		return true;
+	}
+
+	// Unknown execution domain.
+	return false;
+}
+#undef LOCTEXT_NAMESPACE
diff --git a/Engine/Plugins/Runtime/HairStrands/Source/HairStrandsDeformer/Private/DeformerGroomComponentSource.h b/Engine/Plugins/Runtime/HairStrands/Source/HairStrandsDeformer/Private/DeformerGroomComponentSource.h
new file mode 100644
--- /dev/null
+++ b/Engine/Plugins/Runtime/HairStrands/Source/HairStrandsDeformer/Private/DeformerGroomComponentSource.h
@@ -0,0 +1,29 @@
+// Copyright Epic Games, Inc. All Rights Reserved.
+
+#pragma once
+
+#include "OptimusComponentSource.h"
+#include "DeformerGroomComponentSource.generated.h"
+
+UCLASS()
+class UOptimusGroomComponentSource : public UOptimusComponentSource
+{
+	GENERATED_BODY()
+public:
+
+	struct Domains
+	{
+		static FName ControlPoint;
+		static FName Curve;
+	};
+
+	// UOptimusComponentSource implementations
+	FText GetDisplayName() const override;
+	FName GetBindingName() const override { return FName("Groom"); }
+	TSubclassOf<UActorComponent> GetComponentClass() const override;
+	TArray<FName> GetExecutionDomains() const override;
+	int32 GetLodIndex(const UActorComponent* InComponent) const override;
+	uint32 GetDefaultNumInvocations(const UActorComponent* InComponent, int32 InLod) const;
+	bool GetComponentElementCountsForExecutionDomain(FName InDomainName, const UActorComponent* InComponent, int32 InLodIndex, TArray<int32>& OutInvocationElementCounts) const override;
+	bool IsUsableAsPrimarySource() const override;
+};
diff --git a/Engine/Plugins/Runtime/HairStrands/Source/HairStrandsDeformer/Private/DeformerGroomDomainsExec.cpp b/Engine/Plugins/Runtime/HairStrands/Source/HairStrandsDeformer/Private/DeformerGroomDomainsExec.cpp
deleted file mode 100644
diff --git a/Engine/Plugins/Runtime/HairStrands/Source/HairStrandsDeformer/Private/DeformerGroomDomainsExec.h b/Engine/Plugins/Runtime/HairStrands/Source/HairStrandsDeformer/Private/DeformerGroomDomainsExec.h
deleted file mode 100644
diff --git a/Engine/Plugins/Runtime/HairStrands/Source/HairStrandsDeformer/Private/DeformerGroomDomainsSource.cpp b/Engine/Plugins/Runtime/HairStrands/Source/HairStrandsDeformer/Private/DeformerGroomDomainsSource.cpp
deleted file mode 100644
diff --git a/Engine/Plugins/Runtime/HairStrands/Source/HairStrandsDeformer/Private/DeformerGroomDomainsSource.h b/Engine/Plugins/Runtime/HairStrands/Source/HairStrandsDeformer/Private/DeformerGroomDomainsSource.h
deleted file mode 100644
diff --git a/Engine/Plugins/Runtime/HairStrands/Source/HairStrandsDeformer/Private/DeformerGroomGuidesRead.cpp b/Engine/Plugins/Runtime/HairStrands/Source/HairStrandsDeformer/Private/DeformerGroomGuidesRead.cpp
deleted file mode 100644
diff --git a/Engine/Plugins/Runtime/HairStrands/Source/HairStrandsDeformer/Private/DeformerGroomGuidesRead.h b/Engine/Plugins/Runtime/HairStrands/Source/HairStrandsDeformer/Private/DeformerGroomGuidesRead.h
deleted file mode 100644
diff --git a/Engine/Plugins/Runtime/HairStrands/Source/HairStrandsDeformer/Private/DeformerGroomGuidesWrite.cpp b/Engine/Plugins/Runtime/HairStrands/Source/HairStrandsDeformer/Private/DeformerGroomGuidesWrite.cpp
deleted file mode 100644
diff --git a/Engine/Plugins/Runtime/HairStrands/Source/HairStrandsDeformer/Private/DeformerGroomGuidesWrite.h b/Engine/Plugins/Runtime/HairStrands/Source/HairStrandsDeformer/Private/DeformerGroomGuidesWrite.h
deleted file mode 100644
diff --git a/Engine/Plugins/Runtime/HairStrands/Source/HairStrandsDeformer/Private/DeformerGroomInterfaceUtils.h b/Engine/Plugins/Runtime/HairStrands/Source/HairStrandsDeformer/Private/DeformerGroomInterfaceUtils.h
deleted file mode 100644
diff --git a/Engine/Plugins/Runtime/HairStrands/Source/HairStrandsDeformer/Private/DeformerGroomMeshesRead.cpp b/Engine/Plugins/Runtime/HairStrands/Source/HairStrandsDeformer/Private/DeformerGroomMeshesRead.cpp
deleted file mode 100644
diff --git a/Engine/Plugins/Runtime/HairStrands/Source/HairStrandsDeformer/Private/DeformerGroomMeshesRead.h b/Engine/Plugins/Runtime/HairStrands/Source/HairStrandsDeformer/Private/DeformerGroomMeshesRead.h
deleted file mode 100644
diff --git a/Engine/Plugins/Runtime/HairStrands/Source/HairStrandsDeformer/Private/DeformerGroomSolverRead.cpp b/Engine/Plugins/Runtime/HairStrands/Source/HairStrandsDeformer/Private/DeformerGroomSolverRead.cpp
deleted file mode 100644
diff --git a/Engine/Plugins/Runtime/HairStrands/Source/HairStrandsDeformer/Private/DeformerGroomSolverRead.h b/Engine/Plugins/Runtime/HairStrands/Source/HairStrandsDeformer/Private/DeformerGroomSolverRead.h
deleted file mode 100644
diff --git a/Engine/Plugins/Runtime/HairStrands/Source/HairStrandsDeformer/Private/DeformerGroomStrandsRead.cpp b/Engine/Plugins/Runtime/HairStrands/Source/HairStrandsDeformer/Private/DeformerGroomStrandsRead.cpp
deleted file mode 100644
diff --git a/Engine/Plugins/Runtime/HairStrands/Source/HairStrandsDeformer/Private/DeformerGroomStrandsRead.h b/Engine/Plugins/Runtime/HairStrands/Source/HairStrandsDeformer/Private/DeformerGroomStrandsRead.h
deleted file mode 100644
diff --git a/Engine/Plugins/Runtime/HairStrands/Source/HairStrandsDeformer/Private/DeformerGroomStrandsWrite.cpp b/Engine/Plugins/Runtime/HairStrands/Source/HairStrandsDeformer/Private/DeformerGroomStrandsWrite.cpp
deleted file mode 100644
diff --git a/Engine/Plugins/Runtime/HairStrands/Source/HairStrandsDeformer/Private/DeformerGroomStrandsWrite.h b/Engine/Plugins/Runtime/HairStrands/Source/HairStrandsDeformer/Private/DeformerGroomStrandsWrite.h
deleted file mode 100644
diff --git a/Engine/Plugins/Runtime/HairStrands/Source/HairStrandsSolver/Private/GroomSolverComponent.cpp b/Engine/Plugins/Runtime/HairStrands/Source/HairStrandsSolver/Private/GroomSolverComponent.cpp
--- a/Engine/Plugins/Runtime/HairStrands/Source/HairStrandsSolver/Private/GroomSolverComponent.cpp
+++ b/Engine/Plugins/Runtime/HairStrands/Source/HairStrandsSolver/Private/GroomSolverComponent.cpp
@@ -78,89 +78,34 @@ void UGroomSolverComponent::OnUnregister()
 	GroomSolverProxy.DeformerInstance = nullptr;
 }
 
-void UGroomSolverComponent::SelectDynamicCurves() 
+void UGroomSolverComponent::SelectSimulatedCurves(TArray<int32>& NumSimulatedCurves) const
 {
 	TArray<FVector> ViewLocations;
 	UE::Groom::Private::GatherViewLocations(GetWorld(), ViewLocations);
 
-	SolverSettings.CurveDynamicIndices.Reset();
-	SolverSettings.CurveKinematicIndices.Reset();
-	SolverSettings.PointDynamicIndices.Reset();
-	SolverSettings.PointKinematicIndices.Reset();
-	SolverSettings.ObjectDistanceLods.Reset();
-
-	static constexpr int32 GroupSize = 64;
-
 	if (ViewLocations.Num() > 0)
 	{
-		uint32 CurveOffset = 0, PointOffset = 0;
 		for(const TObjectPtr<UGroomComponent>& GroomComponent : GetGroomComponents())
 		{
-			if(GroomComponent && GroomComponent->GroomAsset)
-			{
-				const FBoxSphereBounds GroomBounds = GroomComponent->CalcBounds(GroomComponent->GetComponentTransform());
+			const FBoxSphereBounds GroomBounds = GroomComponent->CalcBounds(GroomComponent->GetComponentTransform());
 						
-				FVector::FReal MinDistance = FMath::Square(WORLD_MAX);
-				for (const FVector& ViewLocation : ViewLocations)
-				{
-					MinDistance = FMath::Min(MinDistance, (GroomBounds.Origin - ViewLocation).SizeSquared());
-				}
-				MinDistance = FMath::Sqrt(MinDistance);
+			FVector::FReal MinDistance = FMath::Square(WORLD_MAX);
+			for (const FVector& ViewLocation : ViewLocations)
+			{
+				MinDistance = FMath::Min(MinDistance, (GroomBounds.Origin - ViewLocation).SizeSquared());
+			}
+			MinDistance = FMath::Sqrt(MinDistance);
 
-				const FVector::FReal DistanceRatio = FMath::Clamp(MinDistance / GetSolverSettings().MaxLODDistance, 0.0f, 1.0f);
+			const FVector::FReal DistanceRatio = FMath::Clamp(MinDistance / GetSolverSettings().MaxLODDistance, 0.0f, 1.0f);
 
-				const uint32 NumGroups = GroomComponent->GetGroupCount();
-				for (uint32 GroupIndex = 0; GroupIndex < NumGroups; ++GroupIndex)
+			const uint32 NumGroups = GroomComponent->GetGroupCount();
+			for (uint32 GroupIndex = 0; GroupIndex < NumGroups; ++GroupIndex)
+			{
+				if (const FHairGroupInstance* Instance = GroomComponent->GetGroupInstance(GroupIndex))
 				{
-					const FHairGroupInstance* GroupInstance = GroomComponent->GetGroupInstance(GroupIndex);
-					if(GroomComponent->IsDeformationEnable(GroupIndex) && GroupInstance)
-					{
-						const int32 NumTotalCurves = GroupInstance->Guides.GetData().GetNumCurves();
-						const int32 NumTotalPoints = GroupInstance->Guides.GetData().GetNumPoints();
+					const int32 NumCurves = Instance->Guides.IsValid() ? Instance->Guides.GetData().GetNumCurves() : 0;
+					NumSimulatedCurves.Add((1.0-DistanceRatio) * NumCurves);
 
-						const int32 NumCurvePoints = NumTotalPoints / NumTotalCurves;
-						const int32 NumDynamicCurves = (1.0-DistanceRatio) * NumTotalCurves;
-						
-						const int32 NumDynamicPoints = NumDynamicCurves * NumCurvePoints;
-						const int32 NumKinematicPoints = NumTotalPoints-NumDynamicPoints;
-						
-						const int32 NumDynamicPadding = FMath::CeilToInt32(NumDynamicPoints / static_cast<float>(GroupSize)) * GroupSize - NumDynamicPoints;
-						const int32 NumKinematicPadding = FMath::CeilToInt32(NumKinematicPoints / static_cast<float>(GroupSize)) * GroupSize - NumKinematicPoints;
-						
-						for (int32 CurveIndex = 0; CurveIndex < NumDynamicCurves; ++CurveIndex)
-						{
-							SolverSettings.CurveDynamicIndices.Add(CurveOffset + CurveIndex);
-						}
-						for (int32 CurveIndex = NumDynamicCurves; CurveIndex < NumTotalCurves; ++CurveIndex)
-						{
-							SolverSettings.CurveKinematicIndices.Add(CurveOffset + CurveIndex);
-						}
-						for (int32 PointIndex = 0; PointIndex < NumDynamicPoints; ++PointIndex)
-						{
-							SolverSettings.PointDynamicIndices.Add(PointOffset + PointIndex);
-						}
-						for (int32 PointIndex = 0; PointIndex < NumDynamicPadding; ++PointIndex)
-						{
-							SolverSettings.PointDynamicIndices.Add(INDEX_NONE);
-						}
-						for (int32 PointIndex = NumDynamicPoints; PointIndex < NumTotalPoints; ++PointIndex)
-						{
-							SolverSettings.PointKinematicIndices.Add(PointOffset + PointIndex);
-						}
-						for (int32 PointIndex = 0; PointIndex < NumKinematicPadding; ++PointIndex)
-						{
-							SolverSettings.PointKinematicIndices.Add(INDEX_NONE);
-						}
-
-						const int32 NumObjectLODs = FMath::CeilLogTwo(NumTotalCurves);
-						const uint32 ObjectLod = NumObjectLODs - 1 - FMath::FloorLog2(NumDynamicCurves);
-
-						SolverSettings.ObjectDistanceLods.Add(ObjectLod);
-						
-						CurveOffset += NumTotalCurves;
-						PointOffset += FMath::CeilToInt32(NumTotalPoints / static_cast<float>(GroupSize)) * GroupSize;
-					}
-				
 					// Keep it for now for debugging since it is WIP
 					//UE_LOG(LogTemp, Log, TEXT("Num Curves = %d Distance Ratio = %f Simulated Curves = %f | Num LODs = %d"), NumCurves, DistanceRatio, (1.0-DistanceRatio) * NumCurves, FMath::CeilLogTwo(NumCurves));
 				}
@@ -177,8 +122,6 @@ void UGroomSolverComponent::TickComponent(float DeltaTime, enum ELevelTick TickT
 	{
 		MarkRenderDynamicDataDirty();
 	}
-
-	SelectDynamicCurves();
 }
 
 void UGroomSolverComponent::SendRenderDynamicData_Concurrent()
diff --git a/Engine/Plugins/Runtime/HairStrands/Source/HairStrandsSolver/Public/GroomSolverComponent.h b/Engine/Plugins/Runtime/HairStrands/Source/HairStrandsSolver/Public/GroomSolverComponent.h
--- a/Engine/Plugins/Runtime/HairStrands/Source/HairStrandsSolver/Public/GroomSolverComponent.h
+++ b/Engine/Plugins/Runtime/HairStrands/Source/HairStrandsSolver/Public/GroomSolverComponent.h
@@ -32,21 +32,6 @@ struct FGroomSolverSettings
 	/** Maximum LOD distance (if distance between the component and the views is higher that this threshold, no simulation)*/
 	UPROPERTY(EditAnywhere, Category = "LODs")
 	float MaxLODDistance = 1000;
-
-	/** List of dynamic curves */
-	TArray<int32> CurveDynamicIndices;
-
-	/** List of kinematic curves */
-	TArray<int32> CurveKinematicIndices;
-
-	/** List of dynamic points */
-	TArray<int32> PointDynamicIndices;
-
-	/** List of kinematic points */
-	TArray<int32> PointKinematicIndices;
-	
-	/** List of object curve lods */
-	TArray<int32> ObjectDistanceLods;
 };
 
 /** Dataflow groom solver proxy used in dataflow simulation */ 
@@ -157,8 +142,8 @@ public:
 
 private :
 
-	/** Select the number of dynamic curves based on the distance component-views */
-	void SelectDynamicCurves();
+	/** Select the number of simulated curves based on the distance component-views */
+	void SelectSimulatedCurves(TArray<int32>& NumSimulatedCurves) const;
 
 	/** Solver settings used to control the simulation */
 	UPROPERTY(EditAnywhere, Category = "Settings")
diff --git a/Engine/Source/Runtime/Experimental/Chaos/Public/GeometryCollection/ManagedArrayTypeValues.inl b/Engine/Source/Runtime/Experimental/Chaos/Public/GeometryCollection/ManagedArrayTypeValues.inl
--- a/Engine/Source/Runtime/Experimental/Chaos/Public/GeometryCollection/ManagedArrayTypeValues.inl
+++ b/Engine/Source/Runtime/Experimental/Chaos/Public/GeometryCollection/ManagedArrayTypeValues.inl
@@ -60,7 +60,6 @@ MANAGED_ARRAY_TYPE(Chaos::PMatrix33d, PMatrix33d)
 MANAGED_ARRAY_TYPE(TArray<Chaos::PMatrix33d>, PMatrix33dArray)
 MANAGED_ARRAY_TYPE(TArray<TArray<FVector3f>>, FVector3fNestedArray)
 MANAGED_ARRAY_TYPE(FUintVector2, UintVector2)
-MANAGED_ARRAY_TYPE(TObjectPtr<UObject>, UObjectArray)
 
 // NOTE: new types must be added at the bottom to keep serialization from breaking
 
