Fix hair jittering when using Niagara simulation cache and pausing the sequencer.

diff --git a/Engine/Plugins/Runtime/HairStrands/Source/HairStrandsCore/Private/Niagara/NiagaraDataInterfaceHairStrands.cpp b/Engine/Plugins/Runtime/HairStrands/Source/HairStrandsCore/Private/Niagara/NiagaraDataInterfaceHairStrands.cpp
--- a/Engine/Plugins/Runtime/HairStrands/Source/HairStrandsCore/Private/Niagara/NiagaraDataInterfaceHairStrands.cpp
+++ b/Engine/Plugins/Runtime/HairStrands/Source/HairStrandsCore/Private/Niagara/NiagaraDataInterfaceHairStrands.cpp
@@ -277,10 +277,15 @@ struct FHairGroupInstanceRDG
 	FVector3f RestPositionOffsetValue = FVector3f::ZeroVector;
 	FRDGBufferSRVRef RestPositionBuffer = nullptr;
 	FRDGBufferUAVRef DeformedPositionBufferUAV = nullptr;
-	FRDGBufferSRVRef DeformedPositionOffset = nullptr;
+	FRDGBufferSRVRef DeformedPositionOffsetSRV = nullptr;
 	FRDGBufferSRVRef CurvesOffsetsBuffer = nullptr;
 	FRDGBufferSRVRef PointToCurveIndexBuffer = nullptr;
 
+	FRDGBufferRef DeformedPositionBuffer = nullptr;
+	FRDGBufferRef DeformedPositionOffset = nullptr;
+	FRDGBufferRef PreviousDeformedPositionBuffer = nullptr;
+	FRDGBufferRef PreviousDeformedPositionOffset = nullptr;
+
 	// Skinning
 	FRDGBufferSRVRef RestTrianglePositionBuffer = nullptr;
 	FRDGBufferSRVRef DeformedTrianglePositionBuffer = nullptr;
@@ -319,8 +324,19 @@ static FHairGroupInstanceRDG Convert(FRDGBuilder& GraphBuilder, const TRefCountP
 
 		if (FHairStrandsDeformedResource* DeformedResource = In->Guides.DeformedResource)
 		{
-			Out.DeformedPositionBufferUAV	= RegisterAsUAV(GraphBuilder, DeformedResource->GetBuffer(FHairStrandsDeformedResource::EFrameType::Current));
-			Out.DeformedPositionOffset		= RegisterAsSRV(GraphBuilder, DeformedResource->GetPositionOffsetBuffer(FHairStrandsDeformedResource::EFrameType::Current));
+			FRDGImportedBuffer CurrDeformedPositionBuffer = Register(GraphBuilder, DeformedResource->GetBuffer(FHairStrandsDeformedResource::EFrameType::Current), ERDGImportedBufferFlags::CreateUAV);
+			FRDGImportedBuffer PrevDeformedPositionBuffer = Register(GraphBuilder, DeformedResource->GetBuffer(FHairStrandsDeformedResource::EFrameType::Previous), ERDGImportedBufferFlags::None);
+
+			FRDGImportedBuffer CurrDeformedPositionOffset = Register(GraphBuilder, DeformedResource->GetPositionOffsetBuffer(FHairStrandsDeformedResource::EFrameType::Current), ERDGImportedBufferFlags::CreateSRV);
+			FRDGImportedBuffer PrevDeformedPositionOffset = Register(GraphBuilder, DeformedResource->GetPositionOffsetBuffer(FHairStrandsDeformedResource::EFrameType::Previous), ERDGImportedBufferFlags::None);
+
+			Out.DeformedPositionBufferUAV = CurrDeformedPositionBuffer.UAV;
+			Out.DeformedPositionBuffer    = CurrDeformedPositionBuffer.Buffer;
+			Out.DeformedPositionOffsetSRV = CurrDeformedPositionOffset.SRV;
+			Out.DeformedPositionOffset    = CurrDeformedPositionOffset.Buffer;
+
+			Out.PreviousDeformedPositionBuffer = PrevDeformedPositionBuffer.Buffer;
+			Out.PreviousDeformedPositionOffset = PrevDeformedPositionOffset.Buffer;
 		}
 
 		const bool bHasSkinnedInterpolation = In->Guides.DeformedRootResource && In->Guides.DeformedRootResource->IsValid(MeshLODIndex);
@@ -360,7 +376,7 @@ static FHairGroupInstanceRDG Convert(FRDGBuilder& GraphBuilder, const TRefCountP
 	// Fallback
 	if (!Out.RestPositionBuffer) 				{ Out.RestPositionBuffer 				= DummyByteAddressBuffer; };
 	if (!Out.DeformedPositionBufferUAV)			{ Out.DeformedPositionBufferUAV 		= GraphBuilder.CreateUAV(GraphBuilder.CreateBuffer(FRDGBufferDesc::CreateByteAddressDesc(16u), TEXT("Niagara.Hair.DummyUAV")), ERDGUnorderedAccessViewFlags::SkipBarrier); };
-	if (!Out.DeformedPositionOffset)			{ Out.DeformedPositionOffset			= DummyStructuredBuffer; };
+	if (!Out.DeformedPositionOffsetSRV)			{ Out.DeformedPositionOffsetSRV			= DummyStructuredBuffer; };
 	if (!Out.CurvesOffsetsBuffer)				{ Out.CurvesOffsetsBuffer 				= DummyByteAddressBuffer; };
 	if (!Out.RestTrianglePositionBuffer)		{ Out.RestTrianglePositionBuffer 		= DummyVertexBuffer16byte; };
 	if (!Out.DeformedTrianglePositionBuffer)	{ Out.DeformedTrianglePositionBuffer 	= DummyVertexBuffer16byte; };
@@ -1026,7 +1042,7 @@ static void InterpolateGroomGuides(
 		PassParameters->RestPositionComponent = RestPositionComponent;
 		PassParameters->WorldToLocal = WorldToLocal;
 
-		PassParameters->DeformedPositionOffset = InstanceRDG.DeformedPositionOffset;
+		PassParameters->DeformedPositionOffset = InstanceRDG.DeformedPositionOffsetSRV;
 		PassParameters->RestPositionOffset = InstanceRDG.RestPositionOffsetValue;
 		PassParameters->DeformedPositionBuffer = InstanceRDG.DeformedPositionBufferUAV;
 		PassParameters->RestPositionBuffer = InstanceRDG.RestPositionBuffer;
@@ -1049,6 +1065,12 @@ static void InterpolateGroomGuides(
 			PassParameters,
 			FIntVector(DispatchCount, 1, 1)
 		);
+
+		if (InstanceRDG.DeformedPositionBuffer && InstanceRDG.PreviousDeformedPositionBuffer)
+		{
+			AddCopyBufferPass(GraphBuilder, InstanceRDG.PreviousDeformedPositionBuffer, InstanceRDG.DeformedPositionBuffer);
+			AddCopyBufferPass(GraphBuilder, InstanceRDG.PreviousDeformedPositionOffset, InstanceRDG.DeformedPositionOffset);
+		}
 	}
 }
 
@@ -3546,7 +3568,7 @@ void UNiagaraDataInterfaceHairStrands::SetShaderParameters(const FNiagaraDataInt
 		ShaderParameters->DeformedPositionBuffer = InstanceRDG.DeformedPositionBufferUAV;
 		ShaderParameters->CurvesOffsetsBuffer = InstanceRDG.CurvesOffsetsBuffer;
 		ShaderParameters->RestPositionBuffer = InstanceRDG.RestPositionBuffer;
-		ShaderParameters->DeformedPositionOffset = InstanceRDG.DeformedPositionOffset;
+		ShaderParameters->DeformedPositionOffset = InstanceRDG.DeformedPositionOffsetSRV;
 		ShaderParameters->RestTrianglePositionBuffer = InstanceRDG.RestTrianglePositionBuffer;
 		ShaderParameters->DeformedTrianglePositionBuffer = InstanceRDG.DeformedTrianglePositionBuffer;
 		ShaderParameters->RestSamplePositionsBuffer = InstanceRDG.RestSamplePositionsBuffer;
@@ -3589,7 +3611,7 @@ void UNiagaraDataInterfaceHairStrands::SetShaderParameters(const FNiagaraDataInt
 		ShaderParameters->DeformedPositionBuffer = InstanceRDG.DeformedPositionBufferUAV;
 		ShaderParameters->CurvesOffsetsBuffer = InstanceRDG.CurvesOffsetsBuffer;
 		ShaderParameters->RestPositionBuffer = InstanceRDG.RestPositionBuffer;
-		ShaderParameters->DeformedPositionOffset = InstanceRDG.DeformedPositionOffset;
+		ShaderParameters->DeformedPositionOffset = InstanceRDG.DeformedPositionOffsetSRV;
 		ShaderParameters->RestTrianglePositionBuffer = InstanceRDG.RestTrianglePositionBuffer;
 		ShaderParameters->DeformedTrianglePositionBuffer = InstanceRDG.DeformedTrianglePositionBuffer;
 		ShaderParameters->RestSamplePositionsBuffer = InstanceRDG.RestSamplePositionsBuffer;
