Enable KMeans to run in parallel within the groom builder.

diff --git a/Engine/Plugins/Runtime/HairStrands/Source/HairStrandsCore/Private/GroomBuilder.cpp b/Engine/Plugins/Runtime/HairStrands/Source/HairStrandsCore/Private/GroomBuilder.cpp
--- a/Engine/Plugins/Runtime/HairStrands/Source/HairStrandsCore/Private/GroomBuilder.cpp
+++ b/Engine/Plugins/Runtime/HairStrands/Source/HairStrandsCore/Private/GroomBuilder.cpp
@@ -57,7 +57,7 @@ FString FGroomBuilder::GetVersion()
 	FString Out = TEXT("v18a");
 	if (IsRenderCurveEnabled())
 	{
-		Out += TEXT("_Curve2");
+		Out += TEXT("_Curve3e");
 	}
 	return Out;
 }
@@ -4236,7 +4236,7 @@ void InternalBuildRenderCurveResourceBulkData(const FHairStrandsDatas& In, FRend
 		
 		TArray<FVector3f> Centers;
 		Cluster.GetUniformSpacedInitialCenters<FVector3f>(RootPoints, NumClusters, Centers);
-		int32 FoundClusters = Cluster.ComputeClusters<FVector3f>(RootPoints, NumClusters, Centers);
+		int32 FoundClusters = Cluster.ComputeClusters<FVector3f, true/*bParallel*/>(RootPoints, NumClusters, Centers);
 
 		// Verify cluster counts
 		check(Cluster.ClusterSizes.Num() == FoundClusters);
diff --git a/Engine/Shaders/Private/HairStrands/RenderCurveRaster.usf b/Engine/Shaders/Private/HairStrands/RenderCurveRaster.usf
--- a/Engine/Shaders/Private/HairStrands/RenderCurveRaster.usf
+++ b/Engine/Shaders/Private/HairStrands/RenderCurveRaster.usf
@@ -2250,8 +2250,9 @@ void DebugDrawingCS(uint2 DispatchThreadId : SV_DispatchThreadID)
 		// Cluster bounds
 		#if 0
 		{
+			const float3 ClusterColor = ColorMapMagma(VisibleClusterFetchIndex / float(VisibleClusterCount));
 			FShaderPrintContext Ctx = InitShaderPrintContext(true, uint2(0,0));
-			AddOBBTWS(Ctx, LocalBoundsCenter - LocalBoundsExtent, LocalBoundsCenter + LocalBoundsExtent, ColorRed, RenderCurveInstanceData.LocalToTranslatedWorld);
+			AddOBBTWS(Ctx, LocalBoundsCenter - LocalBoundsExtent, LocalBoundsCenter + LocalBoundsExtent, float4(ClusterColor,1), RenderCurveInstanceData.LocalToTranslatedWorld);
 		}
 		#endif
 
