Fix TSAN report related to HairDescriptionGroup access.

diff --git a/Engine/Plugins/Runtime/HairStrands/Source/HairStrandsCore/Private/GroomAsset.cpp b/Engine/Plugins/Runtime/HairStrands/Source/HairStrandsCore/Private/GroomAsset.cpp
--- a/Engine/Plugins/Runtime/HairStrands/Source/HairStrandsCore/Private/GroomAsset.cpp
+++ b/Engine/Plugins/Runtime/HairStrands/Source/HairStrandsCore/Private/GroomAsset.cpp
@@ -2424,7 +2424,6 @@ const FHairDescriptionGroups& UGroomAsset::GetHairDescriptionGroups()
 
 	check(HairDescriptionBulkData[HairDescriptionType]);
 	
-	if (!CachedHairDescription[HairDescriptionType] || !CachedHairDescriptionGroups[HairDescriptionType])
 	{
 		UE::TUniqueLock LockScope(InternalLock);
 
