Fix crash when running hair deformers due to incorrect shader path.

diff --git a/Engine/Plugins/Runtime/HairStrands/Source/HairStrandsDeformer/Private/DeformerGroomDomainsExec.cpp b/Engine/Plugins/Runtime/HairStrands/Source/HairStrandsDeformer/Private/DeformerGroomDomainsExec.cpp
--- a/Engine/Plugins/Runtime/HairStrands/Source/HairStrandsDeformer/Private/DeformerGroomDomainsExec.cpp
+++ b/Engine/Plugins/Runtime/HairStrands/Source/HairStrandsDeformer/Private/DeformerGroomDomainsExec.cpp
@@ -49,7 +49,7 @@ void UOptimusGroomExecDataInterface::GetShaderParameters(TCHAR const* UID, FShad
 	InOutBuilder.AddNestedStruct<FGroomExecDataInterfaceParameters>(UID);
 }
 
-TCHAR const* UOptimusGroomExecDataInterface::TemplateFilePath = TEXT("/Plugin/Runtime/HairStrands/Private/Deformers/DeformerDataGroomDomainsExec.ush");
+TCHAR const* UOptimusGroomExecDataInterface::TemplateFilePath = TEXT("/Plugin/Runtime/HairStrands/Private/Deformers/DeformerGroomDomainsExec.ush");
 
 TCHAR const* UOptimusGroomExecDataInterface::GetShaderVirtualPath() const
 {
