Add Groom binding support for Guide-Cache.

This CL allows cache-guide to be used with groom binding, but assume the cache guide is following the surface correctly.

#12354

diff --git a/Engine/Plugins/Runtime/HairStrands/Source/HairStrandsCore/Private/GroomManager.cpp b/Engine/Plugins/Runtime/HairStrands/Source/HairStrandsCore/Private/GroomManager.cpp
--- a/Engine/Plugins/Runtime/HairStrands/Source/HairStrandsCore/Private/GroomManager.cpp
+++ b/Engine/Plugins/Runtime/HairStrands/Source/HairStrandsCore/Private/GroomManager.cpp
@@ -716,7 +716,8 @@ static void RunHairBindingSurfaceUpdate(
 					Instance->BindingType == EHairBindingType::Skinning &&
 					(Instance->HairGroupPublicData->IsGlobalInterpolationEnable(HairLODIndex) ||
 					Instance->HairGroupPublicData->IsSimulationEnable(HairLODIndex) ||
-					Instance->HairGroupPublicData->bIsDeformationEnable);
+					Instance->HairGroupPublicData->bIsDeformationEnable ||
+					Instance->HairGroupPublicData->bIsSimulationCacheEnable);
 				if (bNeedSurfaceUpdate)
 				{
 					// Extract MeshLODData and MeshLODIndex
@@ -809,7 +810,7 @@ struct FGuideInstanceData
 		return 
 			(!bCheckMeshLODIndex || (bCheckMeshLODIndex && MeshLODIndex >= 0)) &&
 			(BindingType == EHairBindingType::Skinning) && 
-			(bGlobalDeformationEnable || bSimulationEnable || bDeformationEnable);
+			(bGlobalDeformationEnable || bSimulationEnable || bDeformationEnable || ActiveGroomCacheType==EGroomCacheType::Guides);
 	}
 };
 
diff --git a/Engine/Plugins/Runtime/HairStrands/Source/HairStrandsCore/Private/HairStrandsInterpolation.cpp b/Engine/Plugins/Runtime/HairStrands/Source/HairStrandsCore/Private/HairStrandsInterpolation.cpp
--- a/Engine/Plugins/Runtime/HairStrands/Source/HairStrandsCore/Private/HairStrandsInterpolation.cpp
+++ b/Engine/Plugins/Runtime/HairStrands/Source/HairStrandsCore/Private/HairStrandsInterpolation.cpp
@@ -655,8 +655,9 @@ void AddHairStrandsInterpolationPass(
 		RenDeformedRootResources->IsValid(MeshLODIndex);
 
 	// Guides
+	// Guide interpolation with dynamic mesh geometry is now support withguide cach. It implicitly assumes that the guide-cache follows correctly the binding surface
 	bool bSupportGlobalInterpolation = false;
-	if (bSupportDynamicMesh && (Instance->Guides.bIsSimulationEnable || Instance->Guides.bHasGlobalInterpolation || Instance->Guides.bIsDeformationEnable)) // No need for Instance->Guides.bHasSimulationCache since it is incompatible with binding
+	if (bSupportDynamicMesh && (Instance->Guides.bIsSimulationEnable || Instance->Guides.bHasGlobalInterpolation || Instance->Guides.bIsDeformationEnable || Instance->Guides.bIsSimulationCacheEnable))
 	{
 		const FHairStrandsLODRestRootResource& Sim_RestLODDatas = *SimRestRootResources->GetLOD(MeshLODIndex);
 		const FHairStrandsLODDeformedRootResource& Sim_DeformedLODDatas = *SimDeformedRootResources->GetLOD(MeshLODIndex);
diff --git a/Engine/Plugins/Runtime/HairStrands/Source/HairStrandsCore/Public/GroomComponent.h b/Engine/Plugins/Runtime/HairStrands/Source/HairStrandsCore/Public/GroomComponent.h
--- a/Engine/Plugins/Runtime/HairStrands/Source/HairStrandsCore/Public/GroomComponent.h
+++ b/Engine/Plugins/Runtime/HairStrands/Source/HairStrandsCore/Public/GroomComponent.h
@@ -33,7 +33,7 @@ public:
 	UPROPERTY(EditAnywhere, BlueprintReadWrite, interp, Category = "Groom")
 	TObjectPtr<UGroomAsset> GroomAsset;
 
-	UPROPERTY(EditAnywhere, BlueprintReadWrite, interp, Category = "GroomCache", meta = (EditCondition = "BindingAsset == nullptr"))
+	UPROPERTY(EditAnywhere, BlueprintReadWrite, interp, Category = "GroomCache")
 	TObjectPtr<UGroomCache> GroomCache;
 
 	/** Niagara components that will be attached to the system*/
@@ -45,7 +45,7 @@ public:
 	TObjectPtr<class USkeletalMesh> SourceSkeletalMesh;
 
 	/** Optional binding asset for binding a groom onto a skeletal mesh. If the binding asset is not provided the projection is done at runtime, which implies a large GPU cost at startup time. */
-	UPROPERTY(EditAnywhere, BlueprintReadOnly, interp, Category = "Groom", meta = (EditCondition = "GroomCache == nullptr"))
+	UPROPERTY(EditAnywhere, BlueprintReadOnly, interp, Category = "Groom")
 	TObjectPtr<class UGroomBindingAsset> BindingAsset;
 
 	/** Physics asset to be used for hair simulation */
