Groom Deformers are only valid if GroupCount is non-zero.
Add this condition to IsValid() for Groom Domain Data Interface.
This avoids crashing later in compute graph execution.

diff --git a/Engine/Plugins/Runtime/HairStrands/Source/HairStrandsDeformer/Private/DeformerGroomDomainsExec.cpp b/Engine/Plugins/Runtime/HairStrands/Source/HairStrandsDeformer/Private/DeformerGroomDomainsExec.cpp
--- a/Engine/Plugins/Runtime/HairStrands/Source/HairStrandsDeformer/Private/DeformerGroomDomainsExec.cpp
+++ b/Engine/Plugins/Runtime/HairStrands/Source/HairStrandsDeformer/Private/DeformerGroomDomainsExec.cpp
@@ -163,6 +163,10 @@ bool FOptimusGroomExecDataProviderProxy::IsValid(FValidationData const& InValida
 	{
 		return false;
 	}
+	if (GroupCounts.Num() == 0)
+	{
+		return false;
+	}
 	if(Domain == EOptimusGroomExecDomain::ControlPoint || Domain == EOptimusGroomExecDomain::Curve ||
 		Domain == EOptimusGroomExecDomain::StrandsEdges || Domain == EOptimusGroomExecDomain::StrandsObjects)
 	{
