Add UI for remapping group when reimporting groom asset with different group layouts.

diff --git a/Engine/Plugins/Runtime/HairStrands/Source/HairStrandsCore/Private/GroomImportOptions.cpp b/Engine/Plugins/Runtime/HairStrands/Source/HairStrandsCore/Private/GroomImportOptions.cpp
--- a/Engine/Plugins/Runtime/HairStrands/Source/HairStrandsCore/Private/GroomImportOptions.cpp
+++ b/Engine/Plugins/Runtime/HairStrands/Source/HairStrandsCore/Private/GroomImportOptions.cpp
@@ -1,6 +1,7 @@
 // Copyright Epic Games, Inc. All Rights Reserved.
 
 #include "GroomImportOptions.h"
+#include "GroomAsset.h"
 
 #include UE_INLINE_GENERATED_CPP_BY_NAME(GroomImportOptions)
 
@@ -14,3 +15,93 @@ UGroomHairGroupsPreview::UGroomHairGroupsPreview(const FObjectInitializer& Objec
 {
 }
 
+UGroomHairGroupsMapping::UGroomHairGroupsMapping(const FObjectInitializer& ObjectInitializer)
+: Super(ObjectInitializer)
+{
+}
+
+const TArray<TSharedPtr<FString>>& UGroomHairGroupsMapping::GetOldGroupNames()
+{
+	if (CachedOldGroupNames.Num() == 0)
+	{
+		for (uint32 GroupIt = 0, GroupCount=OldGroupNames.Num(); GroupIt < GroupCount; ++GroupIt)
+		{
+			CachedOldGroupNames.Add(MakeShared<FString>(OldGroupNames[GroupIt].ToString()));
+		}
+		CachedOldGroupNames.Add(MakeShared<FString>(TEXT("Default")));
+	}
+	return CachedOldGroupNames;
+}
+
+bool UGroomHairGroupsMapping::HasValidMapping() const
+{
+	for (int32 Index : NewToOldGroupIndexMapping)
+	{
+		if (Index != -1)
+		{
+			return true;
+		}
+	}
+
+	return false;
+}
+
+void UGroomHairGroupsMapping::Map(const FHairDescriptionGroups& OldGroups, const FHairDescriptionGroups& NewGroups)
+{
+	// New group names
+	for (const FHairDescriptionGroup& NewGroup : NewGroups.HairGroups)
+	{
+		NewGroupNames.Add(NewGroup.Info.GroupName);
+	}
+
+	// Old group names
+	for (const FHairDescriptionGroup& OldGroup : OldGroups.HairGroups)
+	{
+		OldGroupNames.Add(OldGroup.Info.GroupName);
+	}
+
+	// * Old->New mapping
+	RemapHairDescriptionGroups(OldGroups, NewGroups, OldToNewGroupIndexMapping);
+
+	// * New->Old mapping
+	RemapHairDescriptionGroups(NewGroups, OldGroups, NewToOldGroupIndexMapping);
+}
+
+// Src -> Dst mapping
+void UGroomHairGroupsMapping::RemapHairDescriptionGroups(
+	const FHairDescriptionGroups& SrcGroups, 
+	const FHairDescriptionGroups& DstGroups, 
+	TArray<int32>& Out)
+{
+	Out.Init(-1, SrcGroups.HairGroups.Num());
+
+	// 2. Remap Src settings to the Dst asset (use GroupName/GroomID to do the remapping)
+	for (const FHairDescriptionGroup& SrcGroup : SrcGroups.HairGroups)
+	{
+		int32 MatchGroupIndex = -1;
+		for (const FHairDescriptionGroup& DstGroup : DstGroups.HairGroups)
+		{
+			// For now, only use GroupName, as if no group name is provided, the group name is built from the GroupID
+			if (DstGroup.Info.GroupName == SrcGroup.Info.GroupName)
+			{
+				MatchGroupIndex = DstGroup.Info.GroupIndex;
+				break;
+			}
+		}
+
+		Out[SrcGroup.Info.GroupIndex] = MatchGroupIndex;
+	}
+}
+
+void UGroomHairGroupsMapping::SetIndex(int32 InNewIndex, int32 InOldIndex)
+{
+	if (NewGroupNames.IsValidIndex(InNewIndex))
+	{
+		NewToOldGroupIndexMapping[InNewIndex] = OldGroupNames.IsValidIndex(InOldIndex) ? InOldIndex : -1;
+	}
+
+	if (OldGroupNames.IsValidIndex(InOldIndex))
+	{
+		OldToNewGroupIndexMapping[InOldIndex] = NewGroupNames.IsValidIndex(InNewIndex) ? InNewIndex : -1;
+	}
+}
\ No newline at end of file
diff --git a/Engine/Plugins/Runtime/HairStrands/Source/HairStrandsCore/Public/GroomImportOptions.h b/Engine/Plugins/Runtime/HairStrands/Source/HairStrandsCore/Public/GroomImportOptions.h
--- a/Engine/Plugins/Runtime/HairStrands/Source/HairStrandsCore/Public/GroomImportOptions.h
+++ b/Engine/Plugins/Runtime/HairStrands/Source/HairStrandsCore/Public/GroomImportOptions.h
@@ -8,6 +8,8 @@
 #include "GroomAssetInterpolation.h"
 #include "GroomImportOptions.generated.h"
 
+struct FHairDescriptionGroups;
+
 UCLASS(BlueprintType, config = EditorPerProjectUserSettings, HideCategories = ("Hidden"))
 class HAIRSTRANDSCORE_API UGroomImportOptions : public UObject
 {
@@ -64,3 +66,39 @@ public:
 	UPROPERTY(config, EditAnywhere, BlueprintReadWrite, meta = (ShowOnlyInnerProperties), Category = Preview)
 	TArray<FGroomHairGroupPreview> Groups;
 };
+
+UCLASS(BlueprintType, config = EditorPerProjectUserSettings, HideCategories = ("Hidden"))
+class HAIRSTRANDSCORE_API UGroomHairGroupsMapping : public UObject
+{
+	GENERATED_UCLASS_BODY()
+
+public:
+	UPROPERTY(config, EditAnywhere, BlueprintReadWrite, meta = (ShowOnlyInnerProperties), Category = GroupMapping)
+	TArray<FName> OldGroupNames;
+
+	UPROPERTY(config, EditAnywhere, BlueprintReadWrite, meta = (ShowOnlyInnerProperties), Category = GroupMapping)
+	TArray<FName> NewGroupNames;
+
+	UPROPERTY(config, EditAnywhere, BlueprintReadWrite, meta = (ShowOnlyInnerProperties), Category = GroupMapping)
+	TArray<int32> OldToNewGroupIndexMapping;
+
+	UPROPERTY(config, EditAnywhere, BlueprintReadWrite, meta = (ShowOnlyInnerProperties), Category = GroupMapping)
+	TArray<int32> NewToOldGroupIndexMapping;
+
+	FName GetNewGroupName(int32 In) { return NewGroupNames.IsValidIndex(In) ? NewGroupNames[In] : FName(); }
+	FName GetOldGroupName(int32 In) { return OldGroupNames.IsValidIndex(In) ? OldGroupNames[In] : FName(); }
+
+	const TArray<TSharedPtr<FString>>& GetOldGroupNames();
+
+	bool HasValidMapping() const;
+	bool IsValid() const { return OldToNewGroupIndexMapping.Num() > 0; }
+	void SetIndex(int32 NewIndex, int32 OldIndex);
+
+	TArray<TSharedPtr<FString>> CachedOldGroupNames;
+
+	// Initialize UGroomHairGroupsMapping by mapping OldGroup onto NewGroup
+	void Map(const FHairDescriptionGroups& OldGroup, const FHairDescriptionGroups& NewGroup);
+
+	// Recreate a mapping from SrcGroups to DstGroups
+	static void RemapHairDescriptionGroups(const FHairDescriptionGroups& SrcGroups, const FHairDescriptionGroups& DstGroups, TArray<int32>& Out);
+};
diff --git a/Engine/Plugins/Runtime/HairStrands/Source/HairStrandsEditor/Private/GroomActions.cpp b/Engine/Plugins/Runtime/HairStrands/Source/HairStrandsEditor/Private/GroomActions.cpp
--- a/Engine/Plugins/Runtime/HairStrands/Source/HairStrandsEditor/Private/GroomActions.cpp
+++ b/Engine/Plugins/Runtime/HairStrands/Source/HairStrandsEditor/Private/GroomActions.cpp
@@ -137,7 +137,7 @@ void ExecuteRebuild(const FToolMenuContext& InContext)
 						OutGroup.InterpolationSettings = GroomAsset->GetHairGroupsInterpolation()[GroupIndex];
 					}
 				}
-				TSharedPtr<SGroomImportOptionsWindow> GroomOptionWindow = SGroomImportOptionsWindow::DisplayRebuildOptions(CurrentOptions, GroupsPreview, Filename);
+				TSharedPtr<SGroomImportOptionsWindow> GroomOptionWindow = SGroomImportOptionsWindow::DisplayRebuildOptions(CurrentOptions, GroupsPreview, nullptr/*GroupsMapping*/, Filename);
 
 				if (!GroomOptionWindow->ShouldImport())
 				{
diff --git a/Engine/Plugins/Runtime/HairStrands/Source/HairStrandsEditor/Private/GroomBindingDetailsCustomization.cpp b/Engine/Plugins/Runtime/HairStrands/Source/HairStrandsEditor/Private/GroomBindingDetailsCustomization.cpp
--- a/Engine/Plugins/Runtime/HairStrands/Source/HairStrandsEditor/Private/GroomBindingDetailsCustomization.cpp
+++ b/Engine/Plugins/Runtime/HairStrands/Source/HairStrandsEditor/Private/GroomBindingDetailsCustomization.cpp
@@ -4,10 +4,15 @@
 #include "DetailLayoutBuilder.h"
 #include "GroomBindingAsset.h"
 #include "GroomCreateBindingOptions.h"
+#include "GroomImportOptions.h"
 #include "DetailWidgetRow.h"
 #include "GroomBindingCompiler.h"
 #include "Widgets/Input/SComboBox.h"
 #include "Widgets/Text/STextBlock.h"
+#include "PropertyCustomizationHelpers.h" 
+#include "IDetailChildrenBuilder.h"
+
+#define LOCTEXT_NAMESPACE "GroomBindingDetails"
 
 ///////////////////////////////////////////////////////////////////////////////////////////////////
 
@@ -200,3 +205,83 @@ TSharedRef<IDetailCustomization> FGroomCreateBindingDetailsCustomization::MakeIn
 {
 	return MakeShared<FGroomCreateBindingDetailsCustomization>();
 }
+
+///////////////////////////////////////////////////////////////////////////////////////////////////
+
+// Customization for UGroomHairGroupsMapping
+void FGroomHairGroomRemappingDetailsCustomization::CustomizeDetails(IDetailLayoutBuilder& LayoutBuilder)
+{
+	UGroomHairGroupsMapping* Mapping = GetCutomizeDetailObject<UGroomHairGroupsMapping>(LayoutBuilder);
+
+	LayoutBuilder.HideProperty(GET_MEMBER_NAME_CHECKED(UGroomHairGroupsMapping, OldToNewGroupIndexMapping));
+	LayoutBuilder.HideProperty(GET_MEMBER_NAME_CHECKED(UGroomHairGroupsMapping, NewToOldGroupIndexMapping));
+	LayoutBuilder.HideProperty(GET_MEMBER_NAME_CHECKED(UGroomHairGroupsMapping, OldGroupNames));
+
+	TSharedRef<IPropertyHandle> Property = LayoutBuilder.GetProperty(GET_MEMBER_NAME_CHECKED(UGroomHairGroupsMapping, NewGroupNames), UGroomHairGroupsMapping::StaticClass());
+	if (Property->IsValidHandle())
+	{
+		TSharedRef<FDetailArrayBuilder> PropertyBuilder = MakeShareable(new FDetailArrayBuilder(Property, false, false, false));
+		PropertyBuilder->OnGenerateArrayElementWidget(FOnGenerateArrayElementWidget::CreateSP(this, &FGroomHairGroomRemappingDetailsCustomization::OnGenerateElementForBindingAsset, &LayoutBuilder, Mapping));
+
+		FName CategoryName = FName(TEXT("GroupMapping"));
+		IDetailCategoryBuilder& GroupMappingCategory = LayoutBuilder.EditCategory(CategoryName, FText::GetEmpty(), ECategoryPriority::TypeSpecific);
+		GroupMappingCategory.AddCustomBuilder(PropertyBuilder, false);
+	}
+}
+
+void FGroomHairGroomRemappingDetailsCustomization::OnGenerateElementForBindingAsset(TSharedRef<IPropertyHandle> StructProperty, int32 InNewGroupIndex, IDetailChildrenBuilder& ChildrenBuilder, IDetailLayoutBuilder* DetailLayout, UGroomHairGroupsMapping* InMapping)
+{
+	FProperty* Property = StructProperty->GetProperty();
+	const FLinearColor Color = FLinearColor::White;
+	ChildrenBuilder.AddCustomRow(FText::FromString(TEXT("Preview")))
+	.NameContent()
+	[
+		SNew(STextBlock)
+		.Text(FText::FromName(InMapping->NewGroupNames[InNewGroupIndex]))
+		.Font(IDetailLayoutBuilder::GetDetailFont())
+	]
+	.ValueContent()
+	[
+		SNew(SComboBox<TSharedPtr<FString>>)
+		.OptionsSource(&InMapping->GetOldGroupNames())
+		.OnGenerateWidget_Lambda([](TSharedPtr<FString> InItem)
+		{
+			return SNew(STextBlock)
+			.Text(FText::FromString(*InItem));
+		})
+		.OnSelectionChanged_Lambda([InNewGroupIndex, InMapping](TSharedPtr<FString> InItem, ESelectInfo::Type SelectInfo)
+		{
+			FName Item = FName(*InItem);
+			int32 OldGroupIndex = 0;
+			for (FName OldName : InMapping->OldGroupNames)
+			{
+				if (OldName == Item)
+				{
+					break;
+				}
+				++OldGroupIndex;
+			}
+			InMapping->SetIndex(InNewGroupIndex, OldGroupIndex);
+		})
+		[
+			SNew(STextBlock)
+			.Text_Lambda([InNewGroupIndex, InMapping] 
+			{ 
+				FText OldText = FText::FromString(TEXT("Default"));
+				int32 OldIndex = InMapping->NewToOldGroupIndexMapping[InNewGroupIndex];
+				if (InMapping->OldGroupNames.IsValidIndex(OldIndex))
+				{
+					OldText = FText::FromName(InMapping->OldGroupNames[OldIndex]);
+				}
+				return OldText; 
+			})
+		]
+	];
+}
+
+TSharedRef<IDetailCustomization> FGroomHairGroomRemappingDetailsCustomization::MakeInstance()
+{
+	return MakeShared<FGroomHairGroomRemappingDetailsCustomization>();
+}
+
+#undef LOCTEXT_NAMESPACE
diff --git a/Engine/Plugins/Runtime/HairStrands/Source/HairStrandsEditor/Private/GroomBindingDetailsCustomization.h b/Engine/Plugins/Runtime/HairStrands/Source/HairStrandsEditor/Private/GroomBindingDetailsCustomization.h
--- a/Engine/Plugins/Runtime/HairStrands/Source/HairStrandsEditor/Private/GroomBindingDetailsCustomization.h
+++ b/Engine/Plugins/Runtime/HairStrands/Source/HairStrandsEditor/Private/GroomBindingDetailsCustomization.h
@@ -4,6 +4,10 @@
 
 #include "IDetailCustomization.h"
 
+class UGroomHairGroupsMapping;
+class IDetailChildrenBuilder;
+class IPropertyHandle;
+
 struct FGroomBindingAttributeSelection
 {
 	int32 SelectedBindingAttribute = INDEX_NONE;
@@ -29,3 +33,13 @@ public:
 	static TSharedRef<IDetailCustomization> MakeInstance();
 	FGroomBindingAttributeSelection BindingAttributeSelection;
 };
+
+class FGroomHairGroomRemappingDetailsCustomization : public IDetailCustomization
+{
+public:
+	/** Begin IDetailCustomization interface */
+	virtual void CustomizeDetails(IDetailLayoutBuilder& LayoutBuilder) override;
+	void OnGenerateElementForBindingAsset(TSharedRef<IPropertyHandle> StructProperty, int32 InIndex, IDetailChildrenBuilder& ChildrenBuilder, IDetailLayoutBuilder* DetailLayout, UGroomHairGroupsMapping* InMapping);
+	/** End IDetailCustomization interface */
+	static TSharedRef<IDetailCustomization> MakeInstance();
+};
diff --git a/Engine/Plugins/Runtime/HairStrands/Source/HairStrandsEditor/Private/GroomImportOptionsWindow.cpp b/Engine/Plugins/Runtime/HairStrands/Source/HairStrandsEditor/Private/GroomImportOptionsWindow.cpp
--- a/Engine/Plugins/Runtime/HairStrands/Source/HairStrandsEditor/Private/GroomImportOptionsWindow.cpp
+++ b/Engine/Plugins/Runtime/HairStrands/Source/HairStrandsEditor/Private/GroomImportOptionsWindow.cpp
@@ -195,6 +195,37 @@ static void AddAttribute(SVerticalBox::FScopedWidgetSlotArguments& Slot, FText A
 	];
 }
 
+static void AddGroupRemappingNoFoundMessage(SVerticalBox::FScopedWidgetSlotArguments& Slot, UGroomHairGroupsMapping* InMapping)
+{
+	const FLinearColor AttributeColor(0.0f, 0.72f, 0.0f);
+	const FSlateFontInfo AttributeFont = FAppStyle::GetFontStyle("CurveEd.InfoFont");
+	const FSlateFontInfo AttributeResultFont = FAppStyle::GetFontStyle("CurveEd.InfoFont");
+
+	FText MessageText = LOCTEXT("GroomGroupMapping_NotFound", "No mapping found using group names. Edit the mapping manually or use default values.");
+	FLinearColor MessageColor = FLinearColor::Red;
+
+	Slot
+	.AutoHeight()
+	.Padding(2)
+	[
+		SNew(SBorder)
+		.Padding(FMargin(3))
+		.BorderImage(FAppStyle::GetBrush("ToolPanel.GroupBorder"))
+		[
+			SNew(SHorizontalBox)
+			+ SHorizontalBox::Slot()
+			.AutoWidth()
+			.Padding(10, 0, 0, 0)
+			[
+				SNew(STextBlock)
+				.Font(AttributeFont)
+				.Text(MessageText)
+				.ColorAndOpacity(MessageColor)
+			]
+		]
+	];
+}
+
 FText GetHairAttributeLocText(EHairAttribute In, uint32 InFlags)
 {
 	// If a new optional attribute is added, please add its UI/text description here
@@ -218,6 +249,7 @@ void SGroomImportOptionsWindow::Construct(const FArguments& InArgs)
 	ImportOptions = InArgs._ImportOptions;
 	GroomCacheImportOptions = InArgs._GroomCacheImportOptions;
 	GroupsPreview = InArgs._GroupsPreview;
+	GroupsMapping = InArgs._GroupsMapping;
 	WidgetWindow = InArgs._WidgetWindow;
 
 	FPropertyEditorModule& PropertyEditorModule = FModuleManager::GetModuleChecked<FPropertyEditorModule>("PropertyEditor");
@@ -231,6 +263,9 @@ void SGroomImportOptionsWindow::Construct(const FArguments& InArgs)
 	DetailsView2 = PropertyEditorModule.CreateDetailView(DetailsViewArgs);
 	DetailsView2->SetObject(GroupsPreview);
 
+	DetailsView3 = PropertyEditorModule.CreateDetailView(DetailsViewArgs);
+	DetailsView3->SetObject(GroupsMapping);
+
 	GroomCacheDetailsView = PropertyEditorModule.CreateDetailView(DetailsViewArgs);
 	GroomCacheDetailsView->SetObject(GroomCacheImportOptions);
 
@@ -367,6 +402,13 @@ void SGroomImportOptionsWindow::Construct(const FArguments& InArgs)
 			DetailsView2->AsShared()
 		]
 
+		+ SVerticalBox::Slot()
+		.AutoHeight()
+		.Padding(2)
+		[
+			DetailsView3->AsShared()
+		]
+
 		+ SVerticalBox::Slot()
 		.AutoHeight()
 		.HAlign(HAlign_Right)
@@ -403,6 +445,12 @@ void SGroomImportOptionsWindow::Construct(const FArguments& InArgs)
 		}
 	}
 
+	if (GroupsMapping && !GroupsMapping->HasValidMapping())
+	{
+		SVerticalBox::FScopedWidgetSlotArguments SlotArg = VerticalSlot->InsertSlot(AttributeSlotIndex++);
+		AddGroupRemappingNoFoundMessage(SlotArg, GroupsMapping);
+	}
+
 	this->ChildSlot
 	[
 		VerticalSlot
@@ -441,6 +489,7 @@ TSharedPtr<SGroomImportOptionsWindow> DisplayOptions(
 	UGroomImportOptions* ImportOptions, 
 	UGroomCacheImportOptions* GroomCacheImportOptions, 
 	UGroomHairGroupsPreview* GroupsPreview,
+	UGroomHairGroupsMapping* GroupsMapping,
 	const FString& FilePath, 
 	EGroomOptionsVisibility VisibilityFlag, 
 	FText WindowTitle, 
@@ -476,6 +525,7 @@ TSharedPtr<SGroomImportOptionsWindow> DisplayOptions(
 		.ImportOptions(ImportOptions)
 		.GroomCacheImportOptions(GroomCacheImportOptions)
 		.GroupsPreview(GroupsPreview)
+		.GroupsMapping(GroupsMapping)
 		.WidgetWindow(Window)
 		.FullPath(FText::FromString(FileName))
 		.ButtonLabel(InButtonLabel)
@@ -494,16 +544,16 @@ TSharedPtr<SGroomImportOptionsWindow> DisplayOptions(
 	return OptionsWindow;
 }
 
-TSharedPtr<SGroomImportOptionsWindow> SGroomImportOptionsWindow::DisplayImportOptions(UGroomImportOptions* ImportOptions, UGroomCacheImportOptions* GroomCacheImportOptions, UGroomHairGroupsPreview* GroupsPreview, const FString& FilePath)
+TSharedPtr<SGroomImportOptionsWindow> SGroomImportOptionsWindow::DisplayImportOptions(UGroomImportOptions* ImportOptions, UGroomCacheImportOptions* GroomCacheImportOptions, UGroomHairGroupsPreview* GroupsPreview, UGroomHairGroupsMapping* GroupsMapping, const FString& FilePath)
 {
 	// If there's no groom cache to import, don't show its import options
 	UGroomCacheImportOptions* GroomCacheOptions = GroomCacheImportOptions && GroomCacheImportOptions->ImportSettings.bImportGroomCache ? GroomCacheImportOptions : nullptr;
-	return DisplayOptions(ImportOptions, GroomCacheOptions, GroupsPreview, FilePath, EGroomOptionsVisibility::All, LOCTEXT("GroomImportWindowTitle", "Groom Import Options"), LOCTEXT("Import", "Import"));
+	return DisplayOptions(ImportOptions, GroomCacheOptions, GroupsPreview, GroupsMapping, FilePath, EGroomOptionsVisibility::All, LOCTEXT("GroomImportWindowTitle", "Groom Import Options"), LOCTEXT("Import", "Import"));
 }
 
-TSharedPtr<SGroomImportOptionsWindow> SGroomImportOptionsWindow::DisplayRebuildOptions(UGroomImportOptions* ImportOptions, UGroomHairGroupsPreview* GroupsPreview, const FString& FilePath)
+TSharedPtr<SGroomImportOptionsWindow> SGroomImportOptionsWindow::DisplayRebuildOptions(UGroomImportOptions* ImportOptions, UGroomHairGroupsPreview* GroupsPreview, UGroomHairGroupsMapping* GroupsMapping, const FString& FilePath)
 {
-	return DisplayOptions(ImportOptions, nullptr, GroupsPreview, FilePath, EGroomOptionsVisibility::BuildOptions, LOCTEXT("GroomRebuildWindowTitle ", "Groom Build Options"), LOCTEXT("Build", "Build"));
+	return DisplayOptions(ImportOptions, nullptr, GroupsPreview, GroupsMapping, FilePath, EGroomOptionsVisibility::BuildOptions, LOCTEXT("GroomRebuildWindowTitle ", "Groom Build Options"), LOCTEXT("Build", "Build"));
 }
 
 
diff --git a/Engine/Plugins/Runtime/HairStrands/Source/HairStrandsEditor/Private/GroomImportOptionsWindow.h b/Engine/Plugins/Runtime/HairStrands/Source/HairStrandsEditor/Private/GroomImportOptionsWindow.h
--- a/Engine/Plugins/Runtime/HairStrands/Source/HairStrandsEditor/Private/GroomImportOptionsWindow.h
+++ b/Engine/Plugins/Runtime/HairStrands/Source/HairStrandsEditor/Private/GroomImportOptionsWindow.h
@@ -15,6 +15,7 @@ class SButton;
 class UGroomImportOptions;
 class UGroomCacheImportOptions;
 class UGroomHairGroupsPreview;
+class UGroomHairGroupsMapping;
 struct FHairGroupInfo;
 
 UENUM()
@@ -48,6 +49,7 @@ public:
 		: _ImportOptions(nullptr)
 		, _GroomCacheImportOptions(nullptr)
 		, _GroupsPreview(nullptr)
+		, _GroupsMapping(nullptr)
 		, _WidgetWindow()
 		, _FullPath()
 		, _ButtonLabel()
@@ -56,6 +58,7 @@ public:
 	SLATE_ARGUMENT(UGroomImportOptions*, ImportOptions)
 	SLATE_ARGUMENT(UGroomCacheImportOptions*, GroomCacheImportOptions)
 	SLATE_ARGUMENT(UGroomHairGroupsPreview*, GroupsPreview)
+	SLATE_ARGUMENT(UGroomHairGroupsMapping*, GroupsMapping)
 	SLATE_ARGUMENT(TSharedPtr<SWindow>, WidgetWindow)
 	SLATE_ARGUMENT(FText, FullPath)
 	SLATE_ARGUMENT(FText, ButtonLabel)
@@ -69,11 +72,13 @@ public:
 		UGroomImportOptions* ImportOptions, 
 		UGroomCacheImportOptions* GroomCacheImportOptions, 
 		UGroomHairGroupsPreview* GroupsPreview,
+		UGroomHairGroupsMapping* GroupsMapping,
 		const FString& FilePath);
 
 	static TSharedPtr<SGroomImportOptionsWindow> DisplayRebuildOptions(
 		UGroomImportOptions* ImportOptions,
 		UGroomHairGroupsPreview* GroupsPreview,
+		UGroomHairGroupsMapping* GroupsMapping,
 		const FString& FilePath);
 
 	FReply OnImport()
@@ -116,6 +121,7 @@ public:
 		, GroomCacheImportOptions(nullptr)
 		, bShouldImport(false)
 		, GroupsPreview(nullptr)
+		, GroupsMapping(nullptr)
 	{}
 
 private:
@@ -137,9 +143,11 @@ private:
 	TSharedPtr<class IDetailsView> DetailsView;
 	TSharedPtr<class IDetailsView> GroomCacheDetailsView;
 	TSharedPtr<class IDetailsView> DetailsView2;
+	TSharedPtr<class IDetailsView> DetailsView3;
 	TWeakPtr<SWindow> WidgetWindow;
 	TSharedPtr<SButton> ImportButton;
 	bool bShouldImport;
 public:
 	UGroomHairGroupsPreview* GroupsPreview;
+	UGroomHairGroupsMapping* GroupsMapping;
 };
diff --git a/Engine/Plugins/Runtime/HairStrands/Source/HairStrandsEditor/Private/HairStrandsEditor.cpp b/Engine/Plugins/Runtime/HairStrands/Source/HairStrandsEditor/Private/HairStrandsEditor.cpp
--- a/Engine/Plugins/Runtime/HairStrands/Source/HairStrandsEditor/Private/HairStrandsEditor.cpp
+++ b/Engine/Plugins/Runtime/HairStrands/Source/HairStrandsEditor/Private/HairStrandsEditor.cpp
@@ -23,6 +23,7 @@
 #include "GroomEditorCommands.h"
 #include "GroomEditorMode.h"
 #include "GroomPluginSettings.h"
+#include "GroomImportOptions.h"
 
 #include "AssetRegistry/AssetRegistryModule.h"
 #include "FileHelpers.h"
@@ -121,6 +122,7 @@ void FGroomEditor::StartupModule()
 		PropertyModule.RegisterCustomClassLayout(UGroomComponent::StaticClass()->GetFName(), FOnGetDetailCustomizationInstance::CreateStatic(&FGroomComponentDetailsCustomization::MakeInstance));
 		PropertyModule.RegisterCustomClassLayout(UGroomBindingAsset::StaticClass()->GetFName(), FOnGetDetailCustomizationInstance::CreateStatic(&FGroomBindingDetailsCustomization::MakeInstance));
 		PropertyModule.RegisterCustomClassLayout(UGroomCreateBindingOptions::StaticClass()->GetFName(), FOnGetDetailCustomizationInstance::CreateStatic(&FGroomCreateBindingDetailsCustomization::MakeInstance));
+		PropertyModule.RegisterCustomClassLayout(UGroomHairGroupsMapping::StaticClass()->GetFName(), FOnGetDetailCustomizationInstance::CreateStatic(&FGroomHairGroomRemappingDetailsCustomization::MakeInstance));
 		PropertyModule.RegisterCustomPropertyTypeLayout(FGroomCacheImportSettings::StaticStruct()->GetFName(), FOnGetPropertyTypeCustomizationInstance::CreateStatic(&FGroomCacheImportSettingsCustomization::MakeInstance));
 	}
 
diff --git a/Engine/Plugins/Runtime/HairStrands/Source/HairStrandsEditor/Private/HairStrandsFactory.cpp b/Engine/Plugins/Runtime/HairStrands/Source/HairStrandsEditor/Private/HairStrandsFactory.cpp
--- a/Engine/Plugins/Runtime/HairStrands/Source/HairStrandsEditor/Private/HairStrandsFactory.cpp
+++ b/Engine/Plugins/Runtime/HairStrands/Source/HairStrandsEditor/Private/HairStrandsFactory.cpp
@@ -170,7 +170,7 @@ UObject* UHairStrandsFactory::FactoryCreateFile(UClass* InClass, UObject* InPare
 		if (!bIsUnattended)
 		{
 			// Display import options and handle user cancellation
-			TSharedPtr<SGroomImportOptionsWindow> GroomOptionWindow = SGroomImportOptionsWindow::DisplayImportOptions(ImportOptions, GroomCacheImportOptions, GroupsPreview, Filename);
+			TSharedPtr<SGroomImportOptionsWindow> GroomOptionWindow = SGroomImportOptionsWindow::DisplayImportOptions(ImportOptions, GroomCacheImportOptions, GroupsPreview, nullptr/*GroupsMapping*/, Filename);
 			if (!GroomOptionWindow->ShouldImport())
 			{
 				bOutOperationCanceled = true;
diff --git a/Engine/Plugins/Runtime/HairStrands/Source/HairStrandsEditor/Private/HairStrandsImporter.cpp b/Engine/Plugins/Runtime/HairStrands/Source/HairStrandsEditor/Private/HairStrandsImporter.cpp
--- a/Engine/Plugins/Runtime/HairStrands/Source/HairStrandsEditor/Private/HairStrandsImporter.cpp
+++ b/Engine/Plugins/Runtime/HairStrands/Source/HairStrandsEditor/Private/HairStrandsImporter.cpp
@@ -17,7 +17,6 @@ static FAutoConsoleVariableRef CVarHairStrandsLoadAsset(TEXT("r.HairStrands.Reim
 
 struct FHairGroupSettings
 {
-	//uint32 GroupIndex = 0;
 	FHairGroupInfoWithVisibility Info;
 	FHairGroupsRendering Rendering;
 	FHairGroupsPhysics Physics;
@@ -26,6 +25,7 @@ struct FHairGroupSettings
 	float EffectiveLODBias = 0;
 };
 
+// Extract all the settings from a groom asset
 static TArray<FHairGroupSettings> GetHairGroupSettings(const UGroomAsset* In)
 {
 	TArray<FHairGroupSettings> Out;
@@ -43,6 +43,7 @@ static TArray<FHairGroupSettings> GetHairGroupSettings(const UGroomAsset* In)
 	return Out;
 }
 
+// Apply all the settings to a groom asset
 static void SetHairGroupSettings(UGroomAsset* OutAsset, const TArray<FHairGroupSettings>& InSettings, const TArray<int32>& OldToNewGroupIndexMapping)
 {
 	for (uint32 GroupIndex = 0, GroupCount = OutAsset->GetNumHairGroups(); GroupIndex < GroupCount; ++GroupIndex)
@@ -69,93 +70,72 @@ static void SetHairGroupSettings(UGroomAsset* OutAsset, const TArray<FHairGroupS
 	}
 }
 
-static bool FindHairGroupSettings(uint32 InGroupID, const FName& InGroupName, const TArray<FHairGroupSettings>& InSettings, FHairGroupSettings& Out)
-{
-	for (const FHairGroupSettings& Group : InSettings)
-	{
-		// For now, only use GroupName, as if no group name is provided, the group name is built from the GroupID
-		if (Group.Info.GroupName == InGroupName)
-		{
-			Out = Group;
-			return true;
-		}
-	}
-	return false;
-}
-
+// Remap group settings
 static void RemapHairGroupSettings(
-	const UGroomAsset* InGroomAsset, 
-	const FHairDescriptionGroups& InHairDescriptionGroups, 
-	const TArray<FHairGroupsInterpolation> ImportedInterpolationSettings,
-	TArray<FHairGroupSettings>& OutNewSettings, 
-	TArray<int32>& OutRemapping)
+	const TArray<int32>& NewToOldGroupIndexMapping,
+	const TArray<FHairGroupSettings> InOldSettings,
+	const FHairDescriptionGroups& InNewHairDescriptionGroups, 
+	const TArray<FHairGroupsInterpolation> InNewImportedInterpolationSettings,
+	TArray<FHairGroupSettings>& OutNewSettings)
 {
 	const FHairGroupsLOD DefaultLOD = FHairGroupsLOD::GetDefault();
 
-	const bool bUseGroupRemapping = GHairStrandsReimportGroupRemapping > 0;
-	if (bUseGroupRemapping)
+	OutNewSettings.SetNum(InNewHairDescriptionGroups.HairGroups.Num());
+	for (const FHairDescriptionGroup& NewGroupDesc : InNewHairDescriptionGroups.HairGroups)
 	{
-		// 1. Extract existing settings from the old groom asset
-		const TArray<FHairGroupSettings> OldSettings = GetHairGroupSettings(InGroomAsset);
-		OutRemapping.Init(-1, OldSettings.Num());
-
-		// 2. Remap old settings to the new asset (use GroupName/GroomID to do the remapping)	
-		const uint32 GroupCount = InHairDescriptionGroups.HairGroups.Num();
-		OutNewSettings.SetNum(GroupCount);
-		for (const FHairDescriptionGroup& GroupDesc : InHairDescriptionGroups.HairGroups)
+		FHairGroupSettings& NewGroup = OutNewSettings[NewGroupDesc.Info.GroupIndex];
+
+		const int32 OldGroupIndex = NewToOldGroupIndexMapping[NewGroupDesc.Info.GroupIndex];
+		if (InOldSettings.IsValidIndex(OldGroupIndex))
 		{
-			FHairGroupSettings OldGroup;
-			FHairGroupSettings& NewGroup = OutNewSettings[GroupDesc.Info.GroupIndex];
-			if (FindHairGroupSettings(GroupDesc.Info.GroupID, GroupDesc.Info.GroupName, OldSettings, OldGroup))
-			{
-				NewGroup = OldGroup;
-				OutRemapping[OldGroup.Info.GroupIndex] = GroupDesc.Info.GroupIndex;
-			}
-			else
-			{
-				NewGroup.LOD = DefaultLOD;
-			}
+			NewGroup = InOldSettings[OldGroupIndex];
+		}
+		else
+		{
+			NewGroup.LOD = DefaultLOD;
+		}
 
-			// Preserve the new group GroupIndex/GroupID/GroupName
-			NewGroup.Info.GroupIndex = GroupDesc.Info.GroupIndex;
-			NewGroup.Info.GroupID    = GroupDesc.Info.GroupID;
-			NewGroup.Info.GroupName  = GroupDesc.Info.GroupName;
 
-			// Use import interpolation settings
-			NewGroup.Interpolation   = ImportedInterpolationSettings[GroupDesc.Info.GroupIndex];
-		}
+		// Preserve the new group GroupIndex/GroupID/GroupName
+		NewGroup.Info.GroupIndex = NewGroupDesc.Info.GroupIndex;
+		NewGroup.Info.GroupID    = NewGroupDesc.Info.GroupID;
+		NewGroup.Info.GroupName  = NewGroupDesc.Info.GroupName;
+
+		// Use import interpolation settings
+		NewGroup.Interpolation   = InNewImportedInterpolationSettings[NewGroupDesc.Info.GroupIndex];
 	}
 }
 
 void RemapHairGroupInteprolationSettings(
-	const UGroomAsset* InGroomAsset, 
-	const FHairDescriptionGroups& InHairDescriptionGroups,
-	TArray<FHairGroupsInterpolation>& OutInterpolationSettings)
+	const UGroomAsset* InOldGroomAsset, 
+	const FHairDescriptionGroups& InNewHairDescriptionGroups,
+	const UGroomHairGroupsMapping* InGroupsMapping, 
+	TArray<FHairGroupsInterpolation>& OutNewInterpolationSettings)
 {
 	// Remap existing interpolation settings based on GroupName/GroupID if possible, otherwise initialize them to default
 	// Populate the interpolation settings based on the group count from the description
-	const uint32 GroupCount = InHairDescriptionGroups.HairGroups.Num();
-	OutInterpolationSettings.Init(FHairGroupsInterpolation(), GroupCount);
+	const uint32 NewGroupCount = InNewHairDescriptionGroups.HairGroups.Num();
+	OutNewInterpolationSettings.Init(FHairGroupsInterpolation(), NewGroupCount);
 
 	const bool bUseGroupRemapping = GHairStrandsReimportGroupRemapping > 0;
-	if (bUseGroupRemapping)
+	if (InGroupsMapping && bUseGroupRemapping)
 	{
-		const TArray<FHairGroupSettings> OldSettings = GetHairGroupSettings(InGroomAsset);
-		for (const FHairDescriptionGroup& Group : InHairDescriptionGroups.HairGroups)
+		const TArray<FHairGroupSettings> OldSettings = GetHairGroupSettings(InOldGroomAsset);
+		for (const FHairDescriptionGroup& NewGroupDesc : InNewHairDescriptionGroups.HairGroups)
 		{			
-			FHairGroupSettings GroupSettings;
-			if (FindHairGroupSettings(Group.Info.GroupID, Group.Info.GroupName, OldSettings, GroupSettings))
+			const int32 OldGroupIndex = InGroupsMapping->NewToOldGroupIndexMapping[NewGroupDesc.Info.GroupIndex];
+			if (OldSettings.IsValidIndex(OldGroupIndex))
 			{
-				OutInterpolationSettings[Group.Info.GroupIndex] = GroupSettings.Interpolation;
+				OutNewInterpolationSettings[NewGroupDesc.Info.GroupIndex] = OldSettings[OldGroupIndex].Interpolation;
 			}
 		}
 	}
 	else
 	{
 		// Old path
-		for (uint32 GroupIndex = 0; GroupIndex < GroupCount; ++GroupIndex)
+		for (uint32 GroupIndex = 0; GroupIndex < NewGroupCount; ++GroupIndex)
 		{
-			OutInterpolationSettings[GroupIndex] = InGroomAsset->GetHairGroupsInterpolation()[GroupIndex];
+			OutNewInterpolationSettings[GroupIndex] = InOldGroomAsset->GetHairGroupsInterpolation()[GroupIndex];
 		}
 	}
 }
@@ -171,13 +151,13 @@ FHairImportContext::FHairImportContext(UGroomImportOptions* InImportOptions, UOb
 {
 }
 
-UGroomAsset* FHairStrandsImporter::ImportHair(const FHairImportContext& ImportContext, FHairDescription& HairDescription, UGroomAsset* ExistingHair)
+UGroomAsset* FHairStrandsImporter::ImportHair(const FHairImportContext& ImportContext, FHairDescription& NewHairDescription, UGroomAsset* OldExistingHair, const UGroomHairGroupsMapping* InGroupsMapping)
 {
 	const uint32 GroupCount = ImportContext.ImportOptions->InterpolationSettings.Num();
 	UGroomAsset* OutHairAsset = nullptr;
-	if (ExistingHair)
+	if (OldExistingHair)
 	{
-		OutHairAsset = ExistingHair;
+		OutHairAsset = OldExistingHair;
 	}
 	else
 	{
@@ -190,20 +170,33 @@ UGroomAsset* FHairStrandsImporter::ImportHair(const FHairImportContext& ImportCo
 	}
 
 	const bool bUseGroupRemapping = GHairStrandsReimportGroupRemapping > 0;
-	if (ExistingHair && bUseGroupRemapping)
+	if (OldExistingHair && bUseGroupRemapping)
 	{
-		FHairDescriptionGroups HairDescriptionGroups;
-		FGroomBuilder::BuildHairDescriptionGroups(HairDescription, HairDescriptionGroups);
+		FHairDescriptionGroups NewHairDescriptionGroups;
+		FGroomBuilder::BuildHairDescriptionGroups(NewHairDescription, NewHairDescriptionGroups);
 
-		// Remap group settings when reimporting the asset and group number/order mismatch
-		TArray<FHairGroupSettings> NewSettings;
+		// 1. Extract/Build group remapping
 		TArray<int32> OldToNewGroupIndexMapping;
+		TArray<int32> NewToOldGroupIndexMapping;
+		if (InGroupsMapping)
+		{
+			OldToNewGroupIndexMapping = InGroupsMapping->OldToNewGroupIndexMapping;
+			NewToOldGroupIndexMapping = InGroupsMapping->NewToOldGroupIndexMapping;
+		}
+		else
+		{
+			UGroomHairGroupsMapping::RemapHairDescriptionGroups(OldExistingHair->GetHairDescriptionGroups(), NewHairDescriptionGroups, OldToNewGroupIndexMapping);
+			UGroomHairGroupsMapping::RemapHairDescriptionGroups(NewHairDescriptionGroups, OldExistingHair->GetHairDescriptionGroups(), NewToOldGroupIndexMapping);
+		}
+
+		// 2. Remap group settings when reimporting the asset and group number/order mismatch
+		TArray<FHairGroupSettings> NewSettings;	
 		RemapHairGroupSettings(
-			ExistingHair, 
-			HairDescriptionGroups, 
-			ImportContext.ImportOptions->InterpolationSettings,
-			NewSettings, 
-			OldToNewGroupIndexMapping);
+			NewToOldGroupIndexMapping,
+			GetHairGroupSettings(OldExistingHair),
+			NewHairDescriptionGroups, 
+			ImportContext.ImportOptions->InterpolationSettings, // New interpolation settings, already remapped
+			NewSettings);
 
 		// 3. Apply new settings to the new asset
 		OutHairAsset->ClearNumGroup(GroupCount);
@@ -226,13 +219,13 @@ UGroomAsset* FHairStrandsImporter::ImportHair(const FHairImportContext& ImportCo
 	check(OutHairAsset->AreGroupsValid());
 	check(uint32(OutHairAsset->GetNumHairGroups()) == GroupCount);
 
-	OutHairAsset->CommitHairDescription(MoveTemp(HairDescription), UGroomAsset::EHairDescriptionType::Source);
+	OutHairAsset->CommitHairDescription(MoveTemp(NewHairDescription), UGroomAsset::EHairDescriptionType::Source);
 
 	const bool bSucceeded = OutHairAsset->CacheDerivedDatas();
 	if (!bSucceeded)
 	{
 		// Purge the newly created asset that failed to import
-		if (OutHairAsset != ExistingHair)
+		if (OutHairAsset != OldExistingHair)
 		{
 			OutHairAsset->ClearFlags(RF_Standalone);
 			CollectGarbage(GARBAGE_COLLECTION_KEEPFLAGS);
diff --git a/Engine/Plugins/Runtime/HairStrands/Source/HairStrandsEditor/Private/ReimportHairStrandsFactory.cpp b/Engine/Plugins/Runtime/HairStrands/Source/HairStrandsEditor/Private/ReimportHairStrandsFactory.cpp
--- a/Engine/Plugins/Runtime/HairStrands/Source/HairStrandsEditor/Private/ReimportHairStrandsFactory.cpp
+++ b/Engine/Plugins/Runtime/HairStrands/Source/HairStrandsEditor/Private/ReimportHairStrandsFactory.cpp
@@ -130,11 +130,11 @@ namespace ReimportGroomAssetHelpers
 		}
 	}
 
-	UGroomAsset* ReimportGroom(const FString& SourceFilename, FMD5Hash& FileHash, TSharedPtr<IGroomTranslator> Translator, FHairImportContext& HairImportContext, FHairDescription& HairDescription, UGroomAsset* HairAsset, const FGroomCacheImportSettings& Settings)
+	UGroomAsset* ReimportGroom(const FString& SourceFilename, FMD5Hash& FileHash, TSharedPtr<IGroomTranslator> Translator, FHairImportContext& HairImportContext, FHairDescription& HairDescription, UGroomAsset* HairAsset, const FGroomCacheImportSettings& Settings, const UGroomHairGroupsMapping* InGroupsMapping)
 	{
 		if (Settings.bImportGroomAsset)
 		{
-			UGroomAsset* ReimportedHair = FHairStrandsImporter::ImportHair(HairImportContext, HairDescription, HairAsset);
+			UGroomAsset* ReimportedHair = FHairStrandsImporter::ImportHair(HairImportContext, HairDescription, HairAsset, InGroupsMapping);
 			if (ReimportedHair)
 			{
 				// Move the transient ImportOptions to the asset package and set it on the GroomAssetImportData for serialization
@@ -259,9 +259,10 @@ namespace ReimportGroomCacheHelpers
 }
 
 void RemapHairGroupInteprolationSettings(
-	const UGroomAsset* InGroomAsset, 
-	const FHairDescriptionGroups& InHairDescriptionGroups,
-	TArray<FHairGroupsInterpolation>& OutInterpolationSettings);
+	const UGroomAsset* InOldGroomAsset, 
+	const FHairDescriptionGroups& InNewHairDescriptionGroups,
+	const UGroomHairGroupsMapping* InGroupsMapping, 
+	TArray<FHairGroupsInterpolation>& OutNewInterpolationSettings);
 
 EReimportResult::Type UReimportHairStrandsFactory::Reimport(UObject* Obj)
 {
@@ -343,6 +344,8 @@ EReimportResult::Type UReimportHairStrandsFactory::Reimport(UObject* Obj)
 
 	FGroomAnimationInfo AnimInfo;
 
+	UGroomHairGroupsMapping* GroupsMapping = NewObject<UGroomHairGroupsMapping>();
+
 	// Load the alembic file upfront to preview & report any potential issue
 	FHairDescriptionGroups OutDescription;
 	{
@@ -358,8 +361,11 @@ EReimportResult::Type UReimportHairStrandsFactory::Reimport(UObject* Obj)
 
 		FGroomBuilder::BuildHairDescriptionGroups(HairDescription, OutDescription);
 
+		// Group remapping
+		GroupsMapping->Map(HairAsset->GetHairDescriptionGroups(), OutDescription);
+
 		// Remap existing interpolation settings based on GroupName if possible, otherwise initialize them to default
-		RemapHairGroupInteprolationSettings(HairAsset, OutDescription, GroomReimportOptions->InterpolationSettings);
+		RemapHairGroupInteprolationSettings(HairAsset, OutDescription, GroupsMapping, GroomReimportOptions->InterpolationSettings);
 	}
 
 	FGroomCacheImporter::SetupImportSettings(GroomCacheReimportOptions->ImportSettings, AnimInfo);
@@ -387,7 +393,7 @@ EReimportResult::Type UReimportHairStrandsFactory::Reimport(UObject* Obj)
 			}
 		}
 
-		TSharedPtr<SGroomImportOptionsWindow> GroomOptionWindow = SGroomImportOptionsWindow::DisplayImportOptions(GroomReimportOptions, GroomCacheReimportOptions, GroupsPreview, CurrentFilename);
+		TSharedPtr<SGroomImportOptionsWindow> GroomOptionWindow = SGroomImportOptionsWindow::DisplayImportOptions(GroomReimportOptions, GroomCacheReimportOptions, GroupsPreview, GroupsMapping, CurrentFilename);
 
 		if (!GroomOptionWindow->ShouldImport())
 		{
@@ -416,7 +422,7 @@ EReimportResult::Type UReimportHairStrandsFactory::Reimport(UObject* Obj)
 	// Reimport the GroomAsset
 	FMD5Hash ReimportedFileHash;
 	FHairImportContext HairImportContext(GroomReimportOptions, HairAsset, nullptr, FName(), EObjectFlags::RF_Public | EObjectFlags::RF_Standalone | EObjectFlags::RF_Transactional);
-	UGroomAsset* AssociatedGroomAsset = ReimportGroomAssetHelpers::ReimportGroom(CurrentFilename, ReimportedFileHash, SelectedTranslator, HairImportContext, HairDescription, HairAsset, GroomCacheReimportOptions->ImportSettings);
+	UGroomAsset* AssociatedGroomAsset = ReimportGroomAssetHelpers::ReimportGroom(CurrentFilename, ReimportedFileHash, SelectedTranslator, HairImportContext, HairDescription, HairAsset, GroomCacheReimportOptions->ImportSettings, GroupsMapping);
 
 	// Reimport the GroomCaches
 	ReimportGroomCacheHelpers::ReimportGroomCaches(CurrentFilename, ReimportedFileHash, SelectedTranslator, HairImportContext, AnimInfo, AssociatedGroomAsset, GroomCacheReimportOptions->ImportSettings, SourceGroomCache);
diff --git a/Engine/Plugins/Runtime/HairStrands/Source/HairStrandsEditor/Public/HairStrandsImporter.h b/Engine/Plugins/Runtime/HairStrands/Source/HairStrandsEditor/Public/HairStrandsImporter.h
--- a/Engine/Plugins/Runtime/HairStrands/Source/HairStrandsEditor/Public/HairStrandsImporter.h
+++ b/Engine/Plugins/Runtime/HairStrands/Source/HairStrandsEditor/Public/HairStrandsImporter.h
@@ -8,6 +8,7 @@
 class FHairDescription;
 class UGroomAsset;
 class UGroomImportOptions;
+class UGroomHairGroupsMapping;
 
 struct HAIRSTRANDSEDITOR_API FHairImportContext
 {
@@ -22,5 +23,5 @@ struct HAIRSTRANDSEDITOR_API FHairImportContext
 
 struct HAIRSTRANDSEDITOR_API FHairStrandsImporter
 {
-	static UGroomAsset* ImportHair(const FHairImportContext& ImportContext, FHairDescription& HairDescription, UGroomAsset* ExistingHair = nullptr);
+	static UGroomAsset* ImportHair(const FHairImportContext& ImportContext, FHairDescription& HairDescription, UGroomAsset* ExistingHair = nullptr, const UGroomHairGroupsMapping* GroupsMapping=nullptr);
 };
