Fix groom deformer no longer working.

diff --git a/Engine/Plugins/Runtime/HairStrands/Source/HairStrandsDeformer/Private/DeformerGroomInterfaceUtils.h b/Engine/Plugins/Runtime/HairStrands/Source/HairStrandsDeformer/Private/DeformerGroomInterfaceUtils.h
--- a/Engine/Plugins/Runtime/HairStrands/Source/HairStrandsDeformer/Private/DeformerGroomInterfaceUtils.h
+++ b/Engine/Plugins/Runtime/HairStrands/Source/HairStrandsDeformer/Private/DeformerGroomInterfaceUtils.h
@@ -13,6 +13,11 @@
 
 namespace UE::Groom::Private
 {
+	FORCEINLINE bool HasDeformationEnabledOrHasMeshDeformer(const UGroomComponent* GroomComponent, const int32 GroupIndex)
+	{
+		return GroomComponent->IsDeformationEnable(GroupIndex) || GroomComponent->MeshDeformer;
+	}
+
 	FORCEINLINE USkeletalMeshComponent* GetGroupSkelMesh(const UGroomComponent* GroomComponent, const int32 GroupIndex, int32& MeshLOD)
 	{
 		if(GroomComponent)
@@ -49,7 +54,7 @@ namespace UE::Groom::Private
  				const uint32 NumGroups = GroomComponent->GetGroupCount();
  				for (uint32 GroupIndex = 0; GroupIndex < NumGroups; ++GroupIndex)
  				{
- 					if(GroomComponent->IsDeformationEnable(GroupIndex) && GroomComponent->GetGroupInstance(GroupIndex))
+ 					if(HasDeformationEnabledOrHasMeshDeformer(GroomComponent, GroupIndex) && GroomComponent->GetGroupInstance(GroupIndex))
  					{
  						GroupInstances.Add(GroomComponent->GetGroupInstance(GroupIndex));
  					}
@@ -336,7 +341,7 @@ namespace UE::Groom::Private
 				const uint32 NumGroups = GroomComponent->GetGroupCount();
 				for (uint32 GroupIndex = 0; GroupIndex < NumGroups; ++GroupIndex)
 				{
-					if(GroomComponent->IsDeformationEnable(GroupIndex) && GroomComponent->GetGroupInstance(GroupIndex))
+					if(HasDeformationEnabledOrHasMeshDeformer(GroomComponent, GroupIndex) && GroomComponent->GetGroupInstance(GroupIndex))
 					{
 						InvocationCounts.Add(GroupOffsets[GroupIndex+1]-GroupOffsets[GroupIndex]);
 						TotalCount += InvocationCounts.Last();
@@ -366,7 +371,7 @@ namespace UE::Groom::Private
 				const uint32 NumGroups = GroomComponent->GetGroupCount();
 				for (uint32 GroupIndex = 0; GroupIndex < NumGroups; ++GroupIndex)
 				{ 
-					if(GroomComponent->IsDeformationEnable(GroupIndex) && GroomComponent->GetGroupInstance(GroupIndex))
+					if(HasDeformationEnabledOrHasMeshDeformer(GroomComponent, GroupIndex) && GroomComponent->GetGroupInstance(GroupIndex))
 					{
 						GroupData.GroupIndices.Add(GroupIndex);
 					}
