Fix DF crash when no guides in groom asset + Better support for enable deformation flags + QA game deformer update

#virtualized

diff --git a/Engine/Plugins/Runtime/HairStrands/Source/HairStrandsCore/Private/GroomComponent.cpp b/Engine/Plugins/Runtime/HairStrands/Source/HairStrandsCore/Private/GroomComponent.cpp
--- a/Engine/Plugins/Runtime/HairStrands/Source/HairStrandsCore/Private/GroomComponent.cpp
+++ b/Engine/Plugins/Runtime/HairStrands/Source/HairStrandsCore/Private/GroomComponent.cpp
@@ -4100,6 +4100,20 @@ bool UGroomComponent::CanEditChange(const FProperty* InProperty) const
 			#endif
 			return bIsEditable;
 		}
+		if (PropertyName == GET_MEMBER_NAME_STRING_CHECKED(UGroomComponent, MeshDeformer))
+		{
+			if(GroomAsset)
+			{
+				for(int32 GroupIndex = 0, NumGroups = GroomAsset->GetNumHairGroups(); GroupIndex < NumGroups; ++GroupIndex)
+				{
+					if(GroomAsset->IsDeformationEnable(GroupIndex))
+					{
+						return true;
+					}
+				}
+			}
+			return false;
+		}
 	}
 
 	return Super::CanEditChange(InProperty);
diff --git a/Engine/Plugins/Runtime/HairStrands/Source/HairStrandsCore/Public/GroomComponent.h b/Engine/Plugins/Runtime/HairStrands/Source/HairStrandsCore/Public/GroomComponent.h
--- a/Engine/Plugins/Runtime/HairStrands/Source/HairStrandsCore/Public/GroomComponent.h
+++ b/Engine/Plugins/Runtime/HairStrands/Source/HairStrandsCore/Public/GroomComponent.h
@@ -60,7 +60,7 @@ public:
 	FHairSimulationSettings SimulationSettings;
 
 	/** If set the MeshDeformer will be applied on groom instance for deformation. */
-	UPROPERTY(EditAnywhere, BlueprintReadOnly, Category = "Deformer")
+	UPROPERTY(EditAnywhere, BlueprintReadOnly, Category = "Deformer", Meta = (ToolTip = "Mesh deformer that will be applied on groom instance for deformation. Enable the groom asset deformation flags on groups to be able to set it."))
 	TObjectPtr<UMeshDeformer> MeshDeformer;
 
 	/** Object containing state for the bound MeshDeformer. */
diff --git a/Engine/Plugins/Runtime/HairStrands/Source/HairStrandsDeformer/Private/DeformerGroomInterfaceUtils.h b/Engine/Plugins/Runtime/HairStrands/Source/HairStrandsDeformer/Private/DeformerGroomInterfaceUtils.h
--- a/Engine/Plugins/Runtime/HairStrands/Source/HairStrandsDeformer/Private/DeformerGroomInterfaceUtils.h
+++ b/Engine/Plugins/Runtime/HairStrands/Source/HairStrandsDeformer/Private/DeformerGroomInterfaceUtils.h
@@ -15,7 +15,7 @@ namespace UE::Groom::Private
 {
 	FORCEINLINE bool HasDeformationEnabledOrHasMeshDeformer(const UGroomComponent* GroomComponent, const int32 GroupIndex)
 	{
-		return GroomComponent->IsDeformationEnable(GroupIndex) || GroomComponent->MeshDeformer;
+		return GroomComponent->IsDeformationEnable(GroupIndex);
 	}
 
 	FORCEINLINE USkeletalMeshComponent* GetGroupSkelMesh(const UGroomComponent* GroomComponent, const int32 GroupIndex, int32& MeshLOD)
@@ -83,37 +83,39 @@ namespace UE::Groom::Private
 					{
 						if(USkeletalMeshComponent* SkelMesh =  GetGroupSkelMesh(GroomComponent, GroupIndex, MeshLOD))
 						{
-							const FSkeletalMeshRenderData* const RenderData = SkelMesh->GetSkinnedAsset()->GetResourceForRendering();
-
-							if(RenderData && RenderData->LODRenderData.IsValidIndex(MeshLOD))
+							if(USkinnedAsset* SkinnedAsset = SkelMesh->GetSkinnedAsset())
 							{
-								// Get the matching active indices used for skinning
-								TArray<FBoneIndexType> GuidesBones = RenderData->LODRenderData[MeshLOD].ActiveBoneIndices;
+								const FSkeletalMeshRenderData* const RenderData = SkinnedAsset->GetResourceForRendering();
+								if(RenderData && RenderData->LODRenderData.IsValidIndex(MeshLOD))
+								{
+									// Get the matching active indices used for skinning
+									TArray<FBoneIndexType> GuidesBones = RenderData->LODRenderData[MeshLOD].ActiveBoneIndices;
 								
-								// Get inv ref pose matrices
-								const TArray<FMatrix44f>* RefBasesInvMatrix = &SkelMesh->GetSkinnedAsset()->GetRefBasesInvMatrix();
+									// Get inv ref pose matrices
+									const TArray<FMatrix44f>* RefBasesInvMatrix = &SkelMesh->GetSkinnedAsset()->GetRefBasesInvMatrix();
 
-								// Get the component space transforms
-								BonesTransforms.Init(FTransform::Identity, RefBasesInvMatrix->Num());
+									// Get the component space transforms
+									BonesTransforms.Init(FTransform::Identity, RefBasesInvMatrix->Num());
 								
-								SkelMesh->GetSkinnedAsset()->FillComponentSpaceTransforms(
-									SkelMesh->GetBoneSpaceTransforms(), GuidesBones, BonesTransforms);
+									SkelMesh->GetSkinnedAsset()->FillComponentSpaceTransforms(
+										SkelMesh->GetBoneSpaceTransforms(), GuidesBones, BonesTransforms);
 
-								// Fill ref to local matrices
-								RefToLocals.Init(FMatrix44f::Identity, RefBasesInvMatrix->Num());
+									// Fill ref to local matrices
+									RefToLocals.Init(FMatrix44f::Identity, RefBasesInvMatrix->Num());
 								
-								for (int32 BoneIndex = 0; BoneIndex < RefToLocals.Num(); ++BoneIndex)
-								{
-									RefToLocals[BoneIndex] = (*RefBasesInvMatrix)[BoneIndex] * (FMatrix44f)BonesTransforms[BoneIndex].ToMatrixWithScale();
-								}
+									for (int32 BoneIndex = 0; BoneIndex < RefToLocals.Num(); ++BoneIndex)
+									{
+										RefToLocals[BoneIndex] = (*RefBasesInvMatrix)[BoneIndex] * (FMatrix44f)BonesTransforms[BoneIndex].ToMatrixWithScale();
+									}
 							
-								const FTransform& BonesTransform = SkelMesh->GetComponentTransform();
-								const FTransform& GroupTransform = GroomComponent->GetGroupInstance(GroupIndex)->GetCurrentLocalToWorld();
-								//SkelMesh->GetCurrentRefToLocalMatrices(RefToLocals, SkelMesh->MeshObject->GetLOD(), &RequiredExtraBoneIndices);//MeshLOD);
+									const FTransform& BonesTransform = SkelMesh->GetComponentTransform();
+									const FTransform& GroupTransform = GroomComponent->GetGroupInstance(GroupIndex)->GetCurrentLocalToWorld();
+									//SkelMesh->GetCurrentRefToLocalMatrices(RefToLocals, SkelMesh->MeshObject->GetLOD(), &RequiredExtraBoneIndices);//MeshLOD);
 		
-								SkeletalTransforms.Add(FMatrix44f(BonesTransform.ToMatrixWithScale() * GroupTransform.ToInverseMatrixWithScale()));
-								BonesRefToLocals.Add(RefToLocals);
-								SkeletalMeshes.Add(SkelMesh->MeshObject);
+									SkeletalTransforms.Add(FMatrix44f(BonesTransform.ToMatrixWithScale() * GroupTransform.ToInverseMatrixWithScale()));
+									BonesRefToLocals.Add(RefToLocals);
+									SkeletalMeshes.Add(SkelMesh->MeshObject);
+								}
 							}
 						}
 					}
diff --git a/Engine/Plugins/Runtime/HairStrands/Source/HairStrandsSolver/Private/GroomSolverComponent.cpp b/Engine/Plugins/Runtime/HairStrands/Source/HairStrandsSolver/Private/GroomSolverComponent.cpp
--- a/Engine/Plugins/Runtime/HairStrands/Source/HairStrandsSolver/Private/GroomSolverComponent.cpp
+++ b/Engine/Plugins/Runtime/HairStrands/Source/HairStrandsSolver/Private/GroomSolverComponent.cpp
@@ -113,7 +113,7 @@ void UGroomSolverComponent::SelectDynamicCurves()
 				for (uint32 GroupIndex = 0; GroupIndex < NumGroups; ++GroupIndex)
 				{
 					const FHairGroupInstance* GroupInstance = GroomComponent->GetGroupInstance(GroupIndex);
-					if(GroomComponent->IsDeformationEnable(GroupIndex) && GroupInstance)
+					if(GroomComponent->IsDeformationEnable(GroupIndex) && GroupInstance && GroupInstance->Guides.IsValid())
 					{
 						const int32 NumTotalCurves = GroupInstance->Guides.GetData().GetNumCurves();
 						const int32 NumTotalPoints = GroupInstance->Guides.GetData().GetNumPoints();
