Change hair strands radius to be scale by the instance's transform to be consistent with the scale of the instance.

This change of behavior can be disabled with r.HairStrands.Strands.ScaleHairRadiusByTransform 0

diff --git a/Engine/Plugins/Runtime/HairStrands/Source/HairStrandsCore/Private/HairStrandsInterpolation.cpp b/Engine/Plugins/Runtime/HairStrands/Source/HairStrandsCore/Private/HairStrandsInterpolation.cpp
--- a/Engine/Plugins/Runtime/HairStrands/Source/HairStrandsCore/Private/HairStrandsInterpolation.cpp
+++ b/Engine/Plugins/Runtime/HairStrands/Source/HairStrandsCore/Private/HairStrandsInterpolation.cpp
@@ -58,6 +58,9 @@ static FAutoConsoleVariableRef CVarHairStrands_MotionVector_CheckViewID(TEXT("r.
 static int32 GStrandHairInterpolationForceSingleGuide = 1;
 static FAutoConsoleVariableRef CVarHairStrands_Interpolation_ForceSingleGuide(TEXT("r.HairStrands.Strands.Interpolation.ForceSingleGuide"), GStrandHairInterpolationForceSingleGuide, TEXT("Force to use a single guide during the interpolation. This allows to reduce the cost of interpolation, but can create clumpy looks."));
 
+static int32 GVarHairStrands_ScaleHairRadiusByTransform = 1;
+static FAutoConsoleVariableRef CVarHairStrands_ScaleHairRadiusByTransform(TEXT("r.HairStrands.Strands.ScaleHairRadiusByTransform"), GVarHairStrands_ScaleHairRadiusByTransform, TEXT("Scale hair strands radius by the instance's transform to make hair radius consistent with the instance scale."));
+
 ////////////////////////////////////////////////////////////////////////////////////////////////////////////////////////
 
 uint32 GetHairStrandsForceSingleGuideInterpolation()
@@ -1777,9 +1780,18 @@ FHairGroupPublicData::FVertexFactoryInput InternalComputeHairStrandsVertexInputD
 	bRaytracingGeometry = Instance->Strands.RenRaytracingResource != nullptr && EnumHasAnyFlags(Instance->Strands.ViewRayTracingMask, EHairViewRayTracingMask::RayTracing);
 #endif
 
+	// Scale hair radius based on the max. scaling factor from the instance's transform
+	float InstanceScale = 1.f;
+	if (GVarHairStrands_ScaleHairRadiusByTransform > 0)
+	{
+		const FVector3d S = Instance->LocalToWorld.GetScale3D();
+		InstanceScale = FMath::Max3(S.X, S.Y, S.Z);
+	}
+
 	OutVFInput.Strands.Common.RegisteredIndex = Instance->RegisteredIndex;
 	OutVFInput.Strands.Common.Radius = (GStrandHairWidth > 0 ? GStrandHairWidth : Instance->Strands.Modifier.HairWidth) * 0.5f;
 	OutVFInput.Strands.Common.Radius *= Instance->HairGroupPublicData->GetActiveStrandsRadiusScale();
+	OutVFInput.Strands.Common.Radius *= InstanceScale;
 	OutVFInput.Strands.Common.RootScale = Instance->Strands.Modifier.HairRootScale;
 	OutVFInput.Strands.Common.TipScale = Instance->Strands.Modifier.HairTipScale;
 	OutVFInput.Strands.Common.RaytracingRadiusScale = (GHairRaytracingRadiusScale > 0 ? GHairRaytracingRadiusScale : Instance->Strands.Modifier.HairRaytracingRadiusScale);
