Add groom support for batched import and reimport.

* Now groom asset respect UEditorPerProjectUserSettings for batch reimporting.
* A ImportAll button has also be added to the import window for batch importing.

diff --git a/Engine/Plugins/Runtime/HairStrands/Source/HairStrandsEditor/Private/GroomImportOptionsWindow.cpp b/Engine/Plugins/Runtime/HairStrands/Source/HairStrandsEditor/Private/GroomImportOptionsWindow.cpp
--- a/Engine/Plugins/Runtime/HairStrands/Source/HairStrandsEditor/Private/GroomImportOptionsWindow.cpp
+++ b/Engine/Plugins/Runtime/HairStrands/Source/HairStrandsEditor/Private/GroomImportOptionsWindow.cpp
@@ -19,37 +19,38 @@
 
 #define LOCTEXT_NAMESPACE "GroomImportOptionsWindow"
 
-void SGroomImportOptionsWindow::UpdateStatus(UGroomHairGroupsPreview* Description) const
+FGroomImportStatus GetGroomImportStatus(UGroomHairGroupsPreview* InDescription, UGroomCacheImportOptions* InGroomCacheImportOptions, UGroomHairGroupsMapping* InGroupMapping)
 {
-	CurrentStatus = EHairDescriptionStatus::None;
-	if (!Description)
+	FGroomImportStatus Out;
+	Out.Status = EHairDescriptionStatus::None;
+	if (!InDescription)
 	{
-		EnumAddFlags(CurrentStatus, EHairDescriptionStatus::Unknown);
-		return;
+		EnumAddFlags(Out.Status, EHairDescriptionStatus::Unknown);
+		return Out;
 	}
 
-	const bool bImportGroomAsset = !GroomCacheImportOptions || GroomCacheImportOptions->ImportSettings.bImportGroomAsset;
-	const bool bImportGroomCache = GroomCacheImportOptions && GroomCacheImportOptions->ImportSettings.bImportGroomCache;
+	const bool bImportGroomAsset = !InGroomCacheImportOptions || InGroomCacheImportOptions->ImportSettings.bImportGroomAsset;
+	const bool bImportGroomCache = InGroomCacheImportOptions && InGroomCacheImportOptions->ImportSettings.bImportGroomCache;
 	if (!bImportGroomAsset && !bImportGroomCache)
 	{
-		EnumAddFlags(CurrentStatus, EHairDescriptionStatus::Unknown);
-		return;
+		EnumAddFlags(Out.Status, EHairDescriptionStatus::Unknown);
+		return Out;
 	}
 
-	if (Description->Groups.Num() == 0)
+	if (InDescription->Groups.Num() == 0)
 	{
-		EnumAddFlags(CurrentStatus, EHairDescriptionStatus::NoGroup);
-		return;
+		EnumAddFlags(Out.Status, EHairDescriptionStatus::NoGroup);
+		return Out;
 	}
 
 	// Check the validity of the groom to import
 
 	bool bGuidesOnly = false;
-	for (const FGroomHairGroupPreview& Group : Description->Groups)
+	for (const FGroomHairGroupPreview& Group : InDescription->Groups)
 	{
 		if (Group.CurveCount == 0)
 		{
-			EnumAddFlags(CurrentStatus, EHairDescriptionStatus::NoCurve);
+			EnumAddFlags(Out.Status, EHairDescriptionStatus::NoCurve);
 			if (Group.GuideCount > 0)
 			{
 				bGuidesOnly = true;
@@ -59,110 +60,130 @@ void SGroomImportOptionsWindow::UpdateStatus(UGroomHairGroupsPreview* Descriptio
 	}
 
 	// Check if any curve or point have been trimmed
-	for (const FGroomHairGroupPreview& Group : Description->Groups)
+	for (const FGroomHairGroupPreview& Group : InDescription->Groups)
 	{
 		if (Group.Flags & uint32(EHairGroupInfoFlags::HasTrimmedCurve))
 		{
-			EnumAddFlags(CurrentStatus, EHairDescriptionStatus::CurveLimit);
+			EnumAddFlags(Out.Status, EHairDescriptionStatus::CurveLimit);
 		}
 		if (Group.Flags & uint32(EHairGroupInfoFlags::HasTrimmedPoint))
 		{
-			EnumAddFlags(CurrentStatus, EHairDescriptionStatus::PointLimit);
+			EnumAddFlags(Out.Status, EHairDescriptionStatus::PointLimit);
 		}
 		if (Group.Flags & uint32(EHairGroupInfoFlags::HasInvalidPoint))
 		{
-			EnumAddFlags(CurrentStatus, EHairDescriptionStatus::InvalidPoint);
+			EnumAddFlags(Out.Status, EHairDescriptionStatus::InvalidPoint);
 		}		
 	}
 
+	if (InGroupMapping && !InGroupMapping->HasValidMapping())
+	{
+		EnumAddFlags(Out.Status, EHairDescriptionStatus::InvalidGroupMapping);
+	}
+
 	if (!bImportGroomCache)
 	{
-		EnumAddFlags(CurrentStatus, EHairDescriptionStatus::GroomValid);
-		return;
+		EnumAddFlags(Out.Status, EHairDescriptionStatus::GroomValid);
+		return Out;
 	}
 
 	// Update the states of the properties being monitored
-	bImportGroomAssetState = GroomCacheImportOptions->ImportSettings.bImportGroomAsset;
-	bImportGroomCacheState = GroomCacheImportOptions->ImportSettings.bImportGroomCache;
-	GroomAsset = GroomCacheImportOptions->ImportSettings.GroomAsset;
+	Out.bImportGroomAssetState = InGroomCacheImportOptions->ImportSettings.bImportGroomAsset;
+	Out.bImportGroomCacheState = InGroomCacheImportOptions->ImportSettings.bImportGroomCache;
+	Out.GroomAsset = InGroomCacheImportOptions->ImportSettings.GroomAsset;
 
-	if (!GroomCacheImportOptions->ImportSettings.bImportGroomAsset)
+	if (!InGroomCacheImportOptions->ImportSettings.bImportGroomAsset)
 	{
 		// When importing a groom cache with a provided groom asset, check their compatibility
-		UGroomAsset* GroomAssetForCache = Cast<UGroomAsset>(GroomCacheImportOptions->ImportSettings.GroomAsset.TryLoad());
+		UGroomAsset* GroomAssetForCache = Cast<UGroomAsset>(InGroomCacheImportOptions->ImportSettings.GroomAsset.TryLoad());
 		if (!GroomAssetForCache)
 		{
 			// No groom asset provided or loaded but one is needed with this setting
-			EnumAddFlags(CurrentStatus, bGuidesOnly ? EHairDescriptionStatus::GuidesOnly : EHairDescriptionStatus::GroomCache);
-			return;
+			EnumAddFlags(Out.Status, bGuidesOnly ? EHairDescriptionStatus::GuidesOnly : EHairDescriptionStatus::GroomCache);
+			return Out;
 		}
 
 		const TArray<FHairGroupPlatformData>& GroomHairGroupsData = GroomAssetForCache->GetHairGroupsPlatformData();
-		if (GroomHairGroupsData.Num() != Description->Groups.Num())
+		if (GroomHairGroupsData.Num() != InDescription->Groups.Num())
 		{
-			EnumAddFlags(CurrentStatus, bGuidesOnly ? EHairDescriptionStatus::GuidesOnlyIncompatible : EHairDescriptionStatus::GroomCacheIncompatible);
-			return;
+			EnumAddFlags(Out.Status, bGuidesOnly ? EHairDescriptionStatus::GuidesOnlyIncompatible : EHairDescriptionStatus::GroomCacheIncompatible);
+			return Out;
 		}
 
 		for (int32 Index = 0; Index < GroomHairGroupsData.Num(); ++Index)
 		{
 			// Check the strands compatibility
-			if (!bGuidesOnly && Description->Groups[Index].CurveCount != GroomHairGroupsData[Index].Strands.BulkData.GetNumCurves())
+			if (!bGuidesOnly && InDescription->Groups[Index].CurveCount != GroomHairGroupsData[Index].Strands.BulkData.GetNumCurves())
 			{
-				EnumAddFlags(CurrentStatus, EHairDescriptionStatus::GroomCacheIncompatible);
+				EnumAddFlags(Out.Status, EHairDescriptionStatus::GroomCacheIncompatible);
 				break;
 			}
 
 			// Check the guides compatibility if there were strands tagged as guides
 			// Otherwise, guides will be generated according to the groom asset interpolation settings
 			// and compatibility cannot be determined here
-			if (Description->Groups[Index].GuideCount > 0 &&
-				Description->Groups[Index].GuideCount != GroomHairGroupsData[Index].Guides.BulkData.GetNumCurves())
+			if (InDescription->Groups[Index].GuideCount > 0 &&
+				InDescription->Groups[Index].GuideCount != GroomHairGroupsData[Index].Guides.BulkData.GetNumCurves())
 			{
-				EnumAddFlags(CurrentStatus, bGuidesOnly ? EHairDescriptionStatus::GuidesOnlyIncompatible : EHairDescriptionStatus::GroomCacheIncompatible);
+				EnumAddFlags(Out.Status, bGuidesOnly ? EHairDescriptionStatus::GuidesOnlyIncompatible : EHairDescriptionStatus::GroomCacheIncompatible);
 				break;
 			}
 		}
 
-		EnumAddFlags(CurrentStatus, bGuidesOnly ? EHairDescriptionStatus::GuidesOnlyCompatible : EHairDescriptionStatus::GroomCacheCompatible);
+		EnumAddFlags(Out.Status, bGuidesOnly ? EHairDescriptionStatus::GuidesOnlyCompatible : EHairDescriptionStatus::GroomCacheCompatible);
 	}
 	else
 	{
 		// A guides-only groom cannot be imported as asset, but otherwise the imported groom asset
 		// is always compatible with the groom cache since they are from the same file
-		EnumAddFlags(CurrentStatus, bGuidesOnly ? EHairDescriptionStatus::GuidesOnly : EHairDescriptionStatus::GroomValid);
+		EnumAddFlags(Out.Status, bGuidesOnly ? EHairDescriptionStatus::GuidesOnly : EHairDescriptionStatus::GroomValid);
 	}
+
+	return Out;
 }
 
+void SGroomImportOptionsWindow::UpdateStatus(UGroomHairGroupsPreview* Description) const
+{
+	ImportStatus = GetGroomImportStatus(Description, GroomCacheImportOptions, GroupsMapping);
+}
 
-FText SGroomImportOptionsWindow::GetStatusText() const
+FText GetGroomImportStatusText(FGroomImportStatus In, bool bAddPrefix)
 {
 	FString Out;
-	if (EnumHasAnyFlags(CurrentStatus, EHairDescriptionStatus::Error)) 						{ Out += LOCTEXT("GroomOptionsWindow_ValidationText0",  "Error\n").ToString(); }
-	else if (EnumHasAnyFlags(CurrentStatus, EHairDescriptionStatus::Warning)) 				{ Out += LOCTEXT("GroomOptionsWindow_ValidationText1",  "Warning\n").ToString(); }
-	else if (EnumHasAnyFlags(CurrentStatus, EHairDescriptionStatus::Valid))					{ Out += LOCTEXT("GroomOptionsWindow_ValidationText2",  "Valid\n").ToString(); }
-
-	if (EnumHasAnyFlags(CurrentStatus, EHairDescriptionStatus::NoCurve)) 					{ Out += LOCTEXT("GroomOptionsWindow_ValidationText3",  "Some groups have 0 curves.\n").ToString(); }
-	if (EnumHasAnyFlags(CurrentStatus, EHairDescriptionStatus::NoGroup)) 					{ Out += LOCTEXT("GroomOptionsWindow_ValidationText4",  "The groom does not contain any group.\n").ToString(); }
-	if (EnumHasAnyFlags(CurrentStatus, EHairDescriptionStatus::GroomCache))					{ Out += LOCTEXT("GroomOptionsWindow_ValidationText5",  "A compatible groom asset must be provided to import the groom cache.\n").ToString(); }
-	if (EnumHasAnyFlags(CurrentStatus, EHairDescriptionStatus::GroomCacheCompatible)) 		{ Out += LOCTEXT("GroomOptionsWindow_ValidationText6",  "The groom cache is compatible with the groom asset provided.\n").ToString(); }
-	if (EnumHasAnyFlags(CurrentStatus, EHairDescriptionStatus::GroomCacheIncompatible))		{ Out += LOCTEXT("GroomOptionsWindow_ValidationText7",  "The groom cache is incompatible with the groom asset provided.\n").ToString(); }
-	if (EnumHasAnyFlags(CurrentStatus, EHairDescriptionStatus::GuidesOnly)) 				{ Out += LOCTEXT("GroomOptionsWindow_ValidationText8",  "Only guides were detected. A compatible groom asset must be provided.\n").ToString(); }
-	if (EnumHasAnyFlags(CurrentStatus, EHairDescriptionStatus::GuidesOnlyCompatible)) 		{ Out += LOCTEXT("GroomOptionsWindow_ValidationText9",  "Only guides were detected. The groom asset provided is compatible.\n").ToString(); }
-	if (EnumHasAnyFlags(CurrentStatus, EHairDescriptionStatus::GuidesOnlyIncompatible))		{ Out += LOCTEXT("GroomOptionsWindow_ValidationText10", "Only guides were detected. The groom asset provided is incompatible.\n").ToString(); }
-	if (EnumHasAnyFlags(CurrentStatus, EHairDescriptionStatus::CurveLimit))					{ Out += LOCTEXT("GroomOptionsWindow_ValidationText11", "At least one group contains more curves than allowed limit (Max:4M). Curves beyond that limit will be trimmed.\n").ToString(); static_assert(HAIR_MAX_NUM_CURVE_PER_GROUP == 4194303); }
-	if (EnumHasAnyFlags(CurrentStatus, EHairDescriptionStatus::PointLimit))					{ Out += LOCTEXT("GroomOptionsWindow_ValidationText12", "At least one group contains more control points per curve than the allowed limit (Max:255). Control points beyond that limit will be trimmed.\n").ToString(); static_assert(HAIR_MAX_NUM_POINT_PER_CURVE == 255); }
-	if (EnumHasAnyFlags(CurrentStatus, EHairDescriptionStatus::InvalidPoint)) 				{ Out += LOCTEXT("GroomOptionsWindow_ValidationText13", "At least one group contains a curve with invalid points. These curves will be trimmed from the asset.\n").ToString(); }
-	if (EnumHasAnyFlags(CurrentStatus, EHairDescriptionStatus::Unknown)) 					{ Out += LOCTEXT("GroomOptionsWindow_ValidationText14", "Unknown\n").ToString(); }
+	if (bAddPrefix)
+	{
+		if (EnumHasAnyFlags(In.Status, EHairDescriptionStatus::Error)) 						{ Out += LOCTEXT("GroomOptionsWindow_ValidationText0",  "Error\n").ToString(); }
+		else if (EnumHasAnyFlags(In.Status, EHairDescriptionStatus::Warning)) 				{ Out += LOCTEXT("GroomOptionsWindow_ValidationText1",  "Warning\n").ToString(); }
+		else if (EnumHasAnyFlags(In.Status, EHairDescriptionStatus::Valid))					{ Out += LOCTEXT("GroomOptionsWindow_ValidationText2",  "Valid\n").ToString(); }
+	}
+
+	if (EnumHasAnyFlags(In.Status, EHairDescriptionStatus::NoCurve)) 					{ Out += LOCTEXT("GroomOptionsWindow_ValidationText3",  "Some groups have 0 curves.\n").ToString(); }
+	if (EnumHasAnyFlags(In.Status, EHairDescriptionStatus::NoGroup)) 					{ Out += LOCTEXT("GroomOptionsWindow_ValidationText4",  "The groom does not contain any group.\n").ToString(); }
+	if (EnumHasAnyFlags(In.Status, EHairDescriptionStatus::GroomCache))					{ Out += LOCTEXT("GroomOptionsWindow_ValidationText5",  "A compatible groom asset must be provided to import the groom cache.\n").ToString(); }
+	if (EnumHasAnyFlags(In.Status, EHairDescriptionStatus::GroomCacheCompatible)) 		{ Out += LOCTEXT("GroomOptionsWindow_ValidationText6",  "The groom cache is compatible with the groom asset provided.\n").ToString(); }
+	if (EnumHasAnyFlags(In.Status, EHairDescriptionStatus::GroomCacheIncompatible))		{ Out += LOCTEXT("GroomOptionsWindow_ValidationText7",  "The groom cache is incompatible with the groom asset provided.\n").ToString(); }
+	if (EnumHasAnyFlags(In.Status, EHairDescriptionStatus::GuidesOnly)) 				{ Out += LOCTEXT("GroomOptionsWindow_ValidationText8",  "Only guides were detected. A compatible groom asset must be provided.\n").ToString(); }
+	if (EnumHasAnyFlags(In.Status, EHairDescriptionStatus::GuidesOnlyCompatible)) 		{ Out += LOCTEXT("GroomOptionsWindow_ValidationText9",  "Only guides were detected. The groom asset provided is compatible.\n").ToString(); }
+	if (EnumHasAnyFlags(In.Status, EHairDescriptionStatus::GuidesOnlyIncompatible))		{ Out += LOCTEXT("GroomOptionsWindow_ValidationText10", "Only guides were detected. The groom asset provided is incompatible.\n").ToString(); }
+	if (EnumHasAnyFlags(In.Status, EHairDescriptionStatus::CurveLimit))					{ Out += LOCTEXT("GroomOptionsWindow_ValidationText11", "At least one group contains more curves than allowed limit (Max:4M). Curves beyond that limit will be trimmed.\n").ToString(); static_assert(HAIR_MAX_NUM_CURVE_PER_GROUP == 4194303); }
+	if (EnumHasAnyFlags(In.Status, EHairDescriptionStatus::PointLimit))					{ Out += LOCTEXT("GroomOptionsWindow_ValidationText12", "At least one group contains more control points per curve than the allowed limit (Max:255). Control points beyond that limit will be trimmed.\n").ToString(); static_assert(HAIR_MAX_NUM_POINT_PER_CURVE == 255); }
+	if (EnumHasAnyFlags(In.Status, EHairDescriptionStatus::InvalidPoint)) 				{ Out += LOCTEXT("GroomOptionsWindow_ValidationText13", "At least one group contains a curve with invalid points. These curves will be trimmed from the asset.\n").ToString(); }
+	if (EnumHasAnyFlags(In.Status, EHairDescriptionStatus::InvalidGroupMapping)) 		{ Out += LOCTEXT("GroomOptionsWindow_ValidationText15", "No mapping found using group names.\n").ToString(); }
+	if (EnumHasAnyFlags(In.Status, EHairDescriptionStatus::Unknown)) 					{ Out += LOCTEXT("GroomOptionsWindow_ValidationText14", "Unknown\n").ToString(); }
 
 	return FText::FromString(Out);
 }
 
+FText SGroomImportOptionsWindow::GetStatusText() const
+{
+	return GetGroomImportStatusText(ImportStatus, true);
+}
+
 FSlateColor SGroomImportOptionsWindow::GetStatusColor() const
 {
-	if (EnumHasAnyFlags(CurrentStatus, EHairDescriptionStatus::Error))  	{ return FLinearColor(0.80f, 0, 0, 1); }
-	if (EnumHasAnyFlags(CurrentStatus, EHairDescriptionStatus::Warning))  	{ return FLinearColor(0.80f, 0.80f, 0, 1); }
-	if (EnumHasAnyFlags(CurrentStatus, EHairDescriptionStatus::Valid))		{ return FLinearColor(0, 0.80f, 0, 1); }
+	if (EnumHasAnyFlags(ImportStatus.Status, EHairDescriptionStatus::Error))  	{ return FLinearColor(0.80f, 0, 0, 1); }
+	if (EnumHasAnyFlags(ImportStatus.Status, EHairDescriptionStatus::Warning))  { return FLinearColor(0.80f, 0.80f, 0, 1); }
+	if (EnumHasAnyFlags(ImportStatus.Status, EHairDescriptionStatus::Valid))	{ return FLinearColor(0, 0.80f, 0, 1); }
 
 	return FLinearColor(1, 1, 1);
 }
@@ -244,6 +265,59 @@ FText GetHairAttributeLocText(EHairAttribute In, uint32 InFlags)
 	return FText::GetEmpty();
 }
 
+TSharedRef<SWidget> SGroomImportOptionsWindow::AddImportButtons(FText InMainButtonLabel, bool bShowImportAllButton)
+{
+	if (bShowImportAllButton)
+	{
+		return SNew(SUniformGridPanel)
+		.SlotPadding(2)
+		+ SUniformGridPanel::Slot(0, 0)
+		[
+			SAssignNew(ImportButton, SButton)
+			.HAlign(HAlign_Center)
+			.Text(InMainButtonLabel)
+			.IsEnabled(this, &SGroomImportOptionsWindow::CanImport)
+			.OnClicked(this, &SGroomImportOptionsWindow::OnImport)
+		]
+		+ SUniformGridPanel::Slot(1, 0)
+		[
+			SAssignNew(ImportButton, SButton)
+			.HAlign(HAlign_Center)
+			.Text(LOCTEXT("GroomOptionWindow_ImportAll", "Import All"))
+			.ToolTipText(LOCTEXT("GroomOptionWindow_ImportAll_ToolTip", "Import all files with these same settings"))
+			.IsEnabled(this, &SGroomImportOptionsWindow::CanImport)
+			.OnClicked(this, &SGroomImportOptionsWindow::OnImportAll)
+		]
+		+ SUniformGridPanel::Slot(2, 0)
+		[
+			SNew(SButton)
+			.HAlign(HAlign_Center)
+			.Text(LOCTEXT("GroomOptionWindow_Cancel", "Cancel"))
+			.OnClicked(this, &SGroomImportOptionsWindow::OnCancel)
+		];
+	}
+	else
+	{
+		return SNew(SUniformGridPanel)
+		.SlotPadding(2)
+		+ SUniformGridPanel::Slot(0, 0)
+		[
+			SAssignNew(ImportButton, SButton)
+			.HAlign(HAlign_Center)
+			.Text(InMainButtonLabel)
+			.IsEnabled(this, &SGroomImportOptionsWindow::CanImport)
+			.OnClicked(this, &SGroomImportOptionsWindow::OnImport)
+		]
+		+ SUniformGridPanel::Slot(1, 0)
+		[
+			SNew(SButton)
+			.HAlign(HAlign_Center)
+			.Text(LOCTEXT("GroomOptionWindow_Cancel", "Cancel"))
+			.OnClicked(this, &SGroomImportOptionsWindow::OnCancel)
+		];
+	}
+}
+
 void SGroomImportOptionsWindow::Construct(const FArguments& InArgs)
 {
 	ImportOptions = InArgs._ImportOptions;
@@ -269,7 +343,7 @@ void SGroomImportOptionsWindow::Construct(const FArguments& InArgs)
 	GroomCacheDetailsView = PropertyEditorModule.CreateDetailView(DetailsViewArgs);
 	GroomCacheDetailsView->SetObject(GroomCacheImportOptions);
 
-	CurrentStatus = EHairDescriptionStatus::None;
+	ImportStatus.Status = EHairDescriptionStatus::None;
 	UpdateStatus(GroupsPreview);
 
 	// Aggregate attributes from all groups (ideally we should display each group attribute separately, to check if one groom is not missing data)
@@ -414,23 +488,7 @@ void SGroomImportOptionsWindow::Construct(const FArguments& InArgs)
 		.HAlign(HAlign_Right)
 		.Padding(2)
 		[
-			SNew(SUniformGridPanel)
-			.SlotPadding(2)
-			+ SUniformGridPanel::Slot(0, 0)
-			[
-				SAssignNew(ImportButton, SButton)
-				.HAlign(HAlign_Center)
-				.Text(InArgs._ButtonLabel)
-				.IsEnabled(this, &SGroomImportOptionsWindow::CanImport)
-				.OnClicked(this, &SGroomImportOptionsWindow::OnImport)
-			]
-			+ SUniformGridPanel::Slot(1, 0)
-			[
-				SNew(SButton)
-				.HAlign(HAlign_Center)
-				.Text(LOCTEXT("Cancel", "Cancel"))
-				.OnClicked(this, &SGroomImportOptionsWindow::OnCancel)
-			]
+			AddImportButtons(InArgs._ButtonLabel, InArgs._bShowImportAllButton)
 		];
 
 	// Insert all the optional attributes
@@ -459,12 +517,12 @@ void SGroomImportOptionsWindow::Construct(const FArguments& InArgs)
 
 bool SGroomImportOptionsWindow::CanImport() const
 {
-	bool bNeedUpdate = CurrentStatus == EHairDescriptionStatus::None;
+	bool bNeedUpdate = ImportStatus.Status == EHairDescriptionStatus::None;
 	if (GroomCacheImportOptions)
 	{
-		bNeedUpdate |= bImportGroomAssetState != GroomCacheImportOptions->ImportSettings.bImportGroomAsset;
-		bNeedUpdate |= bImportGroomCacheState != GroomCacheImportOptions->ImportSettings.bImportGroomCache;
-		bNeedUpdate |= GroomAsset != GroomCacheImportOptions->ImportSettings.GroomAsset;
+		bNeedUpdate |= ImportStatus.bImportGroomAssetState != GroomCacheImportOptions->ImportSettings.bImportGroomAsset;
+		bNeedUpdate |= ImportStatus.bImportGroomCacheState != GroomCacheImportOptions->ImportSettings.bImportGroomCache;
+		bNeedUpdate |= ImportStatus.GroomAsset != GroomCacheImportOptions->ImportSettings.GroomAsset;
 	}
 
 	if (bNeedUpdate)
@@ -472,7 +530,7 @@ bool SGroomImportOptionsWindow::CanImport() const
 		UpdateStatus(GroupsPreview);
 	}
 
-	return EnumHasAnyFlags(CurrentStatus, EHairDescriptionStatus::Valid | EHairDescriptionStatus::Warning);
+	return EnumHasAnyFlags(ImportStatus.Status, EHairDescriptionStatus::Valid | EHairDescriptionStatus::Warning);
 }
 
 enum class EGroomOptionsVisibility : uint8
@@ -493,7 +551,8 @@ TSharedPtr<SGroomImportOptionsWindow> DisplayOptions(
 	const FString& FilePath, 
 	EGroomOptionsVisibility VisibilityFlag, 
 	FText WindowTitle, 
-	FText InButtonLabel)
+	FText InButtonLabel,
+	bool bInShowImportAllButton)
 {
 	TSharedRef<SWindow> Window = SNew(SWindow)
 		.Title(WindowTitle)
@@ -529,6 +588,7 @@ TSharedPtr<SGroomImportOptionsWindow> DisplayOptions(
 		.WidgetWindow(Window)
 		.FullPath(FText::FromString(FileName))
 		.ButtonLabel(InButtonLabel)
+		.bShowImportAllButton(bInShowImportAllButton)
 	);
 
 	TSharedPtr<SWindow> ParentWindow;
@@ -544,16 +604,16 @@ TSharedPtr<SGroomImportOptionsWindow> DisplayOptions(
 	return OptionsWindow;
 }
 
-TSharedPtr<SGroomImportOptionsWindow> SGroomImportOptionsWindow::DisplayImportOptions(UGroomImportOptions* ImportOptions, UGroomCacheImportOptions* GroomCacheImportOptions, UGroomHairGroupsPreview* GroupsPreview, UGroomHairGroupsMapping* GroupsMapping, const FString& FilePath)
+TSharedPtr<SGroomImportOptionsWindow> SGroomImportOptionsWindow::DisplayImportOptions(UGroomImportOptions* ImportOptions, UGroomCacheImportOptions* GroomCacheImportOptions, UGroomHairGroupsPreview* GroupsPreview, UGroomHairGroupsMapping* GroupsMapping, const FString& FilePath, bool bShowImportAllButton)
 {
 	// If there's no groom cache to import, don't show its import options
 	UGroomCacheImportOptions* GroomCacheOptions = GroomCacheImportOptions && GroomCacheImportOptions->ImportSettings.bImportGroomCache ? GroomCacheImportOptions : nullptr;
-	return DisplayOptions(ImportOptions, GroomCacheOptions, GroupsPreview, GroupsMapping, FilePath, EGroomOptionsVisibility::All, LOCTEXT("GroomImportWindowTitle", "Groom Import Options"), LOCTEXT("Import", "Import"));
+	return DisplayOptions(ImportOptions, GroomCacheOptions, GroupsPreview, GroupsMapping, FilePath, EGroomOptionsVisibility::All, LOCTEXT("GroomImportWindowTitle", "Groom Import Options"), LOCTEXT("Import", "Import"), bShowImportAllButton);
 }
 
 TSharedPtr<SGroomImportOptionsWindow> SGroomImportOptionsWindow::DisplayRebuildOptions(UGroomImportOptions* ImportOptions, UGroomHairGroupsPreview* GroupsPreview, UGroomHairGroupsMapping* GroupsMapping, const FString& FilePath)
 {
-	return DisplayOptions(ImportOptions, nullptr, GroupsPreview, GroupsMapping, FilePath, EGroomOptionsVisibility::BuildOptions, LOCTEXT("GroomRebuildWindowTitle ", "Groom Build Options"), LOCTEXT("Build", "Build"));
+	return DisplayOptions(ImportOptions, nullptr, GroupsPreview, GroupsMapping, FilePath, EGroomOptionsVisibility::BuildOptions, LOCTEXT("GroomRebuildWindowTitle ", "Groom Build Options"), LOCTEXT("Build", "Build"), false /*bShowImportAllButton*/);
 }
 
 
diff --git a/Engine/Plugins/Runtime/HairStrands/Source/HairStrandsEditor/Private/GroomImportOptionsWindow.h b/Engine/Plugins/Runtime/HairStrands/Source/HairStrandsEditor/Private/GroomImportOptionsWindow.h
--- a/Engine/Plugins/Runtime/HairStrands/Source/HairStrandsEditor/Private/GroomImportOptionsWindow.h
+++ b/Engine/Plugins/Runtime/HairStrands/Source/HairStrandsEditor/Private/GroomImportOptionsWindow.h
@@ -34,14 +34,26 @@ enum class EHairDescriptionStatus : uint32
 	PointLimit				= 1u<<10,
 	CurveLimit				= 1u<<11,
 	InvalidPoint			= 1u<<12,
-	Unknown					= 1u<<13,
+	InvalidGroupMapping		= 1u<<13,
+	Unknown					= 1u<<14,
 
-	Warning = CurveLimit | PointLimit | InvalidPoint,
+	Warning = CurveLimit | PointLimit | InvalidPoint | InvalidGroupMapping,
 	Error   = NoCurve | NoGroup | GroomCacheIncompatible | GuidesOnlyIncompatible | GroomCache | GuidesOnly | Unknown,
 	Valid   = GroomValid | GroomCacheCompatible | GuidesOnlyCompatible
 };
 ENUM_CLASS_FLAGS(EHairDescriptionStatus);
 
+struct FGroomImportStatus
+{
+	EHairDescriptionStatus Status;
+	FSoftObjectPath GroomAsset;
+	bool bImportGroomAssetState = false;
+	bool bImportGroomCacheState = false;
+};
+
+FGroomImportStatus GetGroomImportStatus(UGroomHairGroupsPreview* InDescription, UGroomCacheImportOptions* InGroomCacheImportOptions, UGroomHairGroupsMapping* InGroupMapping);
+FText GetGroomImportStatusText(FGroomImportStatus In, bool bAddPrefix);
+
 class SGroomImportOptionsWindow : public SCompoundWidget
 {
 public:
@@ -53,6 +65,7 @@ public:
 		, _WidgetWindow()
 		, _FullPath()
 		, _ButtonLabel()
+		, _bShowImportAllButton(false)
 	{}
 
 	SLATE_ARGUMENT(UGroomImportOptions*, ImportOptions)
@@ -62,6 +75,7 @@ public:
 	SLATE_ARGUMENT(TSharedPtr<SWindow>, WidgetWindow)
 	SLATE_ARGUMENT(FText, FullPath)
 	SLATE_ARGUMENT(FText, ButtonLabel)
+	SLATE_ARGUMENT(bool, bShowImportAllButton)
 	SLATE_END_ARGS()
 
 public:
@@ -73,7 +87,8 @@ public:
 		UGroomCacheImportOptions* GroomCacheImportOptions, 
 		UGroomHairGroupsPreview* GroupsPreview,
 		UGroomHairGroupsMapping* GroupsMapping,
-		const FString& FilePath);
+		const FString& FilePath,
+		bool bShowImportAllButton);
 
 	static TSharedPtr<SGroomImportOptionsWindow> DisplayRebuildOptions(
 		UGroomImportOptions* ImportOptions,
@@ -84,6 +99,18 @@ public:
 	FReply OnImport()
 	{
 		bShouldImport = true;
+		bShouldImportAll = false;
+		if (WidgetWindow.IsValid())
+		{
+			WidgetWindow.Pin()->RequestDestroyWindow();
+		}
+		return FReply::Handled();
+	}
+
+	FReply OnImportAll()
+	{
+		bShouldImport = true;
+		bShouldImportAll = true;
 		if (WidgetWindow.IsValid())
 		{
 			WidgetWindow.Pin()->RequestDestroyWindow();
@@ -94,6 +121,7 @@ public:
 	FReply OnCancel()
 	{
 		bShouldImport = false;
+		bShouldImportAll = false;
 		if (WidgetWindow.IsValid())
 		{
 			WidgetWindow.Pin()->RequestDestroyWindow();
@@ -116,10 +144,14 @@ public:
 		return bShouldImport;
 	}
 
+	bool ShouldImportAll() const
+	{
+		return bShouldImportAll;
+	}
+
 	SGroomImportOptionsWindow() 
 		: ImportOptions(nullptr)
 		, GroomCacheImportOptions(nullptr)
-		, bShouldImport(false)
 		, GroupsPreview(nullptr)
 		, GroupsMapping(nullptr)
 	{}
@@ -130,12 +162,10 @@ private:
 	void UpdateStatus(UGroomHairGroupsPreview* Description) const;
 	FText GetStatusText() const;
 	FSlateColor GetStatusColor() const;
+	TSharedRef<SWidget> AddImportButtons(FText InMainButtonLabel, bool bShowImportAllButton);
 
 	// Properties that are monitored for changes
-	mutable EHairDescriptionStatus CurrentStatus;
-	mutable FSoftObjectPath GroomAsset;
-	mutable bool bImportGroomAssetState;
-	mutable bool bImportGroomCacheState;
+	mutable FGroomImportStatus ImportStatus;
 
 private:
 	UGroomImportOptions* ImportOptions;
@@ -146,7 +176,8 @@ private:
 	TSharedPtr<class IDetailsView> DetailsView3;
 	TWeakPtr<SWindow> WidgetWindow;
 	TSharedPtr<SButton> ImportButton;
-	bool bShouldImport;
+	bool bShouldImport = false;
+	bool bShouldImportAll = false;
 public:
 	UGroomHairGroupsPreview* GroupsPreview;
 	UGroomHairGroupsMapping* GroupsMapping;
diff --git a/Engine/Plugins/Runtime/HairStrands/Source/HairStrandsEditor/Private/HairStrandsFactory.cpp b/Engine/Plugins/Runtime/HairStrands/Source/HairStrandsEditor/Private/HairStrandsFactory.cpp
--- a/Engine/Plugins/Runtime/HairStrands/Source/HairStrandsEditor/Private/HairStrandsFactory.cpp
+++ b/Engine/Plugins/Runtime/HairStrands/Source/HairStrandsEditor/Private/HairStrandsFactory.cpp
@@ -80,6 +80,7 @@ UObject* UHairStrandsFactory::FactoryCreateFile(UClass* InClass, UObject* InPare
 		|| FApp::IsUnattended()
 		|| IsRunningCommandlet()
 		|| GIsRunningUnattendedScript);
+	const bool ShowImportDialog = !bIsUnattended && !bImportAll;
 
 	// Translate the hair data from the file
 	TSharedPtr<IGroomTranslator> SelectedTranslator = GetTranslator(Filename);
@@ -167,16 +168,22 @@ UObject* UHairStrandsFactory::FactoryCreateFile(UClass* InClass, UObject* InPare
 		}
 
 		// Don't bother saving the options coming from script
-		if (!bIsUnattended)
+		if (ShowImportDialog)
 		{
 			// Display import options and handle user cancellation
-			TSharedPtr<SGroomImportOptionsWindow> GroomOptionWindow = SGroomImportOptionsWindow::DisplayImportOptions(ImportOptions, GroomCacheImportOptions, GroupsPreview, nullptr/*GroupsMapping*/, Filename);
+			TSharedPtr<SGroomImportOptionsWindow> GroomOptionWindow = SGroomImportOptionsWindow::DisplayImportOptions(ImportOptions, GroomCacheImportOptions, GroupsPreview, nullptr/*GroupsMapping*/, Filename, true /*bShowImportAllButton*/);
+
 			if (!GroomOptionWindow->ShouldImport())
 			{
 				bOutOperationCanceled = true;
 				return nullptr;
 			}
 
+			if (GroomOptionWindow->ShouldImportAll())
+			{
+				bImportAll = true;
+			}
+
 			// Save the options as the new default
 			for (const FGroomHairGroupPreview& GroupPreview : GroupsPreview->Groups)
 			{
@@ -290,5 +297,12 @@ TSharedPtr<IGroomTranslator> UHairStrandsFactory::GetTranslator(const FString& F
 	return {};
 }
 
+
+void UHairStrandsFactory::CleanUp() 
+{
+	bImportAll = false;
+	Super::CleanUp();
+}
+
 #undef LOCTEXT_NAMESPACE
 
diff --git a/Engine/Plugins/Runtime/HairStrands/Source/HairStrandsEditor/Private/ReimportHairStrandsFactory.cpp b/Engine/Plugins/Runtime/HairStrands/Source/HairStrandsEditor/Private/ReimportHairStrandsFactory.cpp
--- a/Engine/Plugins/Runtime/HairStrands/Source/HairStrandsEditor/Private/ReimportHairStrandsFactory.cpp
+++ b/Engine/Plugins/Runtime/HairStrands/Source/HairStrandsEditor/Private/ReimportHairStrandsFactory.cpp
@@ -19,6 +19,7 @@
 #include "Misc/PackageName.h"
 #include "Misc/ScopedSlowTask.h"
 #include "UObject/Package.h"
+#include "Editor/EditorPerProjectUserSettings.h"
 
 #include UE_INLINE_GENERATED_CPP_BY_NAME(ReimportHairStrandsFactory)
 
@@ -370,7 +371,7 @@ EReimportResult::Type UReimportHairStrandsFactory::Reimport(UObject* Obj)
 
 	FGroomCacheImporter::SetupImportSettings(GroomCacheReimportOptions->ImportSettings, AnimInfo);
 
-	if (!GIsRunningUnattendedScript && !IsAutomatedImport())
+	// Import asset with or without import window
 	{
 		// Convert the process hair description into hair groups
 		UGroomHairGroupsPreview* GroupsPreview = NewObject<UGroomHairGroupsPreview>();
@@ -393,19 +394,41 @@ EReimportResult::Type UReimportHairStrandsFactory::Reimport(UObject* Obj)
 			}
 		}
 
-		TSharedPtr<SGroomImportOptionsWindow> GroomOptionWindow = SGroomImportOptionsWindow::DisplayImportOptions(GroomReimportOptions, GroomCacheReimportOptions, GroupsPreview, GroupsMapping, CurrentFilename);
-
-		if (!GroomOptionWindow->ShouldImport())
+		// Prevent any UI for automation, unattended and commandlet
+		const bool IsUnattended = IsAutomatedImport() || FApp::IsUnattended() || IsRunningCommandlet() || GIsRunningUnattendedScript;
+		const bool ShowImportDialogAtReimport = (GetDefault<UEditorPerProjectUserSettings>()->bShowImportDialogAtReimport || bForceShowDialog) && !IsUnattended;
+		if (ShowImportDialogAtReimport)
 		{
-			return EReimportResult::Cancelled;
-		}
+			// Not need to show ImportAll button as this is managed by UEditorPerProjectUserSettings
+			TSharedPtr<SGroomImportOptionsWindow> GroomOptionWindow = SGroomImportOptionsWindow::DisplayImportOptions(GroomReimportOptions, GroomCacheReimportOptions, GroupsPreview, GroupsMapping, CurrentFilename, false /*bShowImportAllButton*/);
 
-		// Save the options as the new default
-		for (const FGroomHairGroupPreview& GroupPreview : GroupsPreview->Groups)
+			if (!GroomOptionWindow->ShouldImport())
+			{
+				return EReimportResult::Cancelled;
+			}
+
+			// Save the options as the new default
+			for (const FGroomHairGroupPreview& GroupPreview : GroupsPreview->Groups)
+			{
+				if (GroupPreview.GroupIndex < OutDescription.HairGroups.Num())
+				{
+					GroomReimportOptions->InterpolationSettings[GroupPreview.GroupIndex] = GroupPreview.InterpolationSettings;
+				}
+			}
+		}
+		else
 		{
-			if (GroupPreview.GroupIndex < OutDescription.HairGroups.Num())
+			FGroomImportStatus ImportStatus = GetGroomImportStatus(GroupsPreview, nullptr /*GroomCacheImportOptions*/, GroupsMapping);
+
+			// Display Warning and error in log to not break automatic import flow
+			if (EnumHasAnyFlags(ImportStatus.Status, EHairDescriptionStatus::Error))
+			{
+				UE_LOG(LogReimportHairStrandsFactory, Error, TEXT("Error during groom reimport (file %s) : %s"), *CurrentFilename, *GetGroomImportStatusText(ImportStatus, false /*bAddPrefix*/).ToString());
+				return EReimportResult::Failed;
+			}
+			else if (EnumHasAnyFlags(ImportStatus.Status, EHairDescriptionStatus::Warning))
 			{
-				GroomReimportOptions->InterpolationSettings[GroupPreview.GroupIndex] = GroupPreview.InterpolationSettings;
+				UE_LOG(LogReimportHairStrandsFactory, Warning, TEXT("Warning during groom reimport (file %s) : %s"), *CurrentFilename, *GetGroomImportStatusText(ImportStatus, false /*bAddPrefix*/).ToString());
 			}
 		}
 	}
diff --git a/Engine/Plugins/Runtime/HairStrands/Source/HairStrandsEditor/Public/HairStrandsFactory.h b/Engine/Plugins/Runtime/HairStrands/Source/HairStrandsEditor/Public/HairStrandsFactory.h
--- a/Engine/Plugins/Runtime/HairStrands/Source/HairStrandsEditor/Public/HairStrandsFactory.h
+++ b/Engine/Plugins/Runtime/HairStrands/Source/HairStrandsEditor/Public/HairStrandsFactory.h
@@ -24,6 +24,7 @@ public:
 		const FString& Filename, const TCHAR* Parms, FFeedbackContext* Warn, bool& bOutOperationCanceled) override;
 	virtual bool FactoryCanImport(const FString& Filename) override;
 	virtual void GetSupportedFileExtensions(TArray<FString>& OutExtensions) const override;
+	virtual void CleanUp() override;
 
 protected:
 	void InitTranslators();
@@ -35,5 +36,6 @@ protected:
 
 private:
 	TArray<TSharedPtr<IGroomTranslator>> Translators;
+	bool bImportAll = false;
 };
 
