Fix copy pasting groom material.

diff --git a/Engine/Plugins/Runtime/HairStrands/Source/HairStrandsEditor/HairStrandsEditor.Build.cs b/Engine/Plugins/Runtime/HairStrands/Source/HairStrandsEditor/HairStrandsEditor.Build.cs
--- a/Engine/Plugins/Runtime/HairStrands/Source/HairStrandsEditor/HairStrandsEditor.Build.cs
+++ b/Engine/Plugins/Runtime/HairStrands/Source/HairStrandsEditor/HairStrandsEditor.Build.cs
@@ -39,7 +39,10 @@ namespace UnrealBuildTool.Rules
 					"HairCardGeneratorFramework",
 					"CommonMenuExtensions",
 					"RenderCore",
-					"DataflowEditor"
+					"DataflowEditor",
+					"JsonUtilities",
+					"Json",
+					"ApplicationCore",
 				});
 
 			PrivateIncludePathModuleNames.AddRange(
diff --git a/Engine/Plugins/Runtime/HairStrands/Source/HairStrandsEditor/Private/GroomMaterialDetails.cpp b/Engine/Plugins/Runtime/HairStrands/Source/HairStrandsEditor/Private/GroomMaterialDetails.cpp
--- a/Engine/Plugins/Runtime/HairStrands/Source/HairStrandsEditor/Private/GroomMaterialDetails.cpp
+++ b/Engine/Plugins/Runtime/HairStrands/Source/HairStrandsEditor/Private/GroomMaterialDetails.cpp
@@ -58,6 +58,7 @@
 #include "PropertyEditorDelegates.h"
 #include "PropertyCustomizationHelpers.h"
 #include "GroomCustomAssetEditorToolkit.h"
+#include "JsonObjectConverter.h"
 
 #define LOCTEXT_NAMESPACE "GroomMaterialDetails"
 
@@ -117,47 +118,47 @@ void FGroomMaterialDetails::AddMaterials(IDetailLayoutBuilder& DetailLayout)
 			[
 				SNew(STextBlock)
 				.Font(IDetailLayoutBuilder::GetDetailFont())
-			.Text(LOCTEXT("AddLODLevelCategories_MaterialArrayOperations", "Material Slots"))
+				.Text(LOCTEXT("AddLODLevelCategories_MaterialArrayOperations", "Material Slots"))
 			]
-		.ValueContent()
+			.ValueContent()
 			.HAlign(HAlign_Left)
 			.VAlign(VAlign_Center)
 			[
 				SNew(SVerticalBox)
 				+ SVerticalBox::Slot()
-			.AutoHeight()
-			[
-				SNew(SHorizontalBox)
-				+ SHorizontalBox::Slot()
-			.FillWidth(1.0f)
-			.VAlign(VAlign_Center)
-			[
-				SNew(STextBlock)
-				.Font(IDetailLayoutBuilder::GetDetailFont())
-			.Text(this, &FGroomMaterialDetails::GetMaterialArrayText)
-			]
-		+ SHorizontalBox::Slot()
-			.AutoWidth()
-			.HAlign(HAlign_Center)
-			.VAlign(VAlign_Center)
-			.Padding(2.0f, 1.0f)
-			[
-				SNew(SButton)
-				.ButtonStyle(FAppStyle::Get(), "HoverHintOnly")
-			.Text(LOCTEXT("AddLODLevelCategories_MaterialArrayOpAdd", "Add Material Slot"))
-			.ToolTipText(LOCTEXT("AddLODLevelCategories_MaterialArrayOpAdd_Tooltip", "Add Material Slot at the end of the Material slot array. Those Material slots can be used to override a LODs section, (not the base LOD)"))
-			.ContentPadding(4.0f)
-			.ForegroundColor(FSlateColor::UseForeground())
-			.OnClicked(this, &FGroomMaterialDetails::AddMaterialSlot)
-			.IsEnabled(true)
-			.IsFocusable(false)
-			[
-				SNew(SImage)
-				.Image(FAppStyle::GetBrush("Icons.PlusCircle"))
-			.ColorAndOpacity(FSlateColor::UseForeground())
-			]
-			]
-			]
+				.AutoHeight()
+				[
+					SNew(SHorizontalBox)
+					+ SHorizontalBox::Slot()
+					.FillWidth(1.0f)
+					.VAlign(VAlign_Center)
+					[
+						SNew(STextBlock)
+						.Font(IDetailLayoutBuilder::GetDetailFont())
+						.Text(this, &FGroomMaterialDetails::GetMaterialArrayText)
+					]
+					+ SHorizontalBox::Slot()
+					.AutoWidth()
+					.HAlign(HAlign_Center)
+					.VAlign(VAlign_Center)
+					.Padding(2.0f, 1.0f)
+					[
+						SNew(SButton)
+						.ButtonStyle(FAppStyle::Get(), "HoverHintOnly")
+						.Text(LOCTEXT("AddLODLevelCategories_MaterialArrayOpAdd", "Add Material Slot"))
+						.ToolTipText(LOCTEXT("AddLODLevelCategories_MaterialArrayOpAdd_Tooltip", "Add Material Slot at the end of the Material slot array. Those Material slots can be used to override a LODs section, (not the base LOD)"))
+						.ContentPadding(4.0f)
+						.ForegroundColor(FSlateColor::UseForeground())
+						.OnClicked(this, &FGroomMaterialDetails::AddMaterialSlot)
+						.IsEnabled(true)
+						.IsFocusable(false)
+						[
+							SNew(SImage)
+							.Image(FAppStyle::GetBrush("Icons.PlusCircle"))
+							.ColorAndOpacity(FSlateColor::UseForeground())
+						]
+					]
+				]
 			];
 		{
 			FMaterialListDelegates MaterialListDelegates;
@@ -168,6 +169,10 @@ void FGroomMaterialDetails::AddMaterials(IDetailLayoutBuilder& DetailLayout)
 			MaterialListDelegates.OnGenerateCustomMaterialWidgets.BindSP(this, &FGroomMaterialDetails::OnGenerateCustomMaterialWidgetsForMaterialArray, 0);
 			MaterialListDelegates.OnMaterialListDirty.BindSP(this, &FGroomMaterialDetails::OnMaterialListDirty);
 
+			MaterialListDelegates.OnCopyMaterialItem.BindSP(this, &FGroomMaterialDetails::OnCopyMaterialItem);
+			MaterialListDelegates.OnCanCopyMaterialItem.BindSP(this, &FGroomMaterialDetails::OnCanCopyMaterialItem);
+			MaterialListDelegates.OnPasteMaterialItem.BindSP(this, &FGroomMaterialDetails::OnPasteMaterialItem);
+
 			//Pass an empty material list owner (owner can be use by the asset picker filter. In this case we do not need it)
 			TArray<FAssetData> MaterialListOwner;
 			MaterialListOwner.Add(GroomAsset);
@@ -176,6 +181,84 @@ void FGroomMaterialDetails::AddMaterials(IDetailLayoutBuilder& DetailLayout)
 	}
 }
 
+void FGroomMaterialDetails::OnCopyMaterialItem(int32 CurrentSlot)
+{
+	TSharedRef<FJsonObject> RootJsonObject = MakeShareable(new FJsonObject());
+
+	if (GroomAsset->GetHairGroupsMaterials().IsValidIndex(CurrentSlot))
+	{
+		const FHairGroupsMaterial& Material = GroomAsset->GetHairGroupsMaterials()[CurrentSlot];
+
+		FStaticMaterial TmpMaterial;
+		TmpMaterial.MaterialInterface = Material.Material;
+		TmpMaterial.MaterialSlotName  = Material.SlotName;
+		FJsonObjectConverter::UStructToJsonObject(FStaticMaterial::StaticStruct(), &TmpMaterial, RootJsonObject, 0, 0);
+	}
+
+	typedef TJsonWriter<TCHAR, TPrettyJsonPrintPolicy<TCHAR>> FStringWriter;
+	typedef TJsonWriterFactory<TCHAR, TPrettyJsonPrintPolicy<TCHAR>> FStringWriterFactory;
+
+	FString CopyStr;
+	TSharedRef<FStringWriter> Writer = FStringWriterFactory::Create(&CopyStr);
+	FJsonSerializer::Serialize(RootJsonObject, Writer);
+
+	if (!CopyStr.IsEmpty())
+	{
+		FPlatformApplicationMisc::ClipboardCopy(*CopyStr);
+	}
+}
+
+bool FGroomMaterialDetails::OnCanCopyMaterialItem(int32 CurrentSlot) const
+{
+	return GroomAsset->GetHairGroupsMaterials().IsValidIndex(CurrentSlot);
+}
+
+void FGroomMaterialDetails::OnPasteMaterialItem(int32 CurrentSlot)
+{
+	FString PastedText;
+	FPlatformApplicationMisc::ClipboardPaste(PastedText);
+
+	TSharedPtr<FJsonObject> RootJsonObject;
+	TSharedRef<TJsonReader<>> Reader = TJsonReaderFactory<>::Create(PastedText);
+	FJsonSerializer::Deserialize(Reader, RootJsonObject);
+
+	if (RootJsonObject.IsValid())
+	{
+		FProperty* Property = UGroomAsset::StaticClass()->FindPropertyByName(UGroomAsset::GetHairGroupsMaterialsMemberName());
+		check(Property != nullptr);
+
+		GroomAsset->PreEditChange(Property);
+
+		FScopedTransaction Transaction(LOCTEXT("GroomAssetChangedPasteMaterialItem", "GroomAsset editor: Pasted material item"));
+		GroomAsset->Modify();
+
+		if (GroomAsset->GetHairGroupsMaterials().IsValidIndex(CurrentSlot))
+		{
+			FStaticMaterial TmpMaterial;
+			FJsonObjectConverter::JsonObjectToUStruct(RootJsonObject.ToSharedRef(), FStaticMaterial::StaticStruct(), &TmpMaterial, 0, 0);
+			GroomAsset->GetHairGroupsMaterials()[CurrentSlot].Material = TmpMaterial.MaterialInterface;
+			GroomAsset->GetHairGroupsMaterials()[CurrentSlot].SlotName= TmpMaterial.MaterialSlotName;
+		}
+
+		CallPostEditChange(Property);
+	}
+}
+
+void FGroomMaterialDetails::CallPostEditChange(FProperty* PropertyChanged/*=nullptr*/)
+{
+	if (PropertyChanged)
+	{
+		FPropertyChangedEvent PropertyUpdateStruct(PropertyChanged);
+		GroomAsset->PostEditChangeProperty(PropertyUpdateStruct);
+	}
+	else
+	{
+		GroomAsset->Modify();
+		GroomAsset->PostEditChange();
+	}
+	GroomDetailLayout->ForceRefreshDetails();
+}
+
 void FGroomMaterialDetails::ApplyChanges()
 {
 	GroomDetailLayout->ForceRefreshDetails();
diff --git a/Engine/Plugins/Runtime/HairStrands/Source/HairStrandsEditor/Private/GroomMaterialDetails.h b/Engine/Plugins/Runtime/HairStrands/Source/HairStrandsEditor/Private/GroomMaterialDetails.h
--- a/Engine/Plugins/Runtime/HairStrands/Source/HairStrandsEditor/Private/GroomMaterialDetails.h
+++ b/Engine/Plugins/Runtime/HairStrands/Source/HairStrandsEditor/Private/GroomMaterialDetails.h
@@ -159,10 +159,15 @@ private:
 	bool OnCanCopyMaterialList() const;
 	void OnPasteMaterialList();
 
+	void OnCopyMaterialItem(int32 CurrentSlot);
+	bool OnCanCopyMaterialItem(int32 CurrentSlot) const;
+	void OnPasteMaterialItem(int32 CurrentSlot);
 public:
 	void ApplyChanges();
 
 private:
+	void CallPostEditChange(FProperty* PropertyChanged=nullptr);
+
 	IDetailLayoutBuilder* GroomDetailLayout;
 	/*
 	 * This prevent showing the delete material slot warning dialog more then once per editor session
