Add copy pasting to groom asset panel data.

diff --git a/Engine/Plugins/Runtime/HairStrands/Source/HairStrandsCore/Private/GroomAsset.cpp b/Engine/Plugins/Runtime/HairStrands/Source/HairStrandsCore/Private/GroomAsset.cpp
--- a/Engine/Plugins/Runtime/HairStrands/Source/HairStrandsCore/Private/GroomAsset.cpp
+++ b/Engine/Plugins/Runtime/HairStrands/Source/HairStrandsCore/Private/GroomAsset.cpp
@@ -1575,12 +1575,21 @@ static bool IsRiggingSettings(const FName PropertyName)
 
 static bool IsCardsAttributes(const FName PropertyName)
 {
-	return PropertyName == UGroomAsset::GetHairGroupsCardsMemberName();
+	return PropertyName == UGroomAsset::GetHairGroupsCardsMemberName()
+	|| PropertyName == GET_MEMBER_NAME_CHECKED(FHairGroupsCardsSourceDescription, GuideType)
+	|| PropertyName == GET_MEMBER_NAME_CHECKED(FHairGroupsCardsSourceDescription, ImportedMesh)
+//	|| PropertyName == GET_MEMBER_NAME_CHECKED(FHairGroupsCardsSourceDescription, Textures)
+	|| PropertyName == GET_MEMBER_NAME_CHECKED(FHairGroupsCardsSourceDescription, GroupIndex)
+	|| PropertyName == GET_MEMBER_NAME_CHECKED(FHairGroupsCardsSourceDescription, LODIndex);
 }
 
 static bool IsMeshesAttributes(const FName PropertyName)
 {
-	return PropertyName == UGroomAsset::GetHairGroupsMeshesMemberName();
+	return PropertyName == UGroomAsset::GetHairGroupsMeshesMemberName()
+	|| PropertyName == GET_MEMBER_NAME_CHECKED(FHairGroupsMeshesSourceDescription, ImportedMesh)
+//	|| PropertyName == GET_MEMBER_NAME_CHECKED(FHairGroupsMeshesSourceDescription, Textures)
+	|| PropertyName == GET_MEMBER_NAME_CHECKED(FHairGroupsMeshesSourceDescription, GroupIndex)
+	|| PropertyName == GET_MEMBER_NAME_CHECKED(FHairGroupsMeshesSourceDescription, LODIndex);
 }
 
 bool UGroomAsset::Modify(bool bAlwaysMarkDirty)
diff --git a/Engine/Plugins/Runtime/HairStrands/Source/HairStrandsCore/Private/GroomComponent.cpp b/Engine/Plugins/Runtime/HairStrands/Source/HairStrandsCore/Private/GroomComponent.cpp
--- a/Engine/Plugins/Runtime/HairStrands/Source/HairStrandsCore/Private/GroomComponent.cpp
+++ b/Engine/Plugins/Runtime/HairStrands/Source/HairStrandsCore/Private/GroomComponent.cpp
@@ -2811,6 +2811,51 @@ void UGroomComponent::InitResources(bool bIsBindingReloading)
 			{
 				LocalBindingAsset = nullptr;
 			}
+
+			if (LocalBindingAsset)
+			{
+				// Ensure the groom binding asset's resources match the groom asset's resources.
+				// They can mismatch when the groom asset has been edited and the groom binding hasn't been updated/finished to be built yet.
+				const uint32 NumGroup = LocalBindingAsset->GetHairGroupResources().Num();
+				check(NumGroup == GroomAsset->GetHairGroupsResources().Num());
+				for (uint32 GroupIt = 0; GroupIt < NumGroup; ++GroupIt)
+				{
+					const UGroomBindingAsset::FHairGroupResource& BindingResource = LocalBindingAsset->GetHairGroupResources()[GroupIt];
+					const FHairGroupResources& GroomResources = GroomAsset->GetHairGroupsResources()[GroupIt];
+
+					bool bValid = true;
+					// Check guides
+					if (GroomResources.Guides.IsValid())
+					{
+						bValid = bValid && BindingResource.SimRootResources && BindingResource.SimRootResources->GetLODCount() > 0;
+					}
+					// Check strands
+					if (bValid && GroomResources.Strands.IsValid())
+					{
+						bValid = bValid && BindingResource.RenRootResources && BindingResource.RenRootResources->GetLODCount() > 0;
+					}
+					// Check cards
+					if (bValid && GroomResources.Cards.LODs.Num() > 0)
+					{
+						bValid = bValid && GroomResources.Cards.LODs.Num() == BindingResource.CardsRootResources.Num();
+						if (bValid)
+						{
+							for (uint32 CardIt = 0, CardCount = GroomResources.Cards.LODs.Num(); CardIt < CardCount; ++CardIt)
+							{
+								if (GroomResources.Cards.LODs[CardIt].IsValid())
+								{
+									bValid = bValid && BindingResource.CardsRootResources[CardIt] && BindingResource.CardsRootResources[CardIt]->GetLODCount() > 0;
+								}
+							}
+						}
+					}
+
+					if (!bValid)
+					{
+						LocalBindingAsset = nullptr;
+					}
+				}
+			}
 		}
 	}
 
diff --git a/Engine/Plugins/Runtime/HairStrands/Source/HairStrandsEditor/Private/GroomAssetDetails.cpp b/Engine/Plugins/Runtime/HairStrands/Source/HairStrandsEditor/Private/GroomAssetDetails.cpp
--- a/Engine/Plugins/Runtime/HairStrands/Source/HairStrandsEditor/Private/GroomAssetDetails.cpp
+++ b/Engine/Plugins/Runtime/HairStrands/Source/HairStrandsEditor/Private/GroomAssetDetails.cpp
@@ -61,6 +61,7 @@
 #include "PropertyCustomizationHelpers.h"
 #include "GroomCustomAssetEditorToolkit.h"
 #include "IPropertyUtilities.h"
+#include "JsonObjectConverter.h"
 
 #include "Styling/AppStyle.h"
 #include "GroomEditorStyle.h"
@@ -378,6 +379,202 @@ TSharedRef<SUniformGridPanel> MakeHairCardsInfoGrid(const FSlateFontInfo& Detail
 	return Grid;
 }
 
+//////////////////////////////////////////////////////////////////////////////////////////////////////////////////
+// Copy/Paste LOD
+
+static bool IsValidLODGroup(UGroomAsset* InAsset, int32 InGroupIndex, int32 InLODIndex)
+{
+	return InAsset->GetHairGroupsLOD().IsValidIndex(InGroupIndex) && InAsset->GetHairGroupsLOD()[InGroupIndex].LODs.IsValidIndex(InLODIndex);
+}
+
+static void CopyLODToJson(UGroomAsset* InAsset, int32 InGroupIndex, int32 InLODIndex, TSharedRef<FJsonObject> OutRootJsonObject)
+{
+	if (IsValidLODGroup(InAsset, InGroupIndex, InLODIndex))
+	{
+		FJsonObjectConverter::UStructToJsonObject(FHairLODSettings::StaticStruct(), &InAsset->GetHairGroupsLOD()[InGroupIndex].LODs[InLODIndex], OutRootJsonObject, 0, 0);
+	}
+}
+
+static void PasteLODFromJson(UGroomAsset* InAsset, int32 InGroupIndex, int32 InLODIndex, TSharedRef<FJsonObject> InRootJsonObject)
+{
+	if (IsValidLODGroup(InAsset, InGroupIndex, InLODIndex))
+	{
+		FJsonObjectConverter::JsonObjectToUStruct(InRootJsonObject, FHairLODSettings::StaticStruct(), &InAsset->GetHairGroupsLOD()[InGroupIndex].LODs[InLODIndex], 0, 0);
+	}
+}
+
+void FGroomRenderingDetails::OnCopyLODItem(int32 InGroupIndex, int32 InLODIndex)
+{
+	TSharedRef<FJsonObject> RootJsonObject = MakeShareable(new FJsonObject());
+
+	CopyLODToJson(GroomAsset, InGroupIndex, InLODIndex, RootJsonObject);
+
+	typedef TJsonWriter<TCHAR, TPrettyJsonPrintPolicy<TCHAR>> FStringWriter;
+	typedef TJsonWriterFactory<TCHAR, TPrettyJsonPrintPolicy<TCHAR>> FStringWriterFactory;
+
+	FString CopyStr;
+	TSharedRef<FStringWriter> Writer = FStringWriterFactory::Create(&CopyStr);
+	FJsonSerializer::Serialize(RootJsonObject, Writer);
+
+	if (!CopyStr.IsEmpty())
+	{
+		FPlatformApplicationMisc::ClipboardCopy(*CopyStr);
+	}
+}
+
+bool FGroomRenderingDetails::OnCanCopyLODItem(int32 InGroupIndex, int32 InLODIndex) const
+{
+	return IsValidLODGroup(GroomAsset, InGroupIndex, InLODIndex);
+}
+
+void FGroomRenderingDetails::OnPasteLODItem(int32 InGroupIndex, int32 InLODIndex)
+{
+	FString PastedText;
+	FPlatformApplicationMisc::ClipboardPaste(PastedText);
+
+	TSharedPtr<FJsonObject> RootJsonObject;
+	TSharedRef<TJsonReader<>> Reader = TJsonReaderFactory<>::Create(PastedText);
+	FJsonSerializer::Deserialize(Reader, RootJsonObject);
+
+	if (RootJsonObject.IsValid())
+	{
+		FProperty* Property = FHairGroupsLOD::StaticStruct()->FindPropertyByName(GET_MEMBER_NAME_CHECKED(FHairGroupsLOD, LODs));
+		if (Property)
+		{
+			GroomAsset->PreEditChange(Property);
+
+			FScopedTransaction Transaction(LOCTEXT("GroomAssetChangedPasteLOD", "GroomAsset editor: Pasted LOD"));
+			GroomAsset->Modify();
+			PasteLODFromJson(GroomAsset, InGroupIndex, InLODIndex, RootJsonObject.ToSharedRef());
+
+			CallPostEditChange(Property);
+		}
+	}
+}
+
+//////////////////////////////////////////////////////////////////////////////////////////////////////////////////
+// Copy/Paste Group
+
+static bool IsValidGroup(UGroomAsset* InAsset, int32 InIndex, EMaterialPanelType InType)
+{
+	switch (InType)
+	{
+		case EMaterialPanelType::Strands: 		return InAsset->GetHairGroupsRendering().IsValidIndex(InIndex);
+		case EMaterialPanelType::Cards:			return InAsset->GetHairGroupsCards().IsValidIndex(InIndex);
+		case EMaterialPanelType::Meshes:		return InAsset->GetHairGroupsMeshes().IsValidIndex(InIndex);
+		case EMaterialPanelType::Physics:		return InAsset->GetHairGroupsPhysics().IsValidIndex(InIndex);
+		case EMaterialPanelType::Interpolation:	return InAsset->GetHairGroupsInterpolation().IsValidIndex(InIndex);
+		case EMaterialPanelType::LODs:			return InAsset->GetHairGroupsLOD().IsValidIndex(InIndex);
+	}
+	return false;
+	
+}
+static void CopyGroupToJson(UGroomAsset* InAsset, int32 InIndex, EMaterialPanelType InType, TSharedRef<FJsonObject> OutRootJsonObject)
+{
+	if (!IsValidGroup(InAsset, InIndex, InType)) return;
+
+	switch (InType)
+	{
+		case EMaterialPanelType::Strands:		FJsonObjectConverter::UStructToJsonObject(FHairGroupsRendering::StaticStruct(), &InAsset->GetHairGroupsRendering()[InIndex], OutRootJsonObject, 0, 0); break;
+		case EMaterialPanelType::Cards:			FJsonObjectConverter::UStructToJsonObject(FHairGroupsCardsSourceDescription::StaticStruct(), &InAsset->GetHairGroupsCards()[InIndex], OutRootJsonObject, 0, 0); break;
+		case EMaterialPanelType::Meshes:		FJsonObjectConverter::UStructToJsonObject(FHairGroupsMeshesSourceDescription::StaticStruct(), &InAsset->GetHairGroupsMeshes()[InIndex], OutRootJsonObject, 0, 0); break;
+		case EMaterialPanelType::Physics:		FJsonObjectConverter::UStructToJsonObject(FHairGroupsPhysics::StaticStruct(), &InAsset->GetHairGroupsPhysics()[InIndex], OutRootJsonObject, 0, 0); break;
+		case EMaterialPanelType::Interpolation:	FJsonObjectConverter::UStructToJsonObject(FHairGroupsInterpolation::StaticStruct(), &InAsset->GetHairGroupsInterpolation()[InIndex], OutRootJsonObject, 0, 0); break;
+		case EMaterialPanelType::LODs:			FJsonObjectConverter::UStructToJsonObject(FHairGroupsLOD::StaticStruct(), &InAsset->GetHairGroupsLOD()[InIndex], OutRootJsonObject, 0, 0); break;
+	};
+}
+
+static void PasteGroupFromJson(UGroomAsset* InAsset, int32 InIndex, EMaterialPanelType InType, TSharedRef<FJsonObject> InRootJsonObject)
+{
+	if (!IsValidGroup(InAsset, InIndex, InType)) return;
+
+	switch (InType)
+	{
+		case EMaterialPanelType::Strands: 		FJsonObjectConverter::JsonObjectToUStruct(InRootJsonObject, FHairGroupsRendering::StaticStruct(), &InAsset->GetHairGroupsRendering()[InIndex], 0, 0); break;
+		case EMaterialPanelType::Cards: 		FJsonObjectConverter::JsonObjectToUStruct(InRootJsonObject, FHairGroupsCardsSourceDescription::StaticStruct(), &InAsset->GetHairGroupsCards()[InIndex], 0, 0); break;
+		case EMaterialPanelType::Meshes:  		FJsonObjectConverter::JsonObjectToUStruct(InRootJsonObject, FHairGroupsMeshesSourceDescription::StaticStruct(), &InAsset->GetHairGroupsMeshes()[InIndex], 0, 0); break;
+		case EMaterialPanelType::Physics:  		FJsonObjectConverter::JsonObjectToUStruct(InRootJsonObject, FHairGroupsPhysics::StaticStruct(), &InAsset->GetHairGroupsPhysics()[InIndex], 0, 0); break;
+		case EMaterialPanelType::Interpolation: FJsonObjectConverter::JsonObjectToUStruct(InRootJsonObject, FHairGroupsInterpolation::StaticStruct(), &InAsset->GetHairGroupsInterpolation()[InIndex], 0, 0); break;
+		case EMaterialPanelType::LODs: 			FJsonObjectConverter::JsonObjectToUStruct(InRootJsonObject, FHairGroupsLOD::StaticStruct(), &InAsset->GetHairGroupsLOD()[InIndex], 0, 0); break;
+	};
+}
+
+static FProperty* GetGroupPropertyFromType(EMaterialPanelType InType)
+{
+	switch (InType)
+	{
+		case EMaterialPanelType::Strands: 		return UGroomAsset::StaticClass()->FindPropertyByName(UGroomAsset::GetHairGroupsRenderingMemberName());
+		case EMaterialPanelType::Cards:			return UGroomAsset::StaticClass()->FindPropertyByName(UGroomAsset::GetHairGroupsCardsMemberName());
+		case EMaterialPanelType::Meshes:		return UGroomAsset::StaticClass()->FindPropertyByName(UGroomAsset::GetHairGroupsMeshesMemberName());
+		case EMaterialPanelType::Physics:		return UGroomAsset::StaticClass()->FindPropertyByName(UGroomAsset::GetHairGroupsPhysicsMemberName());
+		case EMaterialPanelType::Interpolation:	return UGroomAsset::StaticClass()->FindPropertyByName(UGroomAsset::GetHairGroupsInterpolationMemberName());
+		case EMaterialPanelType::LODs:			return UGroomAsset::StaticClass()->FindPropertyByName(UGroomAsset::GetHairGroupsLODMemberName());
+	};
+	return nullptr;
+}
+
+void FGroomRenderingDetails::OnCopyGroupItem(int32 CurrentSlot, EMaterialPanelType InType)
+{
+	TSharedRef<FJsonObject> RootJsonObject = MakeShareable(new FJsonObject());
+
+	CopyGroupToJson(GroomAsset, CurrentSlot, PanelType, RootJsonObject);
+
+	typedef TJsonWriter<TCHAR, TPrettyJsonPrintPolicy<TCHAR>> FStringWriter;
+	typedef TJsonWriterFactory<TCHAR, TPrettyJsonPrintPolicy<TCHAR>> FStringWriterFactory;
+
+	FString CopyStr;
+	TSharedRef<FStringWriter> Writer = FStringWriterFactory::Create(&CopyStr);
+	FJsonSerializer::Serialize(RootJsonObject, Writer);
+
+	if (!CopyStr.IsEmpty())
+	{
+		FPlatformApplicationMisc::ClipboardCopy(*CopyStr);
+	}
+}
+
+bool FGroomRenderingDetails::OnCanCopyGroupItem(int32 CurrentSlot, EMaterialPanelType InType) const
+{
+	return IsValidGroup(GroomAsset, CurrentSlot, InType);
+}
+
+void FGroomRenderingDetails::OnPasteGroupItem(int32 CurrentSlot, EMaterialPanelType InType)
+{
+	FString PastedText;
+	FPlatformApplicationMisc::ClipboardPaste(PastedText);
+
+	TSharedPtr<FJsonObject> RootJsonObject;
+	TSharedRef<TJsonReader<>> Reader = TJsonReaderFactory<>::Create(PastedText);
+	FJsonSerializer::Deserialize(Reader, RootJsonObject);
+
+	if (RootJsonObject.IsValid())
+	{
+		if (FProperty* Property = GetGroupPropertyFromType(PanelType))
+		{
+			GroomAsset->PreEditChange(Property);
+
+			FScopedTransaction Transaction(LOCTEXT("GroomAssetChangedPasteGroup", "GroomAsset editor: Pasted group"));
+			GroomAsset->Modify();
+			PasteGroupFromJson(GroomAsset, CurrentSlot, InType, RootJsonObject.ToSharedRef());
+			CallPostEditChange(Property);
+		}
+	}
+}
+
+void FGroomRenderingDetails::CallPostEditChange(FProperty* PropertyChanged/*=nullptr*/)
+{
+	if (PropertyChanged)
+	{
+		FPropertyChangedEvent PropertyUpdateStruct(PropertyChanged);
+		GroomAsset->PostEditChangeProperty(PropertyUpdateStruct);
+	}
+	else
+	{
+		GroomAsset->Modify();
+		GroomAsset->PostEditChange();
+	}
+	GroomDetailLayout->ForceRefreshDetails();
+}
+
 //////////////////////////////////////////////////////////////////////////////////////////////////////////////////
 // Details view
 
@@ -1359,6 +1556,8 @@ void FGroomRenderingDetails::OnGenerateElementForLODs(TSharedRef<IPropertyHandle
 
 	// LOD Bar with add button
 	ChildrenBuilder.AddCustomRow(LOCTEXT("HairInfo_Separator", "Separator"))
+	.CopyAction(FUIAction(FExecuteAction::CreateSP(this, &FGroomRenderingDetails::OnCopyLODItem, GroupIndex, LODIndex), FCanExecuteAction::CreateSP(this, &FGroomRenderingDetails::OnCanCopyLODItem, GroupIndex, LODIndex)))
+	.PasteAction(FUIAction(FExecuteAction::CreateSP(this, &FGroomRenderingDetails::OnPasteLODItem, GroupIndex, LODIndex)))
 	.WholeRowContent()
 	.VAlign(VAlign_Fill)
 	.HAlign(HAlign_Fill)
@@ -1545,6 +1744,8 @@ void FGroomRenderingDetails::OnGenerateElementForHairGroup(TSharedRef<IPropertyH
 	}
 
 	ChildrenBuilder.AddCustomRow(LOCTEXT("HairInfo_Separator", "Separator"))
+	.CopyAction(FUIAction(FExecuteAction::CreateSP(this, &FGroomRenderingDetails::OnCopyGroupItem, GroupIndex, PanelType), FCanExecuteAction::CreateSP(this, &FGroomRenderingDetails::OnCanCopyGroupItem, GroupIndex, PanelType)))
+	.PasteAction(FUIAction(FExecuteAction::CreateSP(this, &FGroomRenderingDetails::OnPasteGroupItem, GroupIndex, PanelType)))
 	.WholeRowContent()
 	.VAlign(VAlign_Fill)
 	.HAlign(HAlign_Fill)
diff --git a/Engine/Plugins/Runtime/HairStrands/Source/HairStrandsEditor/Private/GroomAssetDetails.h b/Engine/Plugins/Runtime/HairStrands/Source/HairStrandsEditor/Private/GroomAssetDetails.h
--- a/Engine/Plugins/Runtime/HairStrands/Source/HairStrandsEditor/Private/GroomAssetDetails.h
+++ b/Engine/Plugins/Runtime/HairStrands/Source/HairStrandsEditor/Private/GroomAssetDetails.h
@@ -105,4 +105,16 @@ private:
 	IDetailLayoutBuilder* GroomDetailLayout = nullptr;
 	EMaterialPanelType PanelType = EMaterialPanelType::Strands;
 	bool bDeleteWarningConsumed; // This prevent showing the delete material slot warning dialog more then once per editor session
+
+	// Copy/Paste group
+	void OnCopyGroupItem(int32 CurrentSlot, EMaterialPanelType InType);
+	bool OnCanCopyGroupItem(int32 CurrentSlot, EMaterialPanelType InType) const;
+	void OnPasteGroupItem(int32 CurrentSlot, EMaterialPanelType InType);
+
+	// Copy/Paste LOD
+	void OnCopyLODItem(int32 CurrentSlot, int32 InLODIndex);
+	bool OnCanCopyLODItem(int32 CurrentSlot, int32 InLODIndex) const;
+	void OnPasteLODItem(int32 CurrentSlot, int32 InLODIndex);
+
+	void CallPostEditChange(FProperty* PropertyChanged=nullptr);
 };
\ No newline at end of file
