Fix crash on startup in CitySample: exit loop on invalid LocalBindingAsset

#rnx

diff --git a/Engine/Plugins/Runtime/HairStrands/Source/HairStrandsCore/Private/GroomComponent.cpp b/Engine/Plugins/Runtime/HairStrands/Source/HairStrandsCore/Private/GroomComponent.cpp
--- a/Engine/Plugins/Runtime/HairStrands/Source/HairStrandsCore/Private/GroomComponent.cpp
+++ b/Engine/Plugins/Runtime/HairStrands/Source/HairStrandsCore/Private/GroomComponent.cpp
@@ -2853,6 +2853,7 @@ void UGroomComponent::InitResources(bool bIsBindingReloading)
 					if (!bValid)
 					{
 						LocalBindingAsset = nullptr;
+						break;
 					}
 				}
 			}
