Convert the Groom Editor Viewport to the new viewport toolbar.
- Consolidate invocation of groom menus between the goom editor and the level viewport.
- Tweak SEditorViewport::GetViewportClient() to return a const reference instead of a new SharedPtr.
- Fix a bug in ToolMenus where top level toolbar submenus that defined a submenu using FNewMenuDelegate did not work correctly.

diff --git a/Engine/Plugins/Runtime/HairStrands/Source/HairStrandsEditor/Private/Widgets/SGroomEditorViewport.cpp b/Engine/Plugins/Runtime/HairStrands/Source/HairStrandsEditor/Private/Widgets/SGroomEditorViewport.cpp
--- a/Engine/Plugins/Runtime/HairStrands/Source/HairStrandsEditor/Private/Widgets/SGroomEditorViewport.cpp
+++ b/Engine/Plugins/Runtime/HairStrands/Source/HairStrandsEditor/Private/Widgets/SGroomEditorViewport.cpp
@@ -13,6 +13,7 @@
 #include "Engine/Canvas.h"
 #include "Engine/Font.h"
 #include "CanvasItem.h"
+#include "ToolMenus.h"
 #include "DrawDebugHelpers.h"
 #include "AdvancedPreviewScene.h"
 #include "Widgets/Docking/SDockTab.h"
@@ -24,7 +25,11 @@
 #include "GroomVisualizationMenuCommands.h"
 #include "EditorViewportCommands.h"
 #include "AssetViewerSettings.h"
+#include "PreviewProfileController.h"
+#include "SLevelViewport.h"
 #include "Editor/EditorPerProjectUserSettings.h"
+#include "Tests/ToolMenusTestUtilities.h"
+#include "ViewportToolbar/UnrealEdViewportToolbar.h"
 
 #define LOCTEXT_NAMESPACE "SGroomEditorViewport"
 
@@ -356,6 +361,141 @@ TSharedPtr<SWidget> SGroomEditorViewport::MakeViewportToolbar()
 	return SNew(SGroomEditorViewportToolbar, SharedThis(this));
 }
 
+TSharedPtr<SWidget> SGroomEditorViewport::BuildViewportToolbar()
+{
+	const FName ViewportToolbarName = "GroomEditor.ViewportToolbar";
+	
+	if (!UToolMenus::Get()->IsMenuRegistered(ViewportToolbarName))
+	{
+		UToolMenu* ToolMenu = UToolMenus::Get()->RegisterMenu(ViewportToolbarName, NAME_None, EMultiBoxType::SlimHorizontalToolBar);
+		ToolMenu->StyleName = "ViewportToolbar";
+		
+		ToolMenu->FindOrAddSection("Left");
+		
+		FToolMenuSection& RightSection = ToolMenu->FindOrAddSection("Right");
+		RightSection.Alignment = EToolMenuSectionAlign::Last;
+		
+		
+		RightSection.AddEntry(UE::UnrealEd::CreateCameraSubmenu());
+		
+		{
+			// View Modes
+			RightSection.AddEntry(UE::UnrealEd::CreateViewModesSubmenu());
+			UToolMenu* ViewMenu = UToolMenus::Get()->ExtendMenu(UToolMenus::JoinMenuPaths(ViewportToolbarName, "ViewModes"));
+			FToolMenuSection& GroomViewSection = ViewMenu->FindOrAddSection("Groom", LOCTEXT("GroomViewSectionName", "Groom"));
+			
+			GroomViewSection.AddDynamicEntry("GroomDynamicViewModes", FNewToolMenuSectionDelegate::CreateLambda([](FToolMenuSection& Section)
+			{
+				const UUnrealEdViewportToolbarContext* ToolbarContext = Section.Context.FindContext<UUnrealEdViewportToolbarContext>();
+				if (!ToolbarContext || !ToolbarContext->Viewport.IsValid())
+				{
+					return;
+				}
+				
+				FToolMenuEntry Entry = FGroomVisualizationMenuCommands::BuildVisualizationSubMenuItemForGroomEditor(ToolbarContext->Viewport); 
+				
+				Entry.SetShowInToolbarTopLevel(true);
+				
+				// Customize the top-level menu to include the option to return to lit
+				Entry.ToolBarData.ComboButtonContextMenuGenerator.NewMenuLegacy.BindLambda([](FMenuBuilder& MenuBuilder)
+				{
+					MenuBuilder.AddMenuEntry(FEditorViewportCommands::Get().LitMode, NAME_None, UViewModeUtils::GetViewModeDisplayName(VMI_Lit));
+					MenuBuilder.AddSeparator();
+					FGroomVisualizationMenuCommands::BuildVisualisationSubMenuForGroomEditor(MenuBuilder);
+				});
+				
+				Section.AddEntry(Entry);
+			}));
+		}
+		
+		{
+			// LOD Menu
+			FToolMenuEntry& Entry = RightSection.AddDynamicEntry("DynamicLOD", FNewToolMenuSectionDelegate::CreateLambda([](FToolMenuSection& Section)
+			{
+				UUnrealEdViewportToolbarContext* Context = Section.FindContext<UUnrealEdViewportToolbarContext>();
+				if (!Context)
+				{
+					return;
+				}
+				
+				TSharedPtr<SEditorViewport> EditorViewport = Context->Viewport.Pin();
+				if (!EditorViewport)
+				{
+					return;
+				}
+				
+				TSharedPtr<SGroomEditorViewport> GroomViewport = StaticCastSharedPtr<SGroomEditorViewport>(EditorViewport);
+				
+				Section.AddSubMenu(
+					"LOD",
+					TAttribute<FText>::CreateLambda([WeakViewport = GroomViewport.ToWeakPtr()]
+					{
+						if (TSharedPtr<SGroomEditorViewport> Viewport = WeakViewport.Pin())
+						{
+							const int32 LODSelectionType = Viewport->GetLODSelection();
+							if (LODSelectionType >= 0)
+							{
+								return FText::Format(LOCTEXT("LodMenu_LabelFormat", "LOD {0}"), LODSelectionType);
+							}
+						}
+						return LOCTEXT("LODMenu_AutoLabel", "LOD Auto");
+					}),
+					LOCTEXT("LodMenu_Tooltip", "Set the Level of Detail of the viewport."),
+					FNewMenuDelegate::CreateLambda([WeakViewport = GroomViewport.ToWeakPtr()](FMenuBuilder& Menu)
+					{
+						TSharedPtr<SGroomEditorViewport> Viewport = WeakViewport.Pin();
+						if (!Viewport)
+						{
+							return;
+						}
+						
+						Menu.BeginSection("GroomViewportPreviewLODs", LOCTEXT("ShowLOD_PreviewLabel", "Preview LODs"));
+						{
+							Menu.AddMenuEntry(FGroomViewportLODCommands::Get().LODAuto);
+							Menu.AddMenuEntry(FGroomViewportLODCommands::Get().LOD0);
+
+							const int32 LODCount = Viewport->GetLODModelCount();
+							for (int32 LODId = 1; LODId < LODCount; ++LODId)
+							{
+								FString TitleLabel = FString::Printf(TEXT(" LOD %d"), LODId);
+
+								FUIAction Action(
+									FExecuteAction::CreateSP(Viewport.Get(), &SGroomEditorViewport::OnSetLODModel, LODId),
+									FCanExecuteAction(),
+									FIsActionChecked::CreateSP(Viewport.Get(), &SGroomEditorViewport::IsLODModelSelected, LODId)
+								);
+
+								Menu.AddMenuEntry(FText::FromString(TitleLabel), FText::GetEmpty(), FSlateIcon(), Action, NAME_None, EUserInterfaceActionType::RadioButton);
+							}
+						}
+						Menu.EndSection();
+					})
+				);
+			}));
+			
+			Entry.InsertPosition.Position = EToolMenuInsertType::Last;
+		}
+		
+		RightSection.AddEntry(UE::UnrealEd::CreatePerformanceAndScalabilitySubmenu());
+		RightSection.AddEntry(UE::UnrealEd::CreateAssetViewerProfileSubmenu());
+	}
+	
+	FToolMenuContext ViewportToolbarContext;
+	ViewportToolbarContext.AppendCommandList(GetCommandList());
+	ViewportToolbarContext.AddExtender(GetExtenders());
+	
+	UUnrealEdViewportToolbarContext* const ContextObject = UE::UnrealEd::CreateViewportToolbarDefaultContext(SharedThis(this));
+	
+	ViewportToolbarContext.AddObject(ContextObject);
+	
+	return UToolMenus::Get()->GenerateWidget(ViewportToolbarName, ViewportToolbarContext);
+}
+
+TSharedPtr<IPreviewProfileController> SGroomEditorViewport::CreatePreviewProfileController()
+{
+	return MakeShared<FPreviewProfileController>();
+}
+
 EVisibility SGroomEditorViewport::OnGetViewportContentVisibility() const
 {
 	EVisibility BaseVisibility = SEditorViewport::OnGetViewportContentVisibility();
@@ -366,15 +506,6 @@ EVisibility SGroomEditorViewport::OnGetViewportContentVisibility() const
 	return IsVisible() ? EVisibility::Visible : EVisibility::Collapsed;
 }
 
-void SGroomEditorViewport::PopulateViewportOverlays(TSharedRef<class SOverlay> Overlay)
-{	
-		Overlay->AddSlot()
-		.VAlign(VAlign_Top)
-		[
-			SNew(SGroomEditorViewportToolbar, SharedThis(this))
-		];	
-}
-
 TSharedRef<class SEditorViewport> SGroomEditorViewport::GetViewportWidget()
 {
 	return SharedThis(this);
@@ -393,13 +524,20 @@ void SGroomEditorViewport::OnFloatingButtonClicked()
 
 int32 SGroomEditorViewport::GetLODSelection() const
 {
-	return GroomComponent->GetForcedLOD();
+	if (GroomComponent)
+	{
+		return GroomComponent->GetForcedLOD();
+	}
+	return -1;
 }
 
 int32 SGroomEditorViewport::GetLODModelCount() const
 {
-	
-	return GroomComponent->GetNumLODs();
+	if (GroomComponent)
+	{
+		return GroomComponent->GetNumLODs();
+	}
+	return -1;
 }
 
 #undef LOCTEXT_NAMESPACE
\ No newline at end of file
diff --git a/Engine/Plugins/Runtime/HairStrands/Source/HairStrandsEditor/Public/Widgets/SGroomEditorViewport.h b/Engine/Plugins/Runtime/HairStrands/Source/HairStrandsEditor/Public/Widgets/SGroomEditorViewport.h
--- a/Engine/Plugins/Runtime/HairStrands/Source/HairStrandsEditor/Public/Widgets/SGroomEditorViewport.h
+++ b/Engine/Plugins/Runtime/HairStrands/Source/HairStrandsEditor/Public/Widgets/SGroomEditorViewport.h
@@ -67,10 +67,11 @@ protected:
 	/** SEditorViewport interface */
 	virtual TSharedRef<FEditorViewportClient> MakeEditorViewportClient() override;
 	virtual TSharedPtr<SWidget> MakeViewportToolbar() override;
+	virtual TSharedPtr<SWidget> BuildViewportToolbar() override;
+	virtual TSharedPtr<IPreviewProfileController> CreatePreviewProfileController() override;
 	virtual EVisibility OnGetViewportContentVisibility() const override;
 	virtual void BindCommands() override;
 	virtual void OnFocusViewportToSelection() override;
-	virtual void PopulateViewportOverlays(TSharedRef<class SOverlay> Overlay) override;
 
 private:
 
diff --git a/Engine/Source/Developer/ToolMenus/Private/ToolMenus.cpp b/Engine/Source/Developer/ToolMenus/Private/ToolMenus.cpp
--- a/Engine/Source/Developer/ToolMenus/Private/ToolMenus.cpp
+++ b/Engine/Source/Developer/ToolMenus/Private/ToolMenus.cpp
@@ -1905,60 +1905,51 @@ void UToolMenus::PopulateToolBarBuilderWithEntry(
 			);
 		}
 	}
-	else if (Block.Type == EMultiBlockType::ToolBarComboButton
-			 || (Block.Type == EMultiBlockType::MenuEntry && Block.IsSubMenu()))
+	else if (Block.Type == EMultiBlockType::ToolBarComboButton || (Block.Type == EMultiBlockType::MenuEntry && Block.IsSubMenu()))
 	{
-		FOnGetContent OnGetContent = ConvertWidgetChoice(
-			Block.ToolBarData.ComboButtonContextMenuGenerator, MenuData->Context);
-		if (OnGetContent.IsBound())
+		bool bCouldHaveChildren = false;
+		FOnGetContent OnGetContent = ConvertWidgetChoice(Block.ToolBarData.ComboButtonContextMenuGenerator, MenuData->Context);
+		
+		// Allow non-tool-menu-creating choices to be applied
+		if (!OnGetContent.IsBound() && !Block.SubMenuData.ConstructMenu.NewToolMenu.IsBound())
 		{
-			ToolBarBuilder.AddComboButton(
-				UIAction,
-				OnGetContent,
-				Block.Label,
-				Block.ToolTip,
-				Block.Icon,
-				Block.ToolBarData.bSimpleComboBox,
-				Block.TutorialHighlightName,
-				VisibilityOverride,
-				ToolbarLabelOverride,
-				Block.ToolBarData.PlacementOverride,
-				Block.UserInterfaceActionType,
-				ActualResizeParams
-			);
+			OnGetContent = ConvertWidgetChoice(Block.SubMenuData.ConstructMenu, MenuData->Context);
 		}
-		else
+		
+		if (!OnGetContent.IsBound())
 		{
+			// Handle tool-menu-generating lambdas.
 			// Make sure to keep a strong reference to this submenu so it stays around until it is opened. This is
 			// needed here because we're could be a submenu of a submenu, so not even our parent is added to the Menus
 			// map and therefore our parent could have been gargage collected before this submenu is opened.
-			FOnGetContent Delegate = FOnGetContent::CreateLambda(
+			OnGetContent = FOnGetContent::CreateLambda(
 				[this, StrongMenuData = TStrongObjectPtr<UToolMenu>(MenuData), BlockName = Block.Name]() -> TSharedRef<SWidget>
 				{
 					return this->GenerateToolbarComboButtonMenu(StrongMenuData.Get(), BlockName);
 				}
 			);
-
-			ToolBarBuilder.AddComboButton(
-				UIAction,
-				Delegate,
-				Block.Label,
-				Block.ToolTip,
-				Block.Icon,
-				Block.ToolBarData.bSimpleComboBox,
-				Block.TutorialHighlightName,
-				VisibilityOverride,
-				ToolbarLabelOverride,
-				Block.ToolBarData.PlacementOverride,
-				Block.UserInterfaceActionType,
-				ActualResizeParams
-			);
-
-			// Also add any top-level flagged children to the toolbar.
-			if (!bIsRaisingToTopLevel)
-			{
-				PopulateToolBarBuilderWithTopLevelChildren(ToolBarBuilder, MenuData, Block, !bIsLastBlockOfLastSection);
-			}
+			bCouldHaveChildren = true;
+		}
+		
+		ToolBarBuilder.AddComboButton(
+			UIAction,
+			OnGetContent,
+			Block.Label,
+			Block.ToolTip,
+			Block.Icon,
+			Block.ToolBarData.bSimpleComboBox,
+			Block.TutorialHighlightName,
+			VisibilityOverride,
+			ToolbarLabelOverride,
+			Block.ToolBarData.PlacementOverride,
+			Block.UserInterfaceActionType,
+			ActualResizeParams
+		);
+			
+		// Also add any top-level flagged children to the toolbar.
+		if (bCouldHaveChildren && !bIsRaisingToTopLevel)
+		{
+			PopulateToolBarBuilderWithTopLevelChildren(ToolBarBuilder, MenuData, Block, !bIsLastBlockOfLastSection);
 		}
 	}
 	else if (Block.Type == EMultiBlockType::Separator)
diff --git a/Engine/Source/Editor/CommonMenuExtensions/Private/GroomVisualizationMenuCommands.cpp b/Engine/Source/Editor/CommonMenuExtensions/Private/GroomVisualizationMenuCommands.cpp
--- a/Engine/Source/Editor/CommonMenuExtensions/Private/GroomVisualizationMenuCommands.cpp
+++ b/Engine/Source/Editor/CommonMenuExtensions/Private/GroomVisualizationMenuCommands.cpp
@@ -11,11 +11,13 @@
 #include "Framework/MultiBox/MultiBoxBuilder.h"
 #include "Internationalization/Internationalization.h"
 #include "Internationalization/Text.h"
+#include "ToolMenus.h"
 #include "GroomVisualizationData.h"
 #include "Misc/AssertionMacros.h"
 #include "Styling/AppStyle.h"
 #include "UObject/UnrealNames.h"
 #include "GroomVisualizationData.h"
+#include "SEditorViewport.h"
 
 #define LOCTEXT_NAMESPACE "GroomVisualizationMenuCommands"
 
@@ -144,6 +146,51 @@ void FGroomVisualizationMenuCommands::BuildVisualisationSubMenuForGroomEditor(FM
 	InternalBuildVisualisationSubMenu(Menu, true);
 }
 
+FToolMenuEntry FGroomVisualizationMenuCommands::BuildVisualizationSubMenuItem(const TWeakPtr<SEditorViewport>& Viewport)
+{
+	return InternalBuildVisualizationSubMenuItem(Viewport, false);
+}
+
+FToolMenuEntry FGroomVisualizationMenuCommands::BuildVisualizationSubMenuItemForGroomEditor(const TWeakPtr<SEditorViewport>& Viewport)
+{
+	return InternalBuildVisualizationSubMenuItem(Viewport, true);
+}
+
+FToolMenuEntry FGroomVisualizationMenuCommands::InternalBuildVisualizationSubMenuItem(const TWeakPtr<SEditorViewport>& Viewport, bool bIsGroomEditor)
+{
+	FToolMenuEntry MenuEntry = FToolMenuEntry::InitSubMenu(
+		"VisualizeGroomViewMode",
+		LOCTEXT("VisualizeGroomViewModeDisplayName", "Groom"),
+		LOCTEXT("GroomVisualizationMenu_ToolTip", "Select a mode for Groom visualization"),
+		FNewMenuDelegate::CreateStatic(
+			bIsGroomEditor
+			? &FGroomVisualizationMenuCommands::BuildVisualisationSubMenuForGroomEditor
+			: &FGroomVisualizationMenuCommands::BuildVisualisationSubMenu
+		),
+		FUIAction(
+			FExecuteAction(),
+			FCanExecuteAction(),
+			FIsActionChecked::CreateLambda(
+				[WeakViewport = Viewport]()
+				{
+					if (TSharedPtr<SEditorViewport> Viewport = WeakViewport.Pin())
+					{
+						if (const TSharedPtr<FEditorViewportClient>& Client = Viewport->GetViewportClient())
+						{
+							return Client->IsViewModeEnabled(VMI_VisualizeGroom);
+						}
+					}
+					return false;
+				}
+			)
+		),
+		EUserInterfaceActionType::RadioButton,
+		/* bInOpenSubMenuOnClick = */ false,
+		FSlateIcon(FAppStyle::GetAppStyleSetName(), "EditorViewport.VisualizeGroomMode")
+	);
+	
+	return MenuEntry;
+}
 
 bool FGroomVisualizationMenuCommands::AddCommandTypeToMenu(FMenuBuilder& Menu, const EGroomViewMode Mode) const
 {
diff --git a/Engine/Source/Editor/CommonMenuExtensions/Public/GroomVisualizationMenuCommands.h b/Engine/Source/Editor/CommonMenuExtensions/Public/GroomVisualizationMenuCommands.h
--- a/Engine/Source/Editor/CommonMenuExtensions/Public/GroomVisualizationMenuCommands.h
+++ b/Engine/Source/Editor/CommonMenuExtensions/Public/GroomVisualizationMenuCommands.h
@@ -10,10 +10,12 @@
 #include "Templates/SharedPointer.h"
 #include "UObject/NameTypes.h"
 
+class SEditorViewport;
 class FEditorViewportClient;
 class FMenuBuilder;
 class FUICommandInfo;
 class FUICommandList;
+struct FToolMenuEntry;
 enum class EGroomViewMode : uint8;
 
 class COMMONMENUEXTENSIONS_API FGroomVisualizationMenuCommands : public TCommands<FGroomVisualizationMenuCommands>
@@ -35,6 +37,8 @@ public:
 
 	static void BuildVisualisationSubMenu(FMenuBuilder& Menu);
 	static void BuildVisualisationSubMenuForGroomEditor(FMenuBuilder& Menu);
+	static FToolMenuEntry BuildVisualizationSubMenuItem(const TWeakPtr<SEditorViewport>& Viewport);
+	static FToolMenuEntry BuildVisualizationSubMenuItemForGroomEditor(const TWeakPtr<SEditorViewport>& Viewport);
 
 	virtual void RegisterCommands() override;
 
@@ -52,6 +56,8 @@ private:
 	static void InternalBuildVisualisationSubMenu(FMenuBuilder& Menu, bool bIsGroomEditor);
 	static void ChangeGroomVisualizationMode(TWeakPtr<FEditorViewportClient> WeakClient, FName InName);
 	static bool IsGroomVisualizationModeSelected(TWeakPtr<FEditorViewportClient> WeakClient, FName InName);
+	
+	static FToolMenuEntry InternalBuildVisualizationSubMenuItem(const TWeakPtr<SEditorViewport>& Viewport, bool bIsGroomEditor);
 
 private:
 	TGroomVisualizationModeCommandMap CommandMap;
diff --git a/Engine/Source/Editor/LevelEditor/Private/ViewportToolbar/LevelEditorSubmenus.cpp b/Engine/Source/Editor/LevelEditor/Private/ViewportToolbar/LevelEditorSubmenus.cpp
--- a/Engine/Source/Editor/LevelEditor/Private/ViewportToolbar/LevelEditorSubmenus.cpp
+++ b/Engine/Source/Editor/LevelEditor/Private/ViewportToolbar/LevelEditorSubmenus.cpp
@@ -1026,28 +1026,7 @@ void PopulateViewModesMenu(UToolMenu* InMenu)
 	if (IsGroomEnabled())
 	{
 		FToolMenuSection& Section = InMenu->FindOrAddSection("ViewMode");
-		Section.AddSubMenu(
-			"VisualizeGroomViewMode",
-			LOCTEXT("VisualizeGroomViewModeDisplayName", "Groom"),
-			LOCTEXT("GroomVisualizationMenu_ToolTip", "Select a mode for Groom visualization"),
-			FNewMenuDelegate::CreateStatic(&FGroomVisualizationMenuCommands::BuildVisualisationSubMenu),
-			FUIAction(
-				FExecuteAction(),
-				FCanExecuteAction(),
-				FIsActionChecked::CreateLambda(
-					[WeakViewport = LevelViewport.ToWeakPtr()]()
-					{
-						const TSharedPtr<::SLevelViewport> Viewport = WeakViewport.Pin();
-						check(Viewport.IsValid());
-						FLevelEditorViewportClient& ViewportClient = Viewport->GetLevelViewportClient();
-						return ViewportClient.IsViewModeEnabled(VMI_VisualizeGroom);
-					}
-				)
-			),
-			EUserInterfaceActionType::RadioButton,
-			/* bInOpenSubMenuOnClick = */ false,
-			FSlateIcon(FAppStyle::GetAppStyleSetName(), "EditorViewport.VisualizeGroomMode")
-		);
+		Section.AddEntry(FGroomVisualizationMenuCommands::BuildVisualizationSubMenuItem(LevelViewport.ToWeakPtr()));
 	}
 
 	{
