Add per-group auto LOD bias for finer control.

diff --git a/Engine/Plugins/Runtime/HairStrands/Source/HairStrandsCore/Private/GroomAsset.cpp b/Engine/Plugins/Runtime/HairStrands/Source/HairStrandsCore/Private/GroomAsset.cpp
--- a/Engine/Plugins/Runtime/HairStrands/Source/HairStrandsCore/Private/GroomAsset.cpp
+++ b/Engine/Plugins/Runtime/HairStrands/Source/HairStrandsCore/Private/GroomAsset.cpp
@@ -3977,6 +3977,16 @@ float UGroomAsset::GetAutoLODBias() const
 	return AutoLODBias;
 }
 
+float UGroomAsset::GetAutoLODBias(int32 InGroupIndex) const
+{
+	float GroupAutoLODBias = 0;
+	if (GetHairGroupsLOD().IsValidIndex(InGroupIndex))
+	{
+		GroupAutoLODBias = GetHairGroupsLOD()[InGroupIndex].AutoLODBias;
+	}
+	return FMath::Clamp(AutoLODBias + GroupAutoLODBias, -1.f, 1.f);
+}
+
 TObjectPtr<UDataflowBaseContent> UGroomAsset::CreateDataflowContent()
 {
 	TObjectPtr<UDataflowGroomContent> GroomContent = NewObject<UDataflowGroomContent>(this, UDataflowGroomContent::StaticClass());
diff --git a/Engine/Plugins/Runtime/HairStrands/Source/HairStrandsCore/Private/GroomComponent.cpp b/Engine/Plugins/Runtime/HairStrands/Source/HairStrandsCore/Private/GroomComponent.cpp
--- a/Engine/Plugins/Runtime/HairStrands/Source/HairStrandsCore/Private/GroomComponent.cpp
+++ b/Engine/Plugins/Runtime/HairStrands/Source/HairStrandsCore/Private/GroomComponent.cpp
@@ -2968,7 +2968,7 @@ void UGroomComponent::InitResources(bool bIsBindingReloading)
 			HairGroupInstance->HairGroupPublicData->SetLODVisibilities(LODVisibility);
 			HairGroupInstance->HairGroupPublicData->SetLODGeometryTypes(LODGeometryTypes);
 			HairGroupInstance->HairGroupPublicData->bAutoLOD = GroomAsset->GetLODMode() == EGroomLODMode::Auto;
-			HairGroupInstance->HairGroupPublicData->AutoLODBias = GroomAsset->GetAutoLODBias();
+			HairGroupInstance->HairGroupPublicData->AutoLODBias = GroomAsset->GetAutoLODBias(GroupIt);
 
 			HairGroupInstance->HairGroupPublicData->BindingMinMeshLOD = INDEX_NONE;
 			if (LocalBindingAsset)
diff --git a/Engine/Plugins/Runtime/HairStrands/Source/HairStrandsCore/Public/GroomAsset.h b/Engine/Plugins/Runtime/HairStrands/Source/HairStrandsCore/Public/GroomAsset.h
--- a/Engine/Plugins/Runtime/HairStrands/Source/HairStrandsCore/Public/GroomAsset.h
+++ b/Engine/Plugins/Runtime/HairStrands/Source/HairStrandsCore/Public/GroomAsset.h
@@ -467,7 +467,7 @@ private:
 	UPROPERTY(EditAnywhere, Category = "HairLOD", meta = (DisplayName = "LOD Mode", ToolTip = "Define how LOD adapts curves & points for strands geometry. Auto: adapts the curve count based on screen coverage. Manual: use the discrete LOD created for each groups"))
 	EGroomLODMode LODMode = EGroomLODMode::Default;
 
-	UPROPERTY(EditAnywhere, Category = "HairLOD", meta = (DisplayName = "Auto LOD Bias", ToolTip = "When LOD mode is set to Auto, decrease the screen size at which curves reduction will occur.", ClampMin = "-1", ClampMax = "1", UIMin = "-1.0", UIMax = "1.0"))
+	UPROPERTY(EditAnywhere, Category = "HairLOD", meta = (DisplayName = "Global Auto LOD Bias", ToolTip = "When LOD mode is set to Auto, decrease the screen size at which curves reduction will occur.", ClampMin = "-1", ClampMax = "1", UIMin = "-1.0", UIMax = "1.0"))
 	float AutoLODBias = 0;
 
 	/** Dataflow settings used for any dataflow related operations */
@@ -658,7 +658,12 @@ public:
 	EGroomLODMode GetLODMode() const;
 
 	static FName GetAutoLODBiasMemberName();
+	// Return the asset Auto LOD Bias
 	float GetAutoLODBias() const;
+
+	// Return the group Auto LOD Bias, which combines both the asset's bias, and the group bias
+	float GetAutoLODBias(int32 InGroupIndex) const;
+
 public:
 
 	//~ Begin UObject Interface.
diff --git a/Engine/Plugins/Runtime/HairStrands/Source/HairStrandsCore/Public/GroomAssetInterpolation.h b/Engine/Plugins/Runtime/HairStrands/Source/HairStrandsCore/Public/GroomAssetInterpolation.h
--- a/Engine/Plugins/Runtime/HairStrands/Source/HairStrandsCore/Public/GroomAssetInterpolation.h
+++ b/Engine/Plugins/Runtime/HairStrands/Source/HairStrandsCore/Public/GroomAssetInterpolation.h
@@ -220,6 +220,9 @@ struct HAIRSTRANDSCORE_API FHairGroupsLOD
 {
 	GENERATED_USTRUCT_BODY()
 
+	UPROPERTY(EditAnywhere, Category = "DecimationSettings", meta = (DisplayName = "Group Auto LOD Bias", ToolTip = "When LOD mode is set to Auto, decrease the screen size at which curves reduction will occur. The final bias value is computed by adding this value to the asset's Auto LOD bias value.", ClampMin = "-1", ClampMax = "1", UIMin = "-1.0", UIMax = "1.0"))
+	float AutoLODBias = 0;
+
 	/** LODs  */
 	UPROPERTY(EditAnywhere, Category = "DecimationSettings", meta = (ToolTip = "Reduce the number of hair strands in a uniform manner", ClampMin = "0", UIMin = "0", UIMax = "1.0"))
 	TArray<FHairLODSettings> LODs;
