Normalize the asset editor LOD selection toolbar menu behind a common interface and toolbar entry function.
- Refactor the Static Mesh Editor's LOD functions to hide the non-standard LOD indices of the override in static mesh component.

diff --git a/Engine/Plugins/Runtime/HairStrands/Source/HairStrandsEditor/Private/GroomEditorViewportToolBar.cpp b/Engine/Plugins/Runtime/HairStrands/Source/HairStrandsEditor/Private/GroomEditorViewportToolBar.cpp
--- a/Engine/Plugins/Runtime/HairStrands/Source/HairStrandsEditor/Private/GroomEditorViewportToolBar.cpp
+++ b/Engine/Plugins/Runtime/HairStrands/Source/HairStrandsEditor/Private/GroomEditorViewportToolBar.cpp
@@ -80,7 +80,7 @@ FText SGroomEditorViewportToolbar::GetLODMenuLabel() const
 	TSharedRef<SEditorViewport> BaseViewportRef = GetInfoProvider().GetViewportWidget();
 	TSharedRef<SGroomEditorViewport> ViewportRef = StaticCastSharedRef<SGroomEditorViewport, SEditorViewport>(BaseViewportRef);
 
-	int32 LODSelectionType = ViewportRef->GetLODSelection();
+	int32 LODSelectionType = ViewportRef->GetCurrentLOD();
 
 	if (LODSelectionType >= 0)
 	{
@@ -116,14 +116,14 @@ TSharedRef<SWidget> SGroomEditorViewportToolbar::GenerateLODMenu() const
 			InMenuBuilder.AddMenuEntry(Actions.LODAuto);
 			InMenuBuilder.AddMenuEntry(Actions.LOD0);
 
-			int32 LODCount = ViewportRef->GetLODModelCount();
+			int32 LODCount = ViewportRef->GetLODCount();
 			for (int32 LODId = 1; LODId < LODCount; ++LODId)
 			{
 				FString TitleLabel = FString::Printf(TEXT(" LOD %d"), LODId);
 
-				FUIAction Action(FExecuteAction::CreateSP(ViewportRef, &SGroomEditorViewport::OnSetLODModel, LODId),
+				FUIAction Action(FExecuteAction::CreateSP(ViewportRef, &SGroomEditorViewport::SetLODLevel, LODId),
 					FCanExecuteAction(),
-					FIsActionChecked::CreateSP(ViewportRef, &SGroomEditorViewport::IsLODModelSelected, LODId));
+					FIsActionChecked::CreateSP(ViewportRef, &SGroomEditorViewport::IsLODSelected, LODId));
 
 				InMenuBuilder.AddMenuEntry(FText::FromString(TitleLabel), FText::GetEmpty(), FSlateIcon(), Action, NAME_None, EUserInterfaceActionType::RadioButton);
 			}
diff --git a/Engine/Plugins/Runtime/HairStrands/Source/HairStrandsEditor/Private/Widgets/SGroomEditorViewport.cpp b/Engine/Plugins/Runtime/HairStrands/Source/HairStrandsEditor/Private/Widgets/SGroomEditorViewport.cpp
--- a/Engine/Plugins/Runtime/HairStrands/Source/HairStrandsEditor/Private/Widgets/SGroomEditorViewport.cpp
+++ b/Engine/Plugins/Runtime/HairStrands/Source/HairStrandsEditor/Private/Widgets/SGroomEditorViewport.cpp
@@ -232,16 +232,16 @@ void SGroomEditorViewport::BindCommands()
 	//LOD Auto
 	CommandList->MapAction(
 		ViewportLODMenuCommands.LODAuto,
-		FExecuteAction::CreateSP(this, &SGroomEditorViewport::OnSetLODModel, -1),
+		FExecuteAction::CreateSP(this, &SGroomEditorViewport::SetLODLevel, -1),
 		FCanExecuteAction(),
-		FIsActionChecked::CreateSP(this, &SGroomEditorViewport::IsLODModelSelected, -1));
+		FIsActionChecked::CreateSP(this, &SGroomEditorViewport::IsLODSelected, -1));
 
 	// LOD 0
 	CommandList->MapAction(
 		ViewportLODMenuCommands.LOD0,
-		FExecuteAction::CreateSP(this, &SGroomEditorViewport::OnSetLODModel, 0),
+		FExecuteAction::CreateSP(this, &SGroomEditorViewport::SetLODLevel, 0),
 		FCanExecuteAction(),
-		FIsActionChecked::CreateSP(this, &SGroomEditorViewport::IsLODModelSelected, 0));
+		FIsActionChecked::CreateSP(this, &SGroomEditorViewport::IsLODSelected, 0));
 	// all other LODs will be added dynamically
 	
 	const FGroomVisualizationMenuCommands& GroomCommands = FGroomVisualizationMenuCommands::Get();
@@ -254,16 +254,16 @@ void SGroomEditorViewport::RefreshViewport()
 	SceneViewport->Invalidate();
 }
 
-void SGroomEditorViewport::OnSetLODModel(int32 InLODSelection)
+void SGroomEditorViewport::SetLODLevel(int32 LODIndex)
 {
 	if (GroomComponent)
 	{
-		GroomComponent->SetForcedLOD(InLODSelection);
+		GroomComponent->SetForcedLOD(LODIndex);
 		RefreshViewport();
-	}
+	}	
 }
 
-bool SGroomEditorViewport::IsLODModelSelected(int32 InLODSelection) const
+bool SGroomEditorViewport::IsLODSelected(int32 InLODSelection) const
 {
 	if (GroomComponent)
 	{
@@ -272,6 +272,30 @@ bool SGroomEditorViewport::IsLODModelSelected(int32 InLODSelection) const
 	return false;
 }
 
+int32 SGroomEditorViewport::GetCurrentLOD() const
+{
+	if (GroomComponent)
+	{
+		return GroomComponent->GetForcedLOD();
+	}
+	return -1;
+}
+
+int32 SGroomEditorViewport::GetLODCount() const
+{
+	if (GroomComponent)
+	{
+		return GroomComponent->GetNumLODs();
+	}
+	return -1;
+}
+
+void SGroomEditorViewport::FillLODCommands(TArray<TSharedPtr<FUICommandInfo>>& Commands)
+{
+	Commands.Add(FGroomViewportLODCommands::Get().LODAuto);
+	Commands.Add(FGroomViewportLODCommands::Get().LOD0);
+}
+
 void SGroomEditorViewport::OnFocusViewportToSelection()
 {
 	if (GroomComponent)
@@ -412,65 +436,11 @@ TSharedPtr<SWidget> SGroomEditorViewport::BuildViewportToolbar()
 			// LOD Menu
 			RightSection.AddDynamicEntry("DynamicLOD", FNewToolMenuSectionDelegate::CreateLambda([](FToolMenuSection& Section)
 			{
-				UUnrealEdViewportToolbarContext* Context = Section.FindContext<UUnrealEdViewportToolbarContext>();
-				if (!Context)
+				if (const UUnrealEdViewportToolbarContext* Context = Section.FindContext<UUnrealEdViewportToolbarContext>())
 				{
-					return;
+					const TWeakPtr<SGroomEditorViewport> GroomViewport = StaticCastWeakPtr<SGroomEditorViewport>(Context->Viewport);
+					Section.AddEntry(UE::UnrealEd::CreatePreviewLODSelectionSubmenu(GroomViewport));
 				}
-				
-				TSharedPtr<SEditorViewport> EditorViewport = Context->Viewport.Pin();
-				if (!EditorViewport)
-				{
-					return;
-				}
-				
-				TSharedPtr<SGroomEditorViewport> GroomViewport = StaticCastSharedPtr<SGroomEditorViewport>(EditorViewport);
-				
-				Section.AddSubMenu(
-					"LOD",
-					TAttribute<FText>::CreateLambda([WeakViewport = GroomViewport.ToWeakPtr()]
-					{
-						if (TSharedPtr<SGroomEditorViewport> Viewport = WeakViewport.Pin())
-						{
-							const int32 LODSelectionType = Viewport->GetLODSelection();
-							if (LODSelectionType >= 0)
-							{
-								return FText::Format(LOCTEXT("LodMenu_LabelFormat", "LOD {0}"), LODSelectionType);
-							}
-						}
-						return LOCTEXT("LODMenu_AutoLabel", "LOD Auto");
-					}),
-					LOCTEXT("LodMenu_Tooltip", "Set the Level of Detail of the viewport."),
-					FNewMenuDelegate::CreateLambda([WeakViewport = GroomViewport.ToWeakPtr()](FMenuBuilder& Menu)
-					{
-						TSharedPtr<SGroomEditorViewport> Viewport = WeakViewport.Pin();
-						if (!Viewport)
-						{
-							return;
-						}
-						
-						Menu.BeginSection("GroomViewportPreviewLODs", LOCTEXT("ShowLOD_PreviewLabel", "Preview LODs"));
-						{
-							Menu.AddMenuEntry(FGroomViewportLODCommands::Get().LODAuto);
-							Menu.AddMenuEntry(FGroomViewportLODCommands::Get().LOD0);
-
-							const int32 LODCount = Viewport->GetLODModelCount();
-							for (int32 LODId = 1; LODId < LODCount; ++LODId)
-							{
-								FString TitleLabel = FString::Printf(TEXT(" LOD %d"), LODId);
-
-								FUIAction Action(
-									FExecuteAction::CreateSP(Viewport.Get(), &SGroomEditorViewport::OnSetLODModel, LODId),
-									FCanExecuteAction(),
-									FIsActionChecked::CreateSP(Viewport.Get(), &SGroomEditorViewport::IsLODModelSelected, LODId)
-								);
-
-								Menu.AddMenuEntry(FText::FromString(TitleLabel), FText::GetEmpty(), FSlateIcon(), Action, NAME_None, EUserInterfaceActionType::RadioButton);
-							}
-						}
-						Menu.EndSection();
-					})
-				);
 			}));
 		}
 		
@@ -520,22 +490,4 @@ void SGroomEditorViewport::OnFloatingButtonClicked()
 
 }
 
-int32 SGroomEditorViewport::GetLODSelection() const
-{
-	if (GroomComponent)
-	{
-		return GroomComponent->GetForcedLOD();
-	}
-	return -1;
-}
-
-int32 SGroomEditorViewport::GetLODModelCount() const
-{
-	if (GroomComponent)
-	{
-		return GroomComponent->GetNumLODs();
-	}
-	return -1;
-}
-
 #undef LOCTEXT_NAMESPACE
\ No newline at end of file
diff --git a/Engine/Plugins/Runtime/HairStrands/Source/HairStrandsEditor/Public/Widgets/SGroomEditorViewport.h b/Engine/Plugins/Runtime/HairStrands/Source/HairStrandsEditor/Public/Widgets/SGroomEditorViewport.h
--- a/Engine/Plugins/Runtime/HairStrands/Source/HairStrandsEditor/Public/Widgets/SGroomEditorViewport.h
+++ b/Engine/Plugins/Runtime/HairStrands/Source/HairStrandsEditor/Public/Widgets/SGroomEditorViewport.h
@@ -13,6 +13,7 @@
 #include "SEditorViewport.h"
 #include "SCommonEditorViewportToolbarBase.h"
 #include "HairStrandsInterface.h"
+#include "IPreviewLODController.h"
 
 
 class UGroomComponent;
@@ -21,7 +22,7 @@ class SDockTab;
 /**
  * Material Editor Preview viewport widget
  */
-class SGroomEditorViewport : public SEditorViewport, public FGCObject, public ICommonEditorViewportToolbarInfoProvider
+class SGroomEditorViewport : public SEditorViewport, public FGCObject, public ICommonEditorViewportToolbarInfoProvider, public IPreviewLODController
 {
 public:
 
@@ -54,11 +55,15 @@ public:
 
 	// set the mesh on which we are grooming
 	void SetSkeletalMeshComponent(USkeletalMeshComponent *SkeletalGroomTarget);
-
-	int32 GetLODSelection() const;
-	int32 GetLODModelCount() const;
-	bool IsLODModelSelected(int32 InLODSelection) const;
-	void OnSetLODModel(int32 InLODSelection);
+	
+	// IPreviewLODController interface
+	virtual int32 GetCurrentLOD() const override;
+	virtual int32 GetLODCount() const override;
+	virtual bool IsLODSelected(int32 LODIndex) const override;
+	virtual void SetLODLevel(int32 LODIndex) override;
+	virtual void FillLODCommands(TArray<TSharedPtr<FUICommandInfo>>& Commands) override;
+	virtual int32 GetAutoLODStartingIndex() const override { return 1; }
+	// ~IPreviewLODController interface
 
 	TSharedPtr<class FAdvancedPreviewScene> GetAdvancedPreviewScene() { return AdvancedPreviewScene; }
 
diff --git a/Engine/Source/Editor/StaticMeshEditor/Private/SStaticMeshEditorViewport.cpp b/Engine/Source/Editor/StaticMeshEditor/Private/SStaticMeshEditorViewport.cpp
--- a/Engine/Source/Editor/StaticMeshEditor/Private/SStaticMeshEditorViewport.cpp
+++ b/Engine/Source/Editor/StaticMeshEditor/Private/SStaticMeshEditorViewport.cpp
@@ -598,7 +598,7 @@ void SStaticMeshEditorViewport::ForceLODLevel(int32 InForcedLOD)
 	SceneViewport->Invalidate();
 }
 
-int32 SStaticMeshEditorViewport::GetLODSelection() const
+int32 SStaticMeshEditorViewport::GetCurrentLOD() const
 {
 	if (PreviewMeshComponent)
 	{
@@ -608,18 +608,18 @@ int32 SStaticMeshEditorViewport::GetLODSelection() const
 				&& PreviewMeshComponent->GetStaticMesh()->GetRenderData()
 				&& PreviewMeshComponent->GetStaticMesh()->GetRenderData()->RayTracingProxy)
 			{
-				return PreviewMeshComponent->GetStaticMesh()->GetRenderData()->RayTracingProxy->PreviewLODLevel + 1;
+				return PreviewMeshComponent->GetStaticMesh()->GetRenderData()->RayTracingProxy->PreviewLODLevel;
 			}
 		}
 		else
 		{
-			return PreviewMeshComponent->ForcedLodModel;
+			return PreviewMeshComponent->ForcedLodModel - 1;
 		}
 	}
-	return 0;
+	return INDEX_NONE;
 }
 
-bool SStaticMeshEditorViewport::IsLODModelSelected(int32 InLODSelection) const
+bool SStaticMeshEditorViewport::IsLODSelected(int32 InLODSelection) const
 {
 	if (PreviewMeshComponent)
 	{
@@ -629,18 +629,18 @@ bool SStaticMeshEditorViewport::IsLODModelSelected(int32 InLODSelection) const
 				&& PreviewMeshComponent->GetStaticMesh()->GetRenderData()
 				&& PreviewMeshComponent->GetStaticMesh()->GetRenderData()->RayTracingProxy)
 			{
-				return InLODSelection - 1 == PreviewMeshComponent->GetStaticMesh()->GetRenderData()->RayTracingProxy->PreviewLODLevel;
+				return InLODSelection == PreviewMeshComponent->GetStaticMesh()->GetRenderData()->RayTracingProxy->PreviewLODLevel;
 			}
 		}
 		else
 		{
-			return (PreviewMeshComponent->ForcedLodModel == InLODSelection) ? true : false;
+			return PreviewMeshComponent->ForcedLodModel - 1 == InLODSelection;
 		}
 	}
 	return false;
 }
 
-void SStaticMeshEditorViewport::OnSetLODModel(int32 InLODSelection)
+void SStaticMeshEditorViewport::SetLODLevel(int32 InLODSelection)
 {
 	if (PreviewMeshComponent)
 	{
@@ -659,14 +659,14 @@ void SStaticMeshEditorViewport::OnSetLODModel(int32 InLODSelection)
 
 			if (RayTracingProxy)
 			{
-				RayTracingProxy->PreviewLODLevel = LODSelection - 1;
+				RayTracingProxy->PreviewLODLevel = LODSelection;
 				PreviewMeshComponent->MarkRenderStateDirty();
 			}
 		}
 		else
 		{
-			PreviewMeshComponent->bOverrideMinLOD = (LODSelection != 0);
-			PreviewMeshComponent->SetForcedLodModel(LODSelection);
+			PreviewMeshComponent->bOverrideMinLOD = (LODSelection >= 0);
+			PreviewMeshComponent->SetForcedLodModel(LODSelection + 1);
 
 			if (RayTracingProxy)
 			{
@@ -681,6 +681,12 @@ void SStaticMeshEditorViewport::OnSetLODModel(int32 InLODSelection)
 	}
 }
 
+void SStaticMeshEditorViewport::FillLODCommands(TArray<TSharedPtr<FUICommandInfo>>& Commands)
+{
+	Commands.Add(FStaticMeshViewportLODCommands::Get().LODAuto);
+	Commands.Add(FStaticMeshViewportLODCommands::Get().LOD0);
+}
+
 void SStaticMeshEditorViewport::OnLODModelChanged()
 {
 	if (PreviewMeshComponent && LODSelection != PreviewMeshComponent->ForcedLodModel)
@@ -689,7 +695,7 @@ void SStaticMeshEditorViewport::OnLODModelChanged()
 	}
 }
 
-int32 SStaticMeshEditorViewport::GetLODModelCount() const
+int32 SStaticMeshEditorViewport::GetLODCount() const
 {
 	if (PreviewMeshComponent && PreviewMeshComponent->GetStaticMesh())
 	{
@@ -894,16 +900,16 @@ void SStaticMeshEditorViewport::BindCommands()
 	//LOD Auto
 	CommandList->MapAction(
 		ViewportLODMenuCommands.LODAuto,
-		FExecuteAction::CreateSP(this, &SStaticMeshEditorViewport::OnSetLODModel, 0),
+		FExecuteAction::CreateSP(this, &SStaticMeshEditorViewport::SetLODLevel, -1),
 		FCanExecuteAction(),
-		FIsActionChecked::CreateSP(this, &SStaticMeshEditorViewport::IsLODModelSelected, 0));
+		FIsActionChecked::CreateSP(this, &SStaticMeshEditorViewport::IsLODSelected, -1));
 
 	// LOD 0
 	CommandList->MapAction(
 		ViewportLODMenuCommands.LOD0,
-		FExecuteAction::CreateSP(this, &SStaticMeshEditorViewport::OnSetLODModel, 1),
+		FExecuteAction::CreateSP(this, &SStaticMeshEditorViewport::SetLODLevel, 0),
 		FCanExecuteAction(),
-		FIsActionChecked::CreateSP(this, &SStaticMeshEditorViewport::IsLODModelSelected, 1));
+		FIsActionChecked::CreateSP(this, &SStaticMeshEditorViewport::IsLODSelected, 0));
 	// all other LODs will be added dynamically
 
 }
diff --git a/Engine/Source/Editor/StaticMeshEditor/Private/SStaticMeshEditorViewport.h b/Engine/Source/Editor/StaticMeshEditor/Private/SStaticMeshEditorViewport.h
--- a/Engine/Source/Editor/StaticMeshEditor/Private/SStaticMeshEditorViewport.h
+++ b/Engine/Source/Editor/StaticMeshEditor/Private/SStaticMeshEditorViewport.h
@@ -10,6 +10,7 @@
 #include "EditorViewportClient.h"
 #include "StaticMeshEditorViewportClient.h"
 #include "AdvancedPreviewScene.h"
+#include "IPreviewLODController.h"
 #include "SEditorViewport.h"
 #include "SCommonEditorViewportToolbarBase.h"
 #include "SAssetEditorViewport.h"
@@ -24,7 +25,7 @@ class SRichTextBlock;
 /**
  * StaticMesh Editor Preview viewport widget
  */
-class SStaticMeshEditorViewport : public SAssetEditorViewport, public FGCObject, public ICommonEditorViewportToolbarInfoProvider
+class SStaticMeshEditorViewport : public SAssetEditorViewport, public FGCObject, public ICommonEditorViewportToolbarInfoProvider, public IPreviewLODController
 {
 public:
 	SLATE_BEGIN_ARGS( SStaticMeshEditorViewport ){}
@@ -61,21 +62,21 @@ public:
 	 */
 	void ForceLODLevel(int32 InForcedLOD);
 
+	
+	// IPreviewLODController interface
 	/**
 	 *
 	 * Query LOD level of static mesh component 
 	 *
 	 */
-	int32 GetLODSelection() const;
-
-	/** Function to get the number of LOD models associated with the preview static mesh*/
-	int32 GetLODModelCount() const;
-
-	/**  LOD model selection checking function*/
-	bool IsLODModelSelected(int32 LODSelectionType) const;
+	virtual int32 GetCurrentLOD() const override;
+	virtual int32 GetLODCount() const override;
+	virtual bool IsLODSelected(int32 LODIndex) const override;
+	virtual void SetLODLevel(int32 LODIndex) override;
+	virtual void FillLODCommands(TArray<TSharedPtr<FUICommandInfo>>& Commands) override;
+	virtual int32 GetAutoLODStartingIndex() const override { return 1; }
+	// ~IPreviewLODController
 
-	/**  Function to set LOD model selection*/
-	void OnSetLODModel(int32 LODSelectionType);
 	void OnLODModelChanged();
 
 	/** Retrieves the static mesh component. */
diff --git a/Engine/Source/Editor/StaticMeshEditor/Private/StaticMeshEditorViewportToolbarSections.cpp b/Engine/Source/Editor/StaticMeshEditor/Private/StaticMeshEditorViewportToolbarSections.cpp
--- a/Engine/Source/Editor/StaticMeshEditor/Private/StaticMeshEditorViewportToolbarSections.cpp
+++ b/Engine/Source/Editor/StaticMeshEditor/Private/StaticMeshEditorViewportToolbarSections.cpp
@@ -13,64 +13,17 @@
 
 #define LOCTEXT_NAMESPACE "StaticMeshEditorViewportToolbarSections"
 
-namespace UE::StaticMeshEditor::Private
-{
-void PopulateLODSubmenu(UToolMenu* InMenu)
-{
-	if (UUnrealEdViewportToolbarContext* const EditorViewportContext =
-			InMenu->FindContext<UUnrealEdViewportToolbarContext>())
-	{
-		if (TSharedPtr<SStaticMeshEditorViewport> StaticMeshEditorViewport =
-				StaticCastSharedPtr<SStaticMeshEditorViewport>(EditorViewportContext->Viewport.Pin()))
-		{
-			const FStaticMeshViewportLODCommands& Actions = FStaticMeshViewportLODCommands::Get();
-
-			FToolMenuSection& PreviewLODSection =
-				InMenu->FindOrAddSection("StaticMeshViewportPreviewLODs", LOCTEXT("ShowLOD_PreviewLabel", "Preview LODs"));
-
-			PreviewLODSection.AddMenuEntry(Actions.LODAuto);
-			PreviewLODSection.AddMenuEntry(Actions.LOD0);
-
-			int32 LODCount = StaticMeshEditorViewport->GetLODModelCount();
-			for (int32 LODId = 1; LODId < LODCount; ++LODId)
-			{
-				FString TitleLabel = FString::Printf(TEXT(" LOD %d"), LODId);
-
-				FUIAction Action(
-					FExecuteAction::CreateSP(
-						StaticMeshEditorViewport.ToSharedRef(), &SStaticMeshEditorViewport::OnSetLODModel, LODId + 1
-					),
-					FCanExecuteAction(),
-					FIsActionChecked::CreateSP(
-						StaticMeshEditorViewport.ToSharedRef(), &SStaticMeshEditorViewport::IsLODModelSelected, LODId + 1
-					)
-				);
-
-				PreviewLODSection.AddMenuEntry(
-					FName(TitleLabel),
-					FText::FromString(TitleLabel),
-					FText::GetEmpty(),
-					FSlateIcon(),
-					Action,
-					EUserInterfaceActionType::RadioButton
-				);
-			}
-		}
-	}
-}
-} // namespace UE::StaticMeshEditor::Private
-
 FText UE::StaticMeshEditor::GetLODMenuLabel(const TSharedPtr<SStaticMeshEditorViewport>& InStaticMeshEditorViewport)
 {
 	FText Label = LOCTEXT("LODMenu_AutoLabel", "LOD Auto");
 
 	if (InStaticMeshEditorViewport)
 	{
-		int32 LODSelectionType = InStaticMeshEditorViewport->GetLODSelection();
+		const int32 LODSelectionType = InStaticMeshEditorViewport->GetCurrentLOD();
 
-		if (LODSelectionType > 0)
+		if (LODSelectionType >= 0)
 		{
-			FString TitleLabel = FString::Printf(TEXT("LOD %d"), LODSelectionType - 1);
+			FString TitleLabel = FString::Printf(TEXT("LOD %d"), LODSelectionType);
 			Label = FText::FromString(TitleLabel);
 		}
 	}
@@ -95,36 +48,10 @@ FToolMenuEntry UE::StaticMeshEditor::CreateLODSubmenu()
 		FNewToolMenuSectionDelegate::CreateLambda(
 			[](FToolMenuSection& InDynamicSection) -> void
 			{
-				if (UUnrealEdViewportToolbarContext* const EditorViewportContext =
-						InDynamicSection.FindContext<UUnrealEdViewportToolbarContext>())
+				if (UUnrealEdViewportToolbarContext* const EditorViewportContext = InDynamicSection.FindContext<UUnrealEdViewportToolbarContext>())
 				{
-					TWeakPtr<SStaticMeshEditorViewport> StaticMeshEditorViewportWeak =
-						StaticCastSharedPtr<SStaticMeshEditorViewport>(EditorViewportContext->Viewport.Pin());
-
-					// Label updates based on currently selected LOD
-					const TAttribute<FText> Label = TAttribute<FText>::CreateLambda(
-						[StaticMeshEditorViewportWeak]()
-						{
-							if (TSharedPtr<SStaticMeshEditorViewport> Viewport = StaticMeshEditorViewportWeak.Pin())
-							{
-								return GetLODMenuLabel(Viewport);
-							}
-
-							return LOCTEXT("LODSubmenuLabel", "LOD");
-						}
-					);
-
-					FToolMenuEntry& Entry = InDynamicSection.AddSubMenu(
-						"LOD",
-						Label,
-						FText(),
-						FNewToolMenuDelegate::CreateLambda(
-							[](UToolMenu* Submenu) -> void
-							{
-								UE::StaticMeshEditor::Private::PopulateLODSubmenu(Submenu);
-							}
-						)
-					);
+					TWeakPtr<SStaticMeshEditorViewport> StaticMeshEditorViewport = StaticCastWeakPtr<SStaticMeshEditorViewport>(EditorViewportContext->Viewport);
+					FToolMenuEntry& Entry = InDynamicSection.AddEntry(UE::UnrealEd::CreatePreviewLODSelectionSubmenu(StaticMeshEditorViewport));
 					Entry.ToolBarData.ResizeParams.ClippingPriority = 800;
 				}
 			}
@@ -173,7 +100,10 @@ TSharedRef<SWidget> UE::StaticMeshEditor::GenerateLODMenuWidget(const TSharedPtr
 				FNewToolMenuDelegate::CreateLambda(
 					[](UToolMenu* InMenu)
 					{
-						UE::StaticMeshEditor::Private::PopulateLODSubmenu(InMenu);
+						if (UUnrealEdViewportToolbarContext* const EditorViewportContext = InMenu->FindContext<UUnrealEdViewportToolbarContext>())
+						{
+							UE::UnrealEd::FillPreviewLODSelectionSubmenu(InMenu, StaticCastWeakPtr<SStaticMeshEditorViewport>(EditorViewportContext->Viewport));
+						}
 					}
 				)
 			);
diff --git a/Engine/Source/Editor/UnrealEd/Private/ViewportToolbar/UnrealEdViewportToolbar.cpp b/Engine/Source/Editor/UnrealEd/Private/ViewportToolbar/UnrealEdViewportToolbar.cpp
--- a/Engine/Source/Editor/UnrealEd/Private/ViewportToolbar/UnrealEdViewportToolbar.cpp
+++ b/Engine/Source/Editor/UnrealEd/Private/ViewportToolbar/UnrealEdViewportToolbar.cpp
@@ -12,6 +12,7 @@
 #include "GPUSkinCache.h"
 #include "GPUSkinCacheVisualizationMenuCommands.h"
 #include "IPreviewProfileController.h"
+#include "IPreviewLODController.h"
 #include "LevelEditor.h"
 #include "LevelEditorActions.h"
 #include "RayTracingDebugVisualizationMenuCommands.h"
@@ -2808,6 +2809,90 @@ FToolMenuEntry CreateAssetViewerProfileSubmenu()
 	);
 }
 
+FToolMenuEntry CreatePreviewLODSelectionSubmenu(TWeakPtr<IPreviewLODController> LODController)
+{
+	return FToolMenuEntry::InitSubMenu(
+		"LOD",
+		TAttribute<FText>::CreateLambda([LODController]
+		{
+			if (TSharedPtr<IPreviewLODController> Controller = LODController.Pin())
+			{
+				const int32 CurrentLOD = Controller->GetCurrentLOD();
+				if (CurrentLOD >= 0)
+				{
+					return FText::Format(LOCTEXT("LODMenu_LabelFormat", "LOD {0}"), CurrentLOD);
+				}
+			}
+			return LOCTEXT("LODMenu_AutoLabel", "LOD Auto");
+		}),
+		LOCTEXT("LODMenu_Tooltip", "Set the Level of Detail of the viewport."),
+		FNewToolMenuDelegate::CreateStatic(&FillPreviewLODSelectionSubmenu, LODController),
+		false,
+		FSlateIcon(FAppStyle::GetAppStyleSetName(), "Icons.LOD")
+	);
+}
+
+void FillPreviewLODSelectionSubmenu(UToolMenu* Menu, TWeakPtr<IPreviewLODController> LODController)
+{
+	if (TSharedPtr<IPreviewLODController> Controller = LODController.Pin())
+	{
+		FToolMenuSection& Section = Menu->FindOrAddSection("LOD", LOCTEXT("LODMenuSectionHeader", "Preview LODs"));
+				
+		// Allow some LOD items to be handled by commands
+		TArray<TSharedPtr<FUICommandInfo>> Commands;
+		Controller->FillLODCommands(Commands);
+		for (const TSharedPtr<FUICommandInfo>& Command : Commands)
+		{
+			Section.AddMenuEntry(Command);
+		}
+				
+		// Fill in remaining LOD levels with menu entries
+		const int32 LODCount = Controller->GetLODCount();
+		for (int32 LODIndex = Controller->GetAutoLODStartingIndex(); LODIndex < LODCount; ++LODIndex)
+		{
+			FText Label;
+			FText Tooltip;
+			
+			if (LODIndex == INDEX_NONE)
+			{
+				Label = LOCTEXT("LODMenu_LODAutoLabel", "LOD Auto");
+				Tooltip = LOCTEXT("LODMenu_LODAutoTooltip", "Automatically select the the Level of Detail.");
+			}
+			else
+			{
+				Label = FText::Format(LOCTEXT("LODMenu_LODLabelFormat", "LOD {0}"), LODIndex);
+				Tooltip = FText::Format(LOCTEXT("LODMenu_LODTooltipFormat", "Sets the Level of Detail to {0}"), LODIndex);
+			}
+		
+			Section.AddMenuEntry(
+				NAME_None,
+				Label,
+				Tooltip,
+				FSlateIcon(),
+				FUIAction(
+					FExecuteAction::CreateLambda([LODController, LODIndex]
+					{
+						if (TSharedPtr<IPreviewLODController> Controller = LODController.Pin())
+						{
+							Controller->SetLODLevel(LODIndex);
+						}
+					}),
+					FCanExecuteAction(),
+					FIsActionChecked::CreateLambda([LODController, LODIndex]
+					{
+						if (TSharedPtr<IPreviewLODController> Controller = LODController.Pin())
+						{
+							return Controller->IsLODSelected(LODIndex);
+						}
+						return false;
+					})
+				),
+				EUserInterfaceActionType::RadioButton
+			);
+		}
+	}
+}
+
 void ExtendPreviewSceneSettingsWithTabEntry(FName InAssetViewerProfileSubmenuName)
 {
 	UToolMenu* const Submenu = UToolMenus::Get()->ExtendMenu(InAssetViewerProfileSubmenuName);
diff --git a/Engine/Source/Editor/UnrealEd/Public/IPreviewLODController.h b/Engine/Source/Editor/UnrealEd/Public/IPreviewLODController.h
new file mode 100644
--- /dev/null
+++ b/Engine/Source/Editor/UnrealEd/Public/IPreviewLODController.h
@@ -0,0 +1,37 @@
+﻿// Copyright Epic Games, Inc. All Rights Reserved.
+
+#pragma once
+
+#include "CoreMinimal.h"
+
+class FUICommandInfo;
+/**
+ * Sets or gets the various LOD levels associated with the controller.
+ * Used especially for asset viewports.
+ */
+class IPreviewLODController
+{
+public:
+	virtual ~IPreviewLODController() = default;
+
+	/** Gets the current LOD. A value of INDEX_NONE implies an 'Auto' LOD. */
+	virtual int32 GetCurrentLOD() const = 0;
+	
+	/** Gets whether a given LOD is active. */
+	virtual bool IsLODSelected(int32 LODIndex) const = 0;
+	
+	/** Gets the number of LODs currently supported. */
+	virtual int32 GetLODCount() const = 0;
+	
+	/** Sets the LOD level by index. A value of INDEX_NONE implies an 'Auto' LOD. */
+	virtual void SetLODLevel(int32 LODIndex) = 0;
+	
+	/** Provide LOD commands that should be included before any auto-generated LOD menu items. */
+	virtual void FillLODCommands(TArray<TSharedPtr<FUICommandInfo>>& Commands) {}
+	
+	/**
+	 * Set which LOD level an auto-generated list of LOD options should start from.
+	 * Allows certain LOD levels to be handled by commands provided by `FillLODCommands()`.
+	 */
+	virtual int32 GetAutoLODStartingIndex() const { return 0; }
+};
\ No newline at end of file
diff --git a/Engine/Source/Editor/UnrealEd/Public/ViewportToolbar/UnrealEdViewportToolbar.h b/Engine/Source/Editor/UnrealEd/Public/ViewportToolbar/UnrealEdViewportToolbar.h
--- a/Engine/Source/Editor/UnrealEd/Public/ViewportToolbar/UnrealEdViewportToolbar.h
+++ b/Engine/Source/Editor/UnrealEd/Public/ViewportToolbar/UnrealEdViewportToolbar.h
@@ -15,6 +15,7 @@ class FTabManager;
 class FText;
 class IPreviewProfileController;
 class IPreviewProfileController;
+class IPreviewLODController;
 class SEditorViewport;
 class UToolMenu;
 enum ELevelViewportType : int;
@@ -170,6 +171,16 @@ UNREALED_API FToolMenuEntry CreateViewportToolbarAssetViewerProfileSubmenu();
 UNREALED_API FToolMenuEntry CreateAssetViewerProfileSubmenu();
 UNREALED_API void ExtendPreviewSceneSettingsWithTabEntry(FName InAssetViewerProfileSubmenuName);
 
+/**
+ * Creates a submenu that displays and manipulates the current LOD of the provided LODController.
+ * This should be invoked from within a dynamic menu or section to provide the correct controller.
+ */
+UNREALED_API FToolMenuEntry CreatePreviewLODSelectionSubmenu(TWeakPtr<IPreviewLODController> LODController);
+/**
+ * Fills the given UToolMenu with the appropriate menu items to reflect & manipulate the provided LODController
+ */
+UNREALED_API void FillPreviewLODSelectionSubmenu(UToolMenu* InToolMenu, TWeakPtr<IPreviewLODController> LODController);
+
 UNREALED_API void AddExposureSection(UToolMenu* InMenu, const TSharedPtr<SEditorViewport>& EditorViewport);
 UNREALED_API void PopulateCameraMenu(UToolMenu* InMenu);
 
