[Hair] Add Width tracking as imported hair attribute.

diff --git a/Engine/Plugins/Runtime/HairStrands/Source/HairStrandsCore/Private/GroomBuilder.cpp b/Engine/Plugins/Runtime/HairStrands/Source/HairStrandsCore/Private/GroomBuilder.cpp
--- a/Engine/Plugins/Runtime/HairStrands/Source/HairStrandsCore/Private/GroomBuilder.cpp
+++ b/Engine/Plugins/Runtime/HairStrands/Source/HairStrandsCore/Private/GroomBuilder.cpp
@@ -54,7 +54,7 @@ bool DoesHairStrandsSupportCompressedPosition();
 
 FString FGroomBuilder::GetVersion()
 {
-	FString Out = TEXT("v18a");
+	FString Out = TEXT("v18b_Debug0");
 	if (IsRenderCurveEnabled())
 	{
 		Out += TEXT("_Curve3e");
diff --git a/Engine/Plugins/Runtime/HairStrands/Source/HairStrandsCore/Private/HairAttributes.cpp b/Engine/Plugins/Runtime/HairStrands/Source/HairStrandsCore/Private/HairAttributes.cpp
--- a/Engine/Plugins/Runtime/HairStrands/Source/HairStrandsCore/Private/HairAttributes.cpp
+++ b/Engine/Plugins/Runtime/HairStrands/Source/HairStrandsCore/Private/HairAttributes.cpp
@@ -62,10 +62,11 @@ void SetHairAttributeFlags(uint32& Out, EHairAttributeFlags InFlag)
 const TCHAR* GetHairAttributeText(EHairAttribute In, uint32 InFlags)
 {
 	// If a new optional attribute is added, please add its UI/text description here
-	static_assert(uint32(EHairAttribute::Count) == 7);
+	static_assert(uint32(EHairAttribute::Count) == 8);
 
 	switch (In)
 	{
+	case EHairAttribute::Width:						return TEXT("Width");
 	case EHairAttribute::RootUV:					return HasHairAttributeFlags(InFlags, EHairAttributeFlags::HasRootUDIM) ? TEXT("RootUV(UDIM)") : TEXT("RootUV");
 	case EHairAttribute::ClumpID:					return HasHairAttributeFlags(InFlags, EHairAttributeFlags::HasMultipleClumpIDs) ? TEXT("ClumpID(3)") : TEXT("ClumpID");
 	case EHairAttribute::StrandID:					return TEXT("StrandID");
@@ -79,6 +80,7 @@ const TCHAR* GetHairAttributeText(EHairAttribute In, uint32 InFlags)
 
 bool FHairStrandsPoints::HasAttribute(EHairAttribute In) const
 {
+	if (In == EHairAttribute::Width)		{ return PointsRadius.Num() > 0; }
 	if (In == EHairAttribute::Color)		{ return PointsBaseColor.Num() > 0; }
 	if (In == EHairAttribute::Roughness)	{ return PointsRoughness.Num() > 0; }
 	if (In == EHairAttribute::AO)			{ return PointsAO.Num() > 0; }
diff --git a/Engine/Plugins/Runtime/HairStrands/Source/HairStrandsCore/Private/HairDescription.cpp b/Engine/Plugins/Runtime/HairStrands/Source/HairStrandsCore/Private/HairDescription.cpp
--- a/Engine/Plugins/Runtime/HairStrands/Source/HairStrandsCore/Private/HairDescription.cpp
+++ b/Engine/Plugins/Runtime/HairStrands/Source/HairStrandsCore/Private/HairDescription.cpp
@@ -69,10 +69,11 @@ bool FHairDescription::IsValid() const
 bool FHairDescription::HasAttribute(EHairAttribute InAttribute) const
 {
 	// If a new optional attribute is added, please add its UI/text description here
-	static_assert(uint32(EHairAttribute::Count) == 7);
+	static_assert(uint32(EHairAttribute::Count) == 8);
 
 	switch (InAttribute)
 	{
+	case EHairAttribute::Width:			return StrandAttributes().GetAttributesRef<float>(HairAttribute::Vertex::Width).IsValid() || StrandAttributes().GetAttributesRef<float>(HairAttribute::Strand::Width).IsValid() || StrandAttributes().GetAttributesRef<float>(HairAttribute::Groom::Width).IsValid();
 	case EHairAttribute::RootUV:		return StrandAttributes().GetAttributesRef<FVector2f>(HairAttribute::Strand::RootUV).IsValid();
 	case EHairAttribute::ClumpID:		return StrandAttributes().GetAttributesRef<int>(HairAttribute::Strand::ClumpID).IsValid() || StrandAttributes().GetAttributesRef<FVector3f>(HairAttribute::Strand::ClumpID).IsValid();
 	case EHairAttribute::StrandID:		return StrandAttributes().GetAttributesRef<int>(HairAttribute::Strand::ID).IsValid();
diff --git a/Engine/Plugins/Runtime/HairStrands/Source/HairStrandsCore/Public/HairAttributes.h b/Engine/Plugins/Runtime/HairStrands/Source/HairStrandsCore/Public/HairAttributes.h
--- a/Engine/Plugins/Runtime/HairStrands/Source/HairStrandsCore/Public/HairAttributes.h
+++ b/Engine/Plugins/Runtime/HairStrands/Source/HairStrandsCore/Public/HairAttributes.h
@@ -56,6 +56,7 @@ enum class EHairAttribute : uint8
 	Color,
 	Roughness,
 	AO,
+	Width,
 	Count
 };
 HAIRSTRANDSCORE_API bool HasHairAttribute(uint32 In, EHairAttribute InAttribute);
diff --git a/Engine/Plugins/Runtime/HairStrands/Source/HairStrandsEditor/Private/GroomImportOptionsWindow.cpp b/Engine/Plugins/Runtime/HairStrands/Source/HairStrandsEditor/Private/GroomImportOptionsWindow.cpp
--- a/Engine/Plugins/Runtime/HairStrands/Source/HairStrandsEditor/Private/GroomImportOptionsWindow.cpp
+++ b/Engine/Plugins/Runtime/HairStrands/Source/HairStrandsEditor/Private/GroomImportOptionsWindow.cpp
@@ -250,10 +250,11 @@ static void AddGroupRemappingNoFoundMessage(SVerticalBox::FScopedWidgetSlotArgum
 FText GetHairAttributeLocText(EHairAttribute In, uint32 InFlags)
 {
 	// If a new optional attribute is added, please add its UI/text description here
-	static_assert(uint32(EHairAttribute::Count) == 7);
+	static_assert(uint32(EHairAttribute::Count) == 8);
 
 	switch (In)
 	{
+	case EHairAttribute::Width:						return LOCTEXT("GroomOptionsWindow_HasWidth", "Width");
 	case EHairAttribute::RootUV:					return HasHairAttributeFlags(InFlags, EHairAttributeFlags::HasRootUDIM) ? LOCTEXT("GroomOptionsWindow_HasRootUDIM", "Root UV (UDIM)") : LOCTEXT("GroomOptionsWindow_HasRootUV", "Root UV");
 	case EHairAttribute::ClumpID:					return HasHairAttributeFlags(InFlags, EHairAttributeFlags::HasMultipleClumpIDs) ? LOCTEXT("GroomOptionsWindow_HasClumpIDs", "Clump IDs (3)") : LOCTEXT("GroomOptionsWindow_HasClumpID", "Clump ID");
 	case EHairAttribute::StrandID:					return LOCTEXT("GroomOptionsWindow_HasStrandID", "Strand ID");
