Groom dataflow improvements

#virtualized

diff --git a/Engine/Plugins/Runtime/HairStrands/Source/HairStrandsCore/Private/GroomComponent.cpp b/Engine/Plugins/Runtime/HairStrands/Source/HairStrandsCore/Private/GroomComponent.cpp
--- a/Engine/Plugins/Runtime/HairStrands/Source/HairStrandsCore/Private/GroomComponent.cpp
+++ b/Engine/Plugins/Runtime/HairStrands/Source/HairStrandsCore/Private/GroomComponent.cpp
@@ -1552,7 +1552,7 @@ void UGroomComponent::CreateHairSimulation(const int32 GroupIndex, const int32 L
 					(GroomAsset->GetHairGroupsPhysics()[GroupIndex].SolverSettings.NiagaraSolver == EGroomNiagaraSolvers::AngularSprings) ? ToRawPtr(AngularSpringsSystem) :
 					(GroomAsset->GetHairGroupsPhysics()[GroupIndex].SolverSettings.NiagaraSolver == EGroomNiagaraSolvers::CosseratRods) ? ToRawPtr(CosseratRodsSystem) :
 					GroomAsset->GetHairGroupsPhysics()[GroupIndex].SolverSettings.CustomSystem.LoadSynchronous();
-				NiagaraComponent->SetVisibleFlag(SimulationSettings.SimulationSetup.bDebugSimulation || GroomAsset->GetHairGroupsPhysics()[GroupIndex].SolverSettings.bForceVisible);
+				NiagaraComponent->SetVisibility(SimulationSettings.SimulationSetup.bDebugSimulation || GroomAsset->GetHairGroupsPhysics()[GroupIndex].SolverSettings.bForceVisible, false);
 				NiagaraComponent->SetAsset(NiagaraAsset);
 				NiagaraComponent->ReinitializeSystem();
 			}
diff --git a/Engine/Plugins/Runtime/HairStrands/Source/HairStrandsCore/Private/GroomEdit.cpp b/Engine/Plugins/Runtime/HairStrands/Source/HairStrandsCore/Private/GroomEdit.cpp
--- a/Engine/Plugins/Runtime/HairStrands/Source/HairStrandsCore/Private/GroomEdit.cpp
+++ b/Engine/Plugins/Runtime/HairStrands/Source/HairStrandsCore/Private/GroomEdit.cpp
@@ -14,12 +14,14 @@
 static void HairStrandsDataToEditableHairGuides(const FHairStrandsDatas& In, TArray<FEditableHairGuide>& Out)
 {
 	const bool bHasValidStrandsIDs = In.StrandsCurves.StrandIDs.Num() > 0;
+	const bool bHasValidRootUVs = In.StrandsCurves.CurvesRootUV.Num() > 0;
+	
 	const uint32 StrandsCount = In.GetNumCurves();
 	Out.SetNum(StrandsCount);
 	for (uint32 StrandsIt = 0; StrandsIt < StrandsCount; ++StrandsIt)
 	{
 		Out[StrandsIt].GuideID= bHasValidStrandsIDs ? In.StrandsCurves.StrandIDs[StrandsIt] : StrandsIt;
-		Out[StrandsIt].RootUV = In.StrandsCurves.CurvesRootUV[StrandsIt];
+		Out[StrandsIt].RootUV = bHasValidRootUVs ? In.StrandsCurves.CurvesRootUV[StrandsIt] : FVector2f(0.0f,0.0f);
 
 		const uint32 PointCount = In.StrandsCurves.CurvesCount[StrandsIt];
 		const uint32 PointOffset = In.StrandsCurves.CurvesOffset[StrandsIt];
@@ -39,6 +41,8 @@ static void HairStrandsDataToEditableHairStrands(const FHairStrandsDatas& In, TA
 	const bool bHasValidStrandsIDs = In.StrandsCurves.StrandIDs.Num() > 0;
 	const bool bHasValidClumpIDs = In.StrandsCurves.ClumpIDs.Num() > 0;	
 	const bool bHasValidClosestGuide = In.StrandsCurves.CurvesClosestGuideIDs.Num() > 0;
+	const bool bHasValidRootUVs = In.StrandsCurves.CurvesRootUV.Num() > 0;
+	
 	const uint32 StrandsCount = In.GetNumCurves();
 	Out.SetNum(StrandsCount);
 	for (uint32 StrandsIt = 0; StrandsIt<StrandsCount; ++StrandsIt)
@@ -48,7 +52,7 @@ static void HairStrandsDataToEditableHairStrands(const FHairStrandsDatas& In, TA
 
 		Out[StrandsIt].StrandID		= bHasValidStrandsIDs ? In.StrandsCurves.StrandIDs[StrandsIt] : StrandsIt;
 		Out[StrandsIt].ClumpID		= bHasValidClumpIDs ? In.StrandsCurves.ClumpIDs[StrandsIt].X : 0;
-		Out[StrandsIt].RootUV		= In.StrandsCurves.CurvesRootUV[StrandsIt];
+		Out[StrandsIt].RootUV		= bHasValidRootUVs ? In.StrandsCurves.CurvesRootUV[StrandsIt] : FVector2f(0.0f,0.0f);
 		Out[StrandsIt].GuideIDs[0]  = GuideIDs.X;
 		Out[StrandsIt].GuideIDs[1]  = GuideIDs.Y;
 		Out[StrandsIt].GuideIDs[2]  = GuideIDs.Z;
@@ -162,7 +166,7 @@ static void EditableGroomGroupToHairDescription(FHairDescription& OutHairDescrip
 }
 #endif // WITH_EDITORONLY_DATA
 
-void ConvertFromGroomAsset(UGroomAsset* InGroom, FEditableGroom* OutGroom)
+void ConvertFromGroomAsset(UGroomAsset* InGroom, FEditableGroom* OutGroom, const bool bAllowCurveReordering)
 {
 #if WITH_EDITORONLY_DATA
 	check(InGroom);
@@ -180,7 +184,7 @@ void ConvertFromGroomAsset(UGroomAsset* InGroom, FEditableGroom* OutGroom)
 		FHairGroupInfo GroupInfo;
 		FHairStrandsDatas StrandsData;
 		FHairStrandsDatas GuidesData;
-		FGroomBuilder::BuildData(HairDescriptionGroups.HairGroups[GroupIndex], In.GetHairGroupsInterpolation()[GroupIndex], GroupInfo, StrandsData, GuidesData);
+		FGroomBuilder::BuildData(HairDescriptionGroups.HairGroups[GroupIndex], In.GetHairGroupsInterpolation()[GroupIndex], GroupInfo, StrandsData, GuidesData, bAllowCurveReordering);
 
 		Out.Groups[GroupIndex].GroupIndex = GroupIndex;
 		Out.Groups[GroupIndex].GroupID = GroupInfo.GroupID;
diff --git a/Engine/Plugins/Runtime/HairStrands/Source/HairStrandsCore/Private/HairStrandsInterpolation.cpp b/Engine/Plugins/Runtime/HairStrands/Source/HairStrandsCore/Private/HairStrandsInterpolation.cpp
--- a/Engine/Plugins/Runtime/HairStrands/Source/HairStrandsCore/Private/HairStrandsInterpolation.cpp
+++ b/Engine/Plugins/Runtime/HairStrands/Source/HairStrandsCore/Private/HairStrandsInterpolation.cpp
@@ -61,6 +61,9 @@ static FAutoConsoleVariableRef CVarHairStrands_Interpolation_ForceSingleGuide(TE
 static int32 GVarHairStrands_ScaleHairRadiusByTransform = 1;
 static FAutoConsoleVariableRef CVarHairStrands_ScaleHairRadiusByTransform(TEXT("r.HairStrands.Strands.ScaleHairRadiusByTransform"), GVarHairStrands_ScaleHairRadiusByTransform, TEXT("Scale hair strands radius by the instance's transform to make hair radius consistent with the instance scale."));
 
+static int32 GVarHairStrands_AddDeformerStrandsDisplacement = 0;
+static FAutoConsoleVariableRef CVarHairStrands_AddDeformerStrandsDisplacement(TEXT("r.HairStrands.Strands.AddDeformerStrandsDisplacement"), GVarHairStrands_AddDeformerStrandsDisplacement, TEXT("Add the deformer displacement to the current position. If set to 0 the deformer position will overwrite the deformed one"));
+
 ////////////////////////////////////////////////////////////////////////////////////////////////////////////////////////
 
 uint32 GetHairStrandsForceSingleGuideInterpolation()
@@ -490,7 +493,7 @@ class FHairInterpolationCS : public FGlobalShader
 	class FDynamicGeometry : SHADER_PERMUTATION_INT("PERMUTATION_DYNAMIC_GEOMETRY", 4); /* These is actually 5 DynamicGeometry permutations, but PERMUTATION_DYNAMIC_GEOMETRY==4 is never used (experimental) */
 	class FSimulation : SHADER_PERMUTATION_BOOL("PERMUTATION_SIMULATION");
 	class FSingleGuide : SHADER_PERMUTATION_BOOL("PERMUTATION_USE_SINGLE_GUIDE");
-	class FDeformer : SHADER_PERMUTATION_BOOL("PERMUTATION_DEFORMER");
+	class FDeformer : SHADER_PERMUTATION_INT("PERMUTATION_DEFORMER", 3);
 	class FWaveOps : SHADER_PERMUTATION_BOOL("PERMUTATION_WAVEOPS");
 	using FPermutationDomain = TShaderPermutationDomain<FDynamicGeometry, FSimulation, FSingleGuide, FDeformer, FPointPerCurve, FWaveOps>;
 
@@ -649,7 +652,7 @@ void AddHairStrandsInterpolationPass(
 	{
 		ShaderPrint::SetParameters(GraphBuilder, *ShaderPrintData, Parameters->ShaderDrawParameters);
 	}
-
+	
 	const bool bSupportDynamicMesh = 
 		RenRestRootResources &&
 		RenRestRootResources->GetRootCount() > 0 && 
@@ -658,7 +661,7 @@ void AddHairStrandsInterpolationPass(
 		RenDeformedRootResources->IsValid(MeshLODIndex);
 
 	// Guides
-	// Guide interpolation with dynamic mesh geometry is now support withguide cach. It implicitly assumes that the guide-cache follows correctly the binding surface
+	// Guide interpolation with dynamic mesh geometry is now support with guide cache. It implicitly assumes that the guide-cache follows correctly the binding surface
 	bool bSupportGlobalInterpolation = false;
 	if (bSupportDynamicMesh && (Instance->Guides.bIsSimulationEnable || Instance->Guides.bHasGlobalInterpolation || Instance->Guides.bIsDeformationEnable || Instance->Guides.bIsSimulationCacheEnable))
 	{
@@ -722,11 +725,14 @@ void AddHairStrandsInterpolationPass(
 	const bool bSingleGuidePermutation = bUseSingleGuide || GetHairStrandsForceSingleGuideInterpolation();
 	const bool bWaveOps = GRHISupportsWaveOperations && GRHIMinimumWaveSize <= 32 && GRHIMaximumWaveSize >= 32 && FHairInterpolationCS::DoesSupportsWaveOps(InPlatform, DispatchInfo.PointPerCurve) != ERHIFeatureSupport::Unsupported;
 
+	// Deformer type : 0 no deformer, 1 add deformer displacement, 2 overwrite deformed position with the deformer one
+	const int32 DeformerType = !bSupportDeformer ? 0 : GVarHairStrands_AddDeformerStrandsDisplacement ? 1 : 2;
+		
 	FHairInterpolationCS::FPermutationDomain PermutationVector;
 	PermutationVector.Set<FHairInterpolationCS::FDynamicGeometry>(DynamicGeometryType);
 	PermutationVector.Set<FHairInterpolationCS::FSimulation>(bHasLocalDeformation);
 	PermutationVector.Set<FHairInterpolationCS::FSingleGuide>(bSingleGuidePermutation);
-	PermutationVector.Set<FHairInterpolationCS::FDeformer>(bSupportDeformer);
+	PermutationVector.Set<FHairInterpolationCS::FDeformer>(DeformerType);
 	PermutationVector.Set<FHairInterpolationCS::FPointPerCurve>(DispatchInfo.PointPerCurve);
 	PermutationVector.Set<FHairInterpolationCS::FWaveOps>(bWaveOps);
 	TShaderMapRef<FHairInterpolationCS> ComputeShader(ShaderMap, PermutationVector);
diff --git a/Engine/Plugins/Runtime/HairStrands/Source/HairStrandsCore/Private/HairStrandsMeshProjection.cpp b/Engine/Plugins/Runtime/HairStrands/Source/HairStrandsCore/Private/HairStrandsMeshProjection.cpp
--- a/Engine/Plugins/Runtime/HairStrands/Source/HairStrandsCore/Private/HairStrandsMeshProjection.cpp
+++ b/Engine/Plugins/Runtime/HairStrands/Source/HairStrandsCore/Private/HairStrandsMeshProjection.cpp
@@ -629,7 +629,7 @@ void AddHairStrandInitMeshSamplesPass(
 	TShaderMapRef<FHairInitMeshSamplesCS> ComputeShader(ShaderMap, PermutationVector);
 	FComputeShaderUtils::AddPass(
 		GraphBuilder,
-		RDG_EVENT_NAME("HairStrands::TriangleMeshUpdate(%s)", bUseRDGPositionBuffer ? TEXT("RDGBuffer") : TEXT("SkinCacheBuffer")),
+		RDG_EVENT_NAME("HairStrands::InitMeshSamples"),
 		ComputeShader,
 		PassParameters,
 		DispatchGroupCount);
diff --git a/Engine/Plugins/Runtime/HairStrands/Source/HairStrandsCore/Public/GroomEdit.h b/Engine/Plugins/Runtime/HairStrands/Source/HairStrandsCore/Public/GroomEdit.h
--- a/Engine/Plugins/Runtime/HairStrands/Source/HairStrandsCore/Public/GroomEdit.h
+++ b/Engine/Plugins/Runtime/HairStrands/Source/HairStrandsCore/Public/GroomEdit.h
@@ -80,7 +80,7 @@ enum EEditableGroomOperations
 };
 
 // Convert a groom asset into an editable groom asset
-HAIRSTRANDSCORE_API void ConvertFromGroomAsset(UGroomAsset* In, FEditableGroom* Out);
+HAIRSTRANDSCORE_API void ConvertFromGroomAsset(UGroomAsset* In, FEditableGroom* Out, const bool bAllowCurveReordering = true);
 
 // Convert an editable groom asset into a groom asset
 // The 'operations' flag indicates what type of modifications have been done onto the editable groom. 
diff --git a/Engine/Plugins/Runtime/HairStrands/Source/HairStrandsDataflow/Private/GenerateGuidesCurvesNode.cpp b/Engine/Plugins/Runtime/HairStrands/Source/HairStrandsDataflow/Private/GenerateGuidesCurvesNode.cpp
--- a/Engine/Plugins/Runtime/HairStrands/Source/HairStrandsDataflow/Private/GenerateGuidesCurvesNode.cpp
+++ b/Engine/Plugins/Runtime/HairStrands/Source/HairStrandsDataflow/Private/GenerateGuidesCurvesNode.cpp
@@ -14,8 +14,9 @@ namespace UE::Groom::Private
 		const uint32 ObjectIndex, const uint32 PrevStrands, const uint32 NextStrands, TArray<uint32>& SampleIndices)
 	{
 		const uint32 NumStrands = NextStrands-PrevStrands;
-		const uint32 NumGuides = (ObjectIndex == StrandsFacade.GetNumObjects()-1) ? GuidesCount - GuidesOffset :
+		uint32 NumGuides = (ObjectIndex == StrandsFacade.GetNumObjects()-1) ? GuidesCount - GuidesOffset :
 			GuidesCount * static_cast<float>(NumStrands)/StrandsFacade.GetNumCurves();
+		NumGuides = FMath::Max(1u, NumGuides);
 
 		TArray<FVector3f> RootPositions;
 		RootPositions.Init(FVector3f::ZeroVector, NumStrands);
@@ -65,8 +66,9 @@ void FGenerateGuidesCurvesDataflowNode::Evaluate(UE::Dataflow::FContext& Context
 		{
 			TArray<FVector3f> PointRestPositions;
 			TArray<int32> ObjectCurveOffsets, CurvePointOffsets, CurveStrandIndices;
+			TArray<FString> ObjectGroupNames;
 			
-			uint32 PrevOffset = 0, GuidesOffset = 0;;
+			uint32 PrevOffset = 0, GuidesOffset = 0;
 			for (int32 ObjectIndex = 0; ObjectIndex < StrandsFacade.GetNumObjects(); ++ObjectIndex)
 			{
 				const uint32 NextOffset = StrandsFacade.GetObjectCurveOffsets()[ObjectIndex];
@@ -78,7 +80,12 @@ void FGenerateGuidesCurvesDataflowNode::Evaluate(UE::Dataflow::FContext& Context
 
 				// Build the sampled guides
 				UE::Groom::Private::BuildGuidesCurves(StrandsFacade, PrevOffset, SampleIndices, PointRestPositions, ObjectCurveOffsets, CurvePointOffsets, CurveStrandIndices);
-
+				
+				const FString GroupName = StrandsFacade.GetObjectGroupNames()[ObjectIndex].Replace(
+					*UE::Groom::FGroomStrandsFacade::GroupPrefix.ToString(), *UE::Groom::FGroomGuidesFacade::GroupPrefix.ToString());
+				
+				ObjectGroupNames.Add(GroupName);
+				
 				GuidesOffset += NumGuides;
 				PrevOffset = NextOffset;
 			}
@@ -87,7 +94,7 @@ void FGenerateGuidesCurvesDataflowNode::Evaluate(UE::Dataflow::FContext& Context
 			const TArray<int32> ObjectPointSamples = GuidesFacade.GetObjectPointSamples();
 
 			// Init the groom collection
-			GuidesFacade.InitGroomCollection(PointRestPositions, CurvePointOffsets, ObjectCurveOffsets);
+			GuidesFacade.InitGroomCollection(PointRestPositions, CurvePointOffsets, ObjectCurveOffsets, ObjectGroupNames);
 
 			// Set the curve strand indices
 			GuidesFacade.SetCurveStrandIndices(CurveStrandIndices);
diff --git a/Engine/Plugins/Runtime/HairStrands/Source/HairStrandsDataflow/Private/GetGroomAssetNode.cpp b/Engine/Plugins/Runtime/HairStrands/Source/HairStrandsDataflow/Private/GetGroomAssetNode.cpp
--- a/Engine/Plugins/Runtime/HairStrands/Source/HairStrandsDataflow/Private/GetGroomAssetNode.cpp
+++ b/Engine/Plugins/Runtime/HairStrands/Source/HairStrandsDataflow/Private/GetGroomAssetNode.cpp
@@ -25,12 +25,14 @@ namespace UE::Groom::Private
 
 		if (GroomAsset)
 		{
-			ConvertFromGroomAsset(const_cast<UGroomAsset*>(GroomAsset), &EditGroom);
+			ConvertFromGroomAsset(const_cast<UGroomAsset*>(GroomAsset), &EditGroom, false);
 
 			TArray<FVector3f> PointRestPositions;
 			TArray<int32> ObjectCurveOffsets;
 			TArray<int32> CurvePointOffsets;
+			TArray<FString> ObjectGroupNames;
 
+			uint32 GroupIndex = 0;
 			for(FEditableGroomGroup& GroomGroup : EditGroom.Groups)
 			{
 				for( const typename FacadeType::FEditableType& EditType : FacadeType::GetEditableGroom(GroomGroup))
@@ -42,9 +44,15 @@ namespace UE::Groom::Private
 					CurvePointOffsets.Add(PointRestPositions.Num());
 				}
 				ObjectCurveOffsets.Add(CurvePointOffsets.Num());
+
+				const FString GroupName = GroomAsset->GetName() + TEXT("_") + FacadeType::GroupPrefix.ToString() + TEXT("_") + FString::FromInt(GroupIndex);
+				ObjectGroupNames.Add(GroupName);
+				++GroupIndex;
 			}
 			FacadeType CurvesFacade(GroomCollection);
-			CurvesFacade.InitGroomCollection(PointRestPositions, CurvePointOffsets, ObjectCurveOffsets);
+			CurvesFacade.InitGroomCollection(PointRestPositions, CurvePointOffsets, ObjectCurveOffsets, ObjectGroupNames);
+
+			const FString GroomName = GroomAsset->GetName();
 
 			return GroomAsset;
 		}
diff --git a/Engine/Plugins/Runtime/HairStrands/Source/HairStrandsDataflow/Private/GroomAssetTerminalNode.cpp b/Engine/Plugins/Runtime/HairStrands/Source/HairStrandsDataflow/Private/GroomAssetTerminalNode.cpp
--- a/Engine/Plugins/Runtime/HairStrands/Source/HairStrandsDataflow/Private/GroomAssetTerminalNode.cpp
+++ b/Engine/Plugins/Runtime/HairStrands/Source/HairStrandsDataflow/Private/GroomAssetTerminalNode.cpp
@@ -132,7 +132,7 @@ void FGroomAssetTerminalDataflowNode::SetAssetValue(TObjectPtr<UObject> Asset, U
 
 			// Build an editable groom asset for the strands
 			FEditableGroom EditGroom;
-			ConvertFromGroomAsset(const_cast<UGroomAsset*>(GroomAsset), &EditGroom);
+			ConvertFromGroomAsset(const_cast<UGroomAsset*>(GroomAsset), &EditGroom, false);
 
 			if(GuidesFacade.GetNumObjects() == EditGroom.Groups.Num())
 			{
diff --git a/Engine/Plugins/Runtime/HairStrands/Source/HairStrandsDataflow/Private/GroomCollectionFacades.cpp b/Engine/Plugins/Runtime/HairStrands/Source/HairStrandsDataflow/Private/GroomCollectionFacades.cpp
--- a/Engine/Plugins/Runtime/HairStrands/Source/HairStrandsDataflow/Private/GroomCollectionFacades.cpp
+++ b/Engine/Plugins/Runtime/HairStrands/Source/HairStrandsDataflow/Private/GroomCollectionFacades.cpp
@@ -32,6 +32,8 @@ namespace UE::Groom
 	const FName FGroomCollectionFacade<DerivedType>::PointRestPositionsAttribute("PointRestPositions");
 	template<typename DerivedType>
 	const FName FGroomCollectionFacade<DerivedType>::VertexLinearColorsAttribute("VertexLinearColors");
+	template<typename DerivedType>
+	const FName FGroomCollectionFacade<DerivedType>::ObjectGroupNamesAttribute("ObjectGroupNames");
 
 	// Strands/Guides Prefix
 	const FName FGroomStrandsFacade::GroupPrefix("Strands");
@@ -57,7 +59,8 @@ namespace UE::Groom
 		ObjectCurveOffsets(InCollection, ObjectCurveOffsetsAttribute, ObjectsGroup),
 		PointCurveIndices(InCollection, PointCurveIndicesAttribute, PointsGroup),
 		CurveObjectIndices(InCollection, CurveObjectIndicesAttribute, CurvesGroup),
-		VertexLinearColors(InCollection, VertexLinearColorsAttribute, VerticesGroup)
+		VertexLinearColors(InCollection, VertexLinearColorsAttribute, VerticesGroup),
+		ObjectGroupNames(InCollection, ObjectGroupNamesAttribute, ObjectsGroup)
 	{
 	}
 
@@ -70,7 +73,8 @@ namespace UE::Groom
 		ObjectCurveOffsets(InCollection, ObjectCurveOffsetsAttribute, ObjectsGroup),
 		PointCurveIndices(InCollection, PointCurveIndicesAttribute, PointsGroup),
 		CurveObjectIndices(InCollection, CurveObjectIndicesAttribute, CurvesGroup),
-		VertexLinearColors(InCollection, VertexLinearColorsAttribute, VerticesGroup)
+		VertexLinearColors(InCollection, VertexLinearColorsAttribute, VerticesGroup),
+		ObjectGroupNames(InCollection, ObjectGroupNamesAttribute, ObjectsGroup)
 	{
 	}
 
@@ -79,7 +83,7 @@ namespace UE::Groom
 	{
 		return EdgeRestOrientations.IsValid() && PointRestPositions.IsValid() && 
 			CurvePointOffsets.IsValid() && ObjectCurveOffsets.IsValid() && PointCurveIndices.IsValid() &&
-					CurveObjectIndices.IsValid() && VertexLinearColors.IsValid() && static_cast<const DerivedType*>(this)->IsFacadeValid();
+					CurveObjectIndices.IsValid() && VertexLinearColors.IsValid() && ObjectGroupNames.IsValid() && static_cast<const DerivedType*>(this)->IsFacadeValid();
 	}
 
 	template<typename DerivedType>
@@ -93,13 +97,14 @@ namespace UE::Groom
 		PointCurveIndices.Add();
 		CurveObjectIndices.Add();
 		VertexLinearColors.Add();
+		ObjectGroupNames.Add();
 		
 		static_cast<DerivedType*>(this)->DefineFacadeSchema();
 	}
 
 	template<typename DerivedType>
 	void FGroomCollectionFacade<DerivedType>::InitGroomCollection(const TArray<FVector3f>& InPointRestPositions,
-		const TArray<int32>& InCurvePointOffsets, const TArray<int32>& InObjectCurveOffsets)
+		const TArray<int32>& InCurvePointOffsets, const TArray<int32>& InObjectCurveOffsets, const TArray<FString>& InObjectGroupNames)
 	{
 		if(Collection)
 		{
@@ -127,6 +132,7 @@ namespace UE::Groom
 			SetObjectCurveOffsets(InObjectCurveOffsets);
 			SetCurvePointOffsets(InCurvePointOffsets);
 			SetPointRestPositions(InPointRestPositions);
+			SetObjectGroupNames(InObjectGroupNames);
 
 			static_cast<DerivedType*>(this)->InitFacadeCollection();
 		}
diff --git a/Engine/Plugins/Runtime/HairStrands/Source/HairStrandsDataflow/Private/GroomDataflowRendering.cpp b/Engine/Plugins/Runtime/HairStrands/Source/HairStrandsDataflow/Private/GroomDataflowRendering.cpp
--- a/Engine/Plugins/Runtime/HairStrands/Source/HairStrandsDataflow/Private/GroomDataflowRendering.cpp
+++ b/Engine/Plugins/Runtime/HairStrands/Source/HairStrandsDataflow/Private/GroomDataflowRendering.cpp
@@ -28,6 +28,13 @@ namespace UE::Groom
 		 TEXT("Guides Lod we want to display")
 	 );
 
+	int32 GGroomDataflowStrandsLod = 0;
+	static FAutoConsoleVariableRef CVarGroomDataflowStrandsLod(
+		 TEXT("p.Groom.Dataflow.StrandsLod"),
+		 GGroomDataflowStrandsLod,
+		 TEXT("Strands Lod we want to display")
+	 );
+
 	float GGroomDataflowRenderingThickness = 0.5f;
 	static FAutoConsoleVariableRef CVarRestGuidesRenderingThickness(
 		 TEXT("p.Groom.Dataflow.RenderingThickness"),
@@ -200,8 +207,9 @@ namespace UE::Groom
 				TArray<FLinearColor> Colors;
 				Colors.Init(FLinearColor(0.2,0.6,1.0), VertexPositions.Num());
 				ColorLambda(GroomFacade, Colors);
- 
-				const int32 GeometryIndex = RenderCollection.StartGeometryGroup(State.GetGuid().ToString());
+
+				const FString GeometryName = GroomFacade.GetObjectGroupNames().Last();
+				const int32 GeometryIndex = RenderCollection.StartGeometryGroup(GeometryName);
 				RenderCollection.AddSurface(MoveTemp(VertexPositions), MoveTemp(FaceIndices), MoveTemp(VertexNormals), MoveTemp(Colors));
 				RenderCollection.EndGeometryGroup(GeometryIndex);
 			}
@@ -288,26 +296,38 @@ namespace UE::Groom
 	void FGroomGuidesRenderingCallbacks::Render(GeometryCollection::Facades::FRenderingFacade& RenderCollection, const Dataflow::FGraphRenderingState& State)
 	{
 		RenderGroomCollection<FGroomGuidesFacade>(RenderCollection, State,
-			[](const FGroomGuidesFacade& GuidesFacade, TArray<FLinearColor>& VertexColors)
+			[this](const FGroomGuidesFacade& GuidesFacade, TArray<FLinearColor>& VertexColors)
 			{
-				if(GGroomDataflowDebugDraw == 1)
-				{
-					RenderBoneWeights(GuidesFacade, VertexColors);
-				}
-				if(GGroomDataflowDebugDraw == 2)
-				{
-					RenderGuidesLods(GuidesFacade, VertexColors);
-				}
+				ComputeVertexColors(GuidesFacade, VertexColors);
 			});
 	}
 
+	void FGroomGuidesRenderingCallbacks::ComputeVertexColors(const FGroomGuidesFacade& GuidesFacade, TArray<FLinearColor>& VertexColors) const
+	{
+		if(GGroomDataflowDebugDraw == 1)
+		{
+			RenderBoneWeights(GuidesFacade, VertexColors);
+		}
+		if(GGroomDataflowDebugDraw == 2)
+		{
+			RenderGuidesLods(GuidesFacade, VertexColors);
+		}
+	}
+
 	Dataflow::FRenderKey FGroomStrandsRenderingCallbacks::RenderKey = { "StrandsRender", FName("FGroomCollection") };
 	
 	void FGroomStrandsRenderingCallbacks::Render(GeometryCollection::Facades::FRenderingFacade& RenderCollection, const Dataflow::FGraphRenderingState& State)
 	{
-		RenderGroomCollection<FGroomStrandsFacade>(RenderCollection, State, [](const FGroomStrandsFacade& StrandsFacade, TArray<FLinearColor>& VertexColors){});
+		RenderGroomCollection<FGroomStrandsFacade>(RenderCollection, State,
+			[this](const FGroomStrandsFacade& StrandsFacade, TArray<FLinearColor>& VertexColors)
+			{
+				ComputeVertexColors(StrandsFacade, VertexColors);
+			});
 	}
 
+	void FGroomStrandsRenderingCallbacks::ComputeVertexColors(const FGroomStrandsFacade& StrandsFacade, TArray<FLinearColor>& VertexColors) const
+	{}
+
 	void RegisterRenderingCallbacks()
 	{
 		Dataflow::FRenderingFactory::GetInstance()->RegisterCallbacks(MakeUnique<FGroomGuidesRenderingCallbacks>());
diff --git a/Engine/Plugins/Runtime/HairStrands/Source/HairStrandsDataflow/Public/GroomCollectionFacades.h b/Engine/Plugins/Runtime/HairStrands/Source/HairStrandsDataflow/Public/GroomCollectionFacades.h
--- a/Engine/Plugins/Runtime/HairStrands/Source/HairStrandsDataflow/Public/GroomCollectionFacades.h
+++ b/Engine/Plugins/Runtime/HairStrands/Source/HairStrandsDataflow/Public/GroomCollectionFacades.h
@@ -51,6 +51,7 @@ namespace UE::Groom
 		static const FName PointCurveIndicesAttribute;
 		static const FName CurveObjectIndicesAttribute;
 		static const FName VertexLinearColorsAttribute;
+		static const FName ObjectGroupNamesAttribute;
 
 		FGroomCollectionFacade(FManagedArrayCollection& InCollection);
 		FGroomCollectionFacade(const FManagedArrayCollection& InCollection);
@@ -99,9 +100,15 @@ namespace UE::Groom
 
 		/** Get the curve object indices */
 		const TArray<FLinearColor>& GetVertexLinearColors() const { return VertexLinearColors.Get().GetConstArray(); }
+
+		/** Get the object group names */
+		const TArray<FString>& GetObjectGroupNames() const { return ObjectGroupNames.Get().GetConstArray(); }
 		
 		/** Set the point rest positions */
 		void SetPointRestPositions(const TArray<FVector3f>& InPointRestPositions) { PointRestPositions.Modify() = InPointRestPositions; UpdateEdgeRestOrientations(); }	
+
+		/** Set the object group names */
+		void SetObjectGroupNames(const TArray<FString>& InObjectGroupNames) { ObjectGroupNames.Modify() = InObjectGroupNames; }	
 		
 		/** Set the curve point offsets */
 		void SetCurvePointOffsets(const TArray<int32>& InCurvePointOffsets) { CurvePointOffsets.Modify() = InCurvePointOffsets; UpdatePointCurveIndices(); }
@@ -114,7 +121,7 @@ namespace UE::Groom
 		
 		/** Initialize the curve vertices from point positions and objects offsets */
 		void InitGroomCollection(const TArray<FVector3f>& InPointRestPositions,
-			const TArray<int32>& InCurvePointOffsets, const TArray<int32>& InObjectCurveOffsets);
+			const TArray<int32>& InCurvePointOffsets, const TArray<int32>& InObjectCurveOffsets, const TArray<FString>& InObjectGroupNames);
 
 		/** Get the managed array collection */
 		const FManagedArrayCollection& GetManagedArrayCollection() const {return ConstCollection;}
@@ -154,8 +161,11 @@ namespace UE::Groom
 		/** Groom curves object index */
 		TManagedArrayAccessor<int32> CurveObjectIndices;
 
-		/** Groom vertices linear color*/
+		/** Groom vertices linear color */
 		TManagedArrayAccessor<FLinearColor> VertexLinearColors;
+
+		/** Groom object group names */
+		TManagedArrayAccessor<FString> ObjectGroupNames;
 	};
 	
 	/**
diff --git a/Engine/Plugins/Runtime/HairStrands/Source/HairStrandsDataflow/Public/GroomDataflowRendering.h b/Engine/Plugins/Runtime/HairStrands/Source/HairStrandsDataflow/Public/GroomDataflowRendering.h
--- a/Engine/Plugins/Runtime/HairStrands/Source/HairStrandsDataflow/Public/GroomDataflowRendering.h
+++ b/Engine/Plugins/Runtime/HairStrands/Source/HairStrandsDataflow/Public/GroomDataflowRendering.h
@@ -7,13 +7,17 @@
 
 namespace UE::Groom
 {
+
+class FGroomStrandsFacade;
+class FGroomGuidesFacade;
+
 /** Strands rendering callback for the dataflow editor */
-class FGroomStrandsRenderingCallbacks : public Dataflow::FRenderingFactory::ICallbackInterface
+class HAIRSTRANDSDATAFLOW_API FGroomStrandsRenderingCallbacks : public Dataflow::FRenderingFactory::ICallbackInterface
 {
 public :
 	static Dataflow::FRenderKey RenderKey;
 	
-private:
+protected:
 
 	//¬ Begin ICallbackInterface interface
 	virtual Dataflow::FRenderKey GetRenderKey() const override
@@ -28,15 +32,17 @@ private:
 
 	virtual void Render(GeometryCollection::Facades::FRenderingFacade& RenderCollection, const Dataflow::FGraphRenderingState& State) override;
 	//¬ End ICallbackInterface interface
+	
+	virtual void ComputeVertexColors(const FGroomStrandsFacade& StrandsFacade, TArray<FLinearColor>& VertexColors) const;
 };
 
 /** Guides rendering callback for the dataflow editor */
-class FGroomGuidesRenderingCallbacks : public Dataflow::FRenderingFactory::ICallbackInterface
+class HAIRSTRANDSDATAFLOW_API FGroomGuidesRenderingCallbacks : public Dataflow::FRenderingFactory::ICallbackInterface
 {
 	public :
 		static Dataflow::FRenderKey RenderKey;
 	
-private:
+protected:
 
 	//¬ Begin ICallbackInterface interface
 	virtual Dataflow::FRenderKey GetRenderKey() const override
@@ -50,6 +56,9 @@ private:
 	}
 
 	virtual void Render(GeometryCollection::Facades::FRenderingFacade& RenderCollection, const Dataflow::FGraphRenderingState& State) override;
+	
+	virtual void ComputeVertexColors(const FGroomGuidesFacade& GuidesFacade, TArray<FLinearColor>& VertexColors) const;
+    	
 	//¬ End ICallbackInterface interface
 };
 
diff --git a/Engine/Plugins/Runtime/HairStrands/Source/HairStrandsDeformer/Private/DeformerGroomDomainsSource.cpp b/Engine/Plugins/Runtime/HairStrands/Source/HairStrandsDeformer/Private/DeformerGroomDomainsSource.cpp
--- a/Engine/Plugins/Runtime/HairStrands/Source/HairStrandsDeformer/Private/DeformerGroomDomainsSource.cpp
+++ b/Engine/Plugins/Runtime/HairStrands/Source/HairStrandsDeformer/Private/DeformerGroomDomainsSource.cpp
@@ -165,6 +165,7 @@ bool UOptimusGroomSolverComponentSource::GetComponentElementCountsForExecutionDo
 		{
 			TotalCount = GroomSolver->GetSolverSettings().PointKinematicIndices.Num();
 		}
+		TotalCount = FMath::Max(TotalCount, 1u);
 		InvocationElementCounts.Init(TotalCount,1);
 		return TotalCount != 0;
 	}
diff --git a/Engine/Plugins/Runtime/HairStrands/Source/HairStrandsEditor/Private/GroomActions.cpp b/Engine/Plugins/Runtime/HairStrands/Source/HairStrandsEditor/Private/GroomActions.cpp
--- a/Engine/Plugins/Runtime/HairStrands/Source/HairStrandsEditor/Private/GroomActions.cpp
+++ b/Engine/Plugins/Runtime/HairStrands/Source/HairStrandsEditor/Private/GroomActions.cpp
@@ -33,6 +33,7 @@
 #include "ObjectTools.h"
 #include "Widgets/Notifications/SNotificationList.h"
 #include "Framework/Notifications/NotificationManager.h"
+#include "Materials/Material.h"
 
 #include "AssetToolsModule.h"
 #include "ContentBrowserMenuContexts.h"
@@ -62,8 +63,14 @@ EAssetCommandResult UAssetDefinition_GroomAsset::OpenAssets(const FAssetOpenArgs
 			{
 				UAssetEditorSubsystem* const AssetEditorSubsystem = GEditor->GetEditorSubsystem<UAssetEditorSubsystem>();
 				UDataflowEditor* const AssetEditor = NewObject<UDataflowEditor>(AssetEditorSubsystem, NAME_None, RF_Transient);
+				
 				const TSubclassOf<AActor> ActorClass = StaticLoadClass(AActor::StaticClass(), nullptr,
 					TEXT("/HairStrands/BP_PreviewGroom.BP_PreviewGroom_C"), nullptr, LOAD_None, nullptr);
+
+				UMaterial* HairMaterial = Cast<UMaterial>(StaticLoadObject(UMaterial::StaticClass(),nullptr,
+					TEXT("/HairStrands/Materials/HairDataflowMaterial.HairDataflowMaterial"), nullptr, LOAD_None, nullptr));
+				
+				GroomAsset->GetDataflowSettings().GetDataflowAsset()->Material = HairMaterial;
 				AssetEditor->Initialize({ GroomAsset}, ActorClass);
 			}
 		}
diff --git a/Engine/Plugins/Runtime/HairStrands/Source/HairStrandsSolver/Private/GroomSolverComponent.cpp b/Engine/Plugins/Runtime/HairStrands/Source/HairStrandsSolver/Private/GroomSolverComponent.cpp
--- a/Engine/Plugins/Runtime/HairStrands/Source/HairStrandsSolver/Private/GroomSolverComponent.cpp
+++ b/Engine/Plugins/Runtime/HairStrands/Source/HairStrandsSolver/Private/GroomSolverComponent.cpp
@@ -107,8 +107,15 @@ void UGroomSolverComponent::SelectDynamicCurves()
 				}
 				MinDistance = FMath::Sqrt(MinDistance);
 
-				const FVector::FReal DistanceRatio = FMath::Clamp(MinDistance / GetSolverSettings().MaxLODDistance, 0.0f, 1.0f);
+				const bool bHasValidBounds = GetSolverSettings().MaxLODDistance >= GetSolverSettings().MinLODDistance;
+				const FVector::FReal DistanceRatio = bHasValidBounds ?
+					FMath::Clamp((MinDistance - GetSolverSettings().MinLODDistance) / (GetSolverSettings().MaxLODDistance - GetSolverSettings().MinLODDistance), 0.0f, 1.0f) : 1.0f;
 
+				if(!bHasValidBounds)
+				{
+					UE_LOG(LogTemp, Warning, TEXT("Groom solver max distance should be higher than the min distance"))
+				}
+				
 				const uint32 NumGroups = GroomComponent->GetGroupCount();
 				for (uint32 GroupIndex = 0; GroupIndex < NumGroups; ++GroupIndex)
 				{
diff --git a/Engine/Plugins/Runtime/HairStrands/Source/HairStrandsSolver/Public/GroomSolverComponent.h b/Engine/Plugins/Runtime/HairStrands/Source/HairStrandsSolver/Public/GroomSolverComponent.h
--- a/Engine/Plugins/Runtime/HairStrands/Source/HairStrandsSolver/Public/GroomSolverComponent.h
+++ b/Engine/Plugins/Runtime/HairStrands/Source/HairStrandsSolver/Public/GroomSolverComponent.h
@@ -30,8 +30,12 @@ struct FGroomSolverSettings
 	float SimulationFrameRate = 30;
 
 	/** Maximum LOD distance (if distance between the component and the views is higher that this threshold, no simulation)*/
-	UPROPERTY(EditAnywhere, Category = "LODs")
+	UPROPERTY(EditAnywhere, Category = "LODs", meta = (ClampMin = "0", UIMin = "0"))
 	float MaxLODDistance = 1000;
+	
+	/** Minimum LOD distance (if distance between the component and the views is lower that this threshold, no simulation)*/
+	UPROPERTY(EditAnywhere, Category = "LODs", meta = (ClampMin = "0", UIMin = "0"))
+	float MinLODDistance = 100;
 
 	/** List of dynamic curves */
 	TArray<int32> CurveDynamicIndices;
diff --git a/Engine/Shaders/Private/HairStrands/HairStrandsInterpolation.usf b/Engine/Shaders/Private/HairStrands/HairStrandsInterpolation.usf
--- a/Engine/Shaders/Private/HairStrands/HairStrandsInterpolation.usf
+++ b/Engine/Shaders/Private/HairStrands/HairStrandsInterpolation.usf
@@ -22,6 +22,12 @@
 
 #pragma warning (default:7203) // downgrade 'not enough registers available for the entire program.' to a warning instead of error
 
+
+// Deformer types
+#define DEFORMER_NONE 0
+#define DEFORMER_ADDITIVE 1
+#define DEFORMER_OVERWRITE 2
+
 ///////////////////////////////////////////////////////////////////////////////////////////////////
 // Guides
 
@@ -549,13 +555,20 @@ void MainCS(uint3 DispatchThreadId : SV_DispatchThreadID, uint ThreadIndex : SV_
 				#endif
 
 				// 3. Apply deformer
-				#if PERMUTATION_DEFORMER
+				#if PERMUTATION_DEFORMER == DEFORMER_ADDITIVE
 				{
 					const FPackedHairPosition PackedDeformerPosition = ReadPackedHairPosition(RenDeformerPositionBuffer, VertexIndex);
 					const float3 DeformerPosition = UnpackHairControlPointPosition(PackedDeformerPosition, InRenHairPositionOffset);
 					OutPosition = ApplyDeformer(RenTri, RestPosition, DeformerPosition, OutPosition);
 					CopyPackedAttribute(PackedRestPosition, PackedDeformerPosition);
 				}
+				#elif PERMUTATION_DEFORMER == DEFORMER_OVERWRITE
+				{
+					const FPackedHairPosition PackedDeformerPosition = ReadPackedHairPosition(RenDeformerPositionBuffer, VertexIndex);
+					OutPosition = UnpackHairControlPointPosition(PackedDeformerPosition, InRenHairPositionOffset);
+					CopyPackedAttribute(PackedRestPosition, PackedDeformerPosition);
+				}
+				// 3. Linear blend skinning (translation)
 				#endif
 
 				// 4. Write out the final position
