Fix groom deformer shaders

#virtualized

diff --git a/Engine/Plugins/Runtime/HairStrands/Shaders/Private/Deformers/DeformerGroomAttributeRead.ush b/Engine/Plugins/Runtime/HairStrands/Shaders/Private/Deformers/DeformerGroomAttributeRead.ush
--- a/Engine/Plugins/Runtime/HairStrands/Shaders/Private/Deformers/DeformerGroomAttributeRead.ush
+++ b/Engine/Plugins/Runtime/HairStrands/Shaders/Private/Deformers/DeformerGroomAttributeRead.ush
@@ -86,10 +86,10 @@ float4x4 ReadValueTypedFloat4x4_{DataInterfaceName}(uint ValueIndex)
 
 }
 
-float3x4 ReadValueTypedFloat4x4_{DataInterfaceName}(uint ValueIndex)
+float3x4 ReadValueTypedFloat3x4_{DataInterfaceName}(uint ValueIndex)
 {
 	const int ValueOffset = 12 * ValueIndex;
-	return (ValueIndex < {DataInterfaceName}_NumValues) ? float4x4({DataInterfaceName}_ValueTypedData[ValueOffset], {DataInterfaceName}_ValueTypedData[ValueOffset+1],
+	return (ValueIndex < {DataInterfaceName}_NumValues) ? float3x4({DataInterfaceName}_ValueTypedData[ValueOffset], {DataInterfaceName}_ValueTypedData[ValueOffset+1],
 						{DataInterfaceName}_ValueTypedData[ValueOffset+2], {DataInterfaceName}_ValueTypedData[ValueOffset+3],
 						{DataInterfaceName}_ValueTypedData[ValueOffset+4], {DataInterfaceName}_ValueTypedData[ValueOffset+5],
 						{DataInterfaceName}_ValueTypedData[ValueOffset+6], {DataInterfaceName}_ValueTypedData[ValueOffset+7],
diff --git a/Engine/Plugins/Runtime/HairStrands/Shaders/Private/Deformers/DeformerGroomSolverRead.ush b/Engine/Plugins/Runtime/HairStrands/Shaders/Private/Deformers/DeformerGroomSolverRead.ush
--- a/Engine/Plugins/Runtime/HairStrands/Shaders/Private/Deformers/DeformerGroomSolverRead.ush
+++ b/Engine/Plugins/Runtime/HairStrands/Shaders/Private/Deformers/DeformerGroomSolverRead.ush
@@ -113,7 +113,7 @@ int ReadKinematicPointIndex_{DataInterfaceName}(uint PointIndex)
 
 int ReadDynamicCurveIndex_{DataInterfaceName}(uint CurveIndex)
 {
-	return (PointIndex < {DataInterfaceName}_NumDynamicCurves) ? {DataInterfaceName}_DynamicCurveIndices[PointIndex] : -1;
+	return (CurveIndex < {DataInterfaceName}_NumDynamicCurves) ? {DataInterfaceName}_DynamicCurveIndices[CurveIndex] : -1;
 }
 
 int ReadKinematicCurveIndex_{DataInterfaceName}(uint CurveIndex)
