Fix InitResources() reentrance when calling InitResources() while the BindingAssets are async. building at the same time.

diff --git a/Engine/Plugins/Runtime/HairStrands/Source/HairStrandsCore/Private/GroomComponent.cpp b/Engine/Plugins/Runtime/HairStrands/Source/HairStrandsCore/Private/GroomComponent.cpp
--- a/Engine/Plugins/Runtime/HairStrands/Source/HairStrandsCore/Private/GroomComponent.cpp
+++ b/Engine/Plugins/Runtime/HairStrands/Source/HairStrandsCore/Private/GroomComponent.cpp
@@ -1927,21 +1927,7 @@ void UGroomComponent::SetBinding(UGroomBindingAsset* InBinding)
 
 void UGroomComponent::SetBindingAsset(UGroomBindingAsset* InBinding)
 {
-	if (BindingAsset != InBinding)
-	{
-		bool bIsValid = InBinding != nullptr;
-		if (bIsValid && bValidationEnable)
-		{
-			bIsValid = 
-				UGroomBindingAsset::IsBindingAssetValid(InBinding, false, bValidationEnable) &&
-				UGroomBindingAsset::IsCompatible(GroomAsset, InBinding, bValidationEnable);
-		}
-		if (bIsValid)
-		{
-			BindingAsset = InBinding;
-			InitResources();
-		}
-	}
+	SetGroomAsset(GroomAsset, InBinding);
 }
 
 void UGroomComponent::SetPhysicsAsset(UPhysicsAsset* InPhysicsAsset)
