Fix incorrect assert in groom builder.

diff --git a/Engine/Plugins/Runtime/HairStrands/Source/HairStrandsCore/Private/GroomBuilder.cpp b/Engine/Plugins/Runtime/HairStrands/Source/HairStrandsCore/Private/GroomBuilder.cpp
--- a/Engine/Plugins/Runtime/HairStrands/Source/HairStrandsCore/Private/GroomBuilder.cpp
+++ b/Engine/Plugins/Runtime/HairStrands/Source/HairStrandsCore/Private/GroomBuilder.cpp
@@ -2277,7 +2277,7 @@ bool FGroomBuilder::BuildHairDescriptionGroups(const FHairDescription& HairDescr
 				const int32 P1Index = CurrentIndex-2;
 				check(CurveNumVertices >= 2);
 				check(P0Index > 0);
-				check(P1Index > 0);
+				check(P1Index >= 0);
 
 				// Add a small position offset to avoid superposed positions of the extra position
 				const FVector3f P0 = CurrentHairStrandsDatas->StrandsPoints.PointsPosition[P0Index];
