Fix memory leak due to unrelease ref. count on skel. mesh LOD data.

diff --git a/Engine/Plugins/Runtime/HairStrands/Source/HairStrandsCore/Private/GroomManager.cpp b/Engine/Plugins/Runtime/HairStrands/Source/HairStrandsCore/Private/GroomManager.cpp
--- a/Engine/Plugins/Runtime/HairStrands/Source/HairStrandsCore/Private/GroomManager.cpp
+++ b/Engine/Plugins/Runtime/HairStrands/Source/HairStrandsCore/Private/GroomManager.cpp
@@ -225,7 +225,7 @@ struct FHairGeometryCache
 	struct FData
 	{
 		const FSkeletalMeshObject* MeshObject		= nullptr;
-		FSkeletalMeshLODRenderData* LODData 		= nullptr;
+		TRefCountPtr<FSkeletalMeshLODRenderData> LODData; // Ref. count to ensure the data are not streamed out while being used
 		FRDGBufferRef PositionBuffer				= nullptr;
 		FRDGBufferRef PreviousPositionBuffer		= nullptr;
 		FRDGBufferRef TangentBuffer					= nullptr;
@@ -558,9 +558,6 @@ static void GetOrAllocateCachedGeometry(
 	FRDGBufferSRVRef DeformedTangentSRV = nullptr;
 	OutHairGeometryCache.GetOrAdd(GraphBuilder, SkeletalMeshObject, &LODData, LODIndex, UniqueSections, DeformedPositionSRV, DeformedPreviousPositionSRV, DeformedTangentSRV);
 
-	// Add reference to be sure the data are not streamed out while they are used
-	LODData.AddRef();
-
 	// Fill in result
 	for (uint32 SectionIndex : UniqueSections)
 	{
@@ -784,13 +781,6 @@ static void RunHairBindingSurfaceUpdate(
 	}
 
 	AddHairSkinCacheDebugPass(GraphBuilder, ShaderMap, View, ShaderPrintData, Instances, HairGeometryCache);
-
-	// Release reference on skel. mesh data (only for manual skin cache)
-	for (FHairGeometryCache::FData& Data : HairGeometryCache.Datas)
-	{
-		check(Data.LODData);
-		Data.LODData->Release();
-	}
 }
 
 ///////////////////////////////////////////////////////////////////////////////////////////////////
