Add thumbnails for Groom assets

#rnx

diff --git a/Engine/Plugins/Runtime/HairStrands/Source/HairStrandsEditor/Private/GroomAssetThumbnailRenderer.cpp b/Engine/Plugins/Runtime/HairStrands/Source/HairStrandsEditor/Private/GroomAssetThumbnailRenderer.cpp
new file mode 100644
--- /dev/null
+++ b/Engine/Plugins/Runtime/HairStrands/Source/HairStrandsEditor/Private/GroomAssetThumbnailRenderer.cpp
@@ -0,0 +1,48 @@
+// Copyright Epic Games, Inc. All Rights Reserved.
+#include "GroomAssetThumbnailRenderer.h"
+
+#include "Misc/App.h"
+#include "ShowFlags.h"
+#include "SceneInterface.h"
+#include "SceneView.h"
+#include "ThumbnailHelpers.h"
+
+#include "GroomAsset.h"
+
+void UGroomAssetThumbnailRenderer::Draw(UObject* Object, int32 X, int32 Y, uint32 Width, uint32 Height, FRenderTarget* RenderTarget, FCanvas* Canvas, bool bAdditionalViewFamily)
+{
+	UGroomAsset* GroomAsset = Cast<UGroomAsset>(Object);
+	if (IsValid(GroomAsset))
+	{
+		if (!ThumbnailScene.IsValid() || !ensure(ThumbnailScene->GetWorld()))
+		{
+			if (ThumbnailScene.IsValid())
+			{
+				FlushRenderingCommands();
+				ThumbnailScene.Reset();
+			}
+
+			ThumbnailScene = MakeUnique<FGroomAssetThumbnailScene>();
+		}
+
+		ThumbnailScene->SetGroomAsset(GroomAsset);
+		ThumbnailScene->GetScene()->UpdateSpeedTreeWind(0.0);
+
+		FSceneViewFamilyContext ViewFamily(FSceneViewFamily::ConstructionValues(RenderTarget, ThumbnailScene->GetScene(), FEngineShowFlags(ESFIM_Game))
+			.SetTime(UThumbnailRenderer::GetTime())
+			.SetAdditionalViewFamily(bAdditionalViewFamily));
+
+		ViewFamily.EngineShowFlags.DisableAdvancedFeatures();
+		ViewFamily.EngineShowFlags.MotionBlur = 0;
+		ViewFamily.EngineShowFlags.LOD = 0;
+
+		RenderViewFamily(Canvas, &ViewFamily, ThumbnailScene->CreateView(&ViewFamily, X, Y, Width, Height));
+		ThumbnailScene->CleanupSceneAfterThumbnailRender();
+	}
+}
+
+void UGroomAssetThumbnailRenderer::BeginDestroy()
+{
+	ThumbnailScene.Reset();
+	Super::BeginDestroy();
+}
diff --git a/Engine/Plugins/Runtime/HairStrands/Source/HairStrandsEditor/Private/GroomAssetThumbnailRenderer.h b/Engine/Plugins/Runtime/HairStrands/Source/HairStrandsEditor/Private/GroomAssetThumbnailRenderer.h
new file mode 100644
--- /dev/null
+++ b/Engine/Plugins/Runtime/HairStrands/Source/HairStrandsEditor/Private/GroomAssetThumbnailRenderer.h
@@ -0,0 +1,28 @@
+// Copyright Epic Games, Inc. All Rights Reserved.
+#pragma once
+#include "GroomAssetThumbnailScene.h"
+#include "Templates/UniquePtr.h"
+#include "ThumbnailRendering/DefaultSizedThumbnailRenderer.h"
+#include "UObject/ObjectMacros.h"
+#include "GroomAssetThumbnailRenderer.generated.h"
+
+class FCanvas;
+class FRenderTarget;
+
+UCLASS(config = Editor, MinimalAPI)
+class UGroomAssetThumbnailRenderer : public UDefaultSizedThumbnailRenderer
+{
+	GENERATED_BODY()
+
+public:
+	// Begin UThumbnailRenderer Object
+	HAIRSTRANDSEDITOR_API virtual void Draw(UObject* Object, int32 X, int32 Y, uint32 Width, uint32 Height, FRenderTarget* RenderTarget, FCanvas* Canvas, bool bAdditionalViewFamily) override;
+	// End UThumbnailRenderer Object
+
+	// UObject implementation
+	HAIRSTRANDSEDITOR_API virtual void BeginDestroy() override;
+
+private:
+	TUniquePtr<FGroomAssetThumbnailScene> ThumbnailScene;
+};
+
diff --git a/Engine/Plugins/Runtime/HairStrands/Source/HairStrandsEditor/Private/GroomAssetThumbnailScene.cpp b/Engine/Plugins/Runtime/HairStrands/Source/HairStrandsEditor/Private/GroomAssetThumbnailScene.cpp
new file mode 100644
--- /dev/null
+++ b/Engine/Plugins/Runtime/HairStrands/Source/HairStrandsEditor/Private/GroomAssetThumbnailScene.cpp
@@ -0,0 +1,71 @@
+// Copyright Epic Games, Inc. All Rights Reserved.
+#include "GroomAssetThumbnailScene.h"
+
+#include "GroomActor.h"
+#include "GroomComponent.h"
+
+#include "ThumbnailRendering/SceneThumbnailInfo.h"
+
+FGroomAssetThumbnailScene::FGroomAssetThumbnailScene()
+	:FThumbnailPreviewScene()
+{
+	bForceAllUsedMipsResident = false;
+
+	// Create preview actor
+	// checked
+	FActorSpawnParameters SpawnInfo;
+	SpawnInfo.SpawnCollisionHandlingOverride = ESpawnActorCollisionHandlingMethod::AlwaysSpawn;
+	SpawnInfo.bNoFail = true;
+	SpawnInfo.ObjectFlags = RF_Transient;
+	PreviewActor = GetWorld()->SpawnActor<AGroomActor>(SpawnInfo);
+	PreviewActor->GetRootComponent()->SetCanEverAffectNavigation(false);
+	PreviewActor->GetRootComponent()->SetMobility(EComponentMobility::Movable);
+	PreviewActor->SetActorEnableCollision(false);
+}
+
+void FGroomAssetThumbnailScene::SetGroomAsset(UGroomAsset* GroomAsset)
+{
+	PreviewActor->GetGroomComponent()->SetGroomAsset(GroomAsset);
+	
+	if (ensure(GroomAsset))
+	{
+
+		PreviewActor->SetActorLocation(FVector(0, 0, 0), false);
+		PreviewActor->GetGroomComponent()->UpdateBounds();
+
+		const FBoxSphereBounds GroomAssetBounds = PreviewActor->GetGroomComponent()->Bounds;
+		// Center the mesh at the world origin then offset to put it on top of the plane
+		const float BoundsZOffset = GetBoundsZOffset(GroomAssetBounds);
+		PreviewActor->SetActorLocation(-GroomAssetBounds.Origin + FVector(0, 0, BoundsZOffset), false);
+	}
+
+	PreviewActor->GetGroomComponent()->MarkRenderStateDirty();
+}
+
+void FGroomAssetThumbnailScene::CleanupSceneAfterThumbnailRender()
+{
+	PreviewActor->GetGroomComponent()->SetGroomAsset(nullptr);
+	PreviewActor->GetGroomComponent()->MarkRenderStateDirty();
+}
+
+void FGroomAssetThumbnailScene::GetViewMatrixParameters(const float InFOVDegrees, FVector& OutOrigin, float& OutOrbitPitch, float& OutOrbitYaw, float& OutOrbitZoom) const
+{
+	check(PreviewActor);
+	check(PreviewActor->GetGroomComponent());
+	check(PreviewActor->GetGroomComponent()->GroomAsset);
+
+	const FBoxSphereBounds GroomAssetBounds = PreviewActor->GetGroomComponent()->Bounds;
+	const float HalfFOVRadians = FMath::DegreesToRadians<float>(!FMath::IsNearlyZero(InFOVDegrees, SMALL_NUMBER) ? InFOVDegrees : 5.0f) * 0.5f;
+
+	// Add extra size to view slightly outside of the sphere to compensate for perspective
+	const float HalfMeshSize = static_cast<float>(GroomAssetBounds.SphereRadius * 1.15);
+	const float BoundsZOffset = GetBoundsZOffset(GroomAssetBounds);
+	const float TargetDistance = HalfMeshSize / FMath::Tan(HalfFOVRadians);
+
+	USceneThumbnailInfo* ThumbnailInfo = USceneThumbnailInfo::StaticClass()->GetDefaultObject<USceneThumbnailInfo>();;
+	
+	OutOrigin = FVector(0, 0, -BoundsZOffset);
+	OutOrbitPitch = ThumbnailInfo->OrbitPitch;
+	OutOrbitYaw = ThumbnailInfo->OrbitYaw;
+	OutOrbitZoom = TargetDistance + ThumbnailInfo->OrbitZoom;
+}
diff --git a/Engine/Plugins/Runtime/HairStrands/Source/HairStrandsEditor/Private/GroomAssetThumbnailScene.h b/Engine/Plugins/Runtime/HairStrands/Source/HairStrandsEditor/Private/GroomAssetThumbnailScene.h
new file mode 100644
--- /dev/null
+++ b/Engine/Plugins/Runtime/HairStrands/Source/HairStrandsEditor/Private/GroomAssetThumbnailScene.h
@@ -0,0 +1,27 @@
+// Copyright Epic Games, Inc. All Rights Reserved.
+#pragma once
+#include "ThumbnailHelpers.h"
+#include "UObject/ObjectPtr.h"
+
+class UGroomAsset;
+class AGroomActor;
+
+class FGroomAssetThumbnailScene : public FThumbnailPreviewScene
+{
+public:
+	/** Constructor */
+	HAIRSTRANDSEDITOR_API FGroomAssetThumbnailScene();
+
+	/** Sets the groom asset to use in the next CreateView() */
+	HAIRSTRANDSEDITOR_API void SetGroomAsset(UGroomAsset* GroomAsset);
+
+
+	HAIRSTRANDSEDITOR_API void CleanupSceneAfterThumbnailRender();
+protected:
+	// FThumbnailPreviewScene implementation
+	HAIRSTRANDSEDITOR_API virtual void GetViewMatrixParameters(const float InFOVDegrees, FVector& OutOrigin, float& OutOrbitPitch, float& OutOrbitYaw, float& OutOrbitZoom) const override;
+
+private:
+	/** The static mesh actor used to display all static mesh thumbnails */
+	TObjectPtr<AGroomActor> PreviewActor;
+};
diff --git a/Engine/Plugins/Runtime/HairStrands/Source/HairStrandsEditor/Private/GroomBindingAssetThumbnailRenderer.cpp b/Engine/Plugins/Runtime/HairStrands/Source/HairStrandsEditor/Private/GroomBindingAssetThumbnailRenderer.cpp
new file mode 100644
--- /dev/null
+++ b/Engine/Plugins/Runtime/HairStrands/Source/HairStrandsEditor/Private/GroomBindingAssetThumbnailRenderer.cpp
@@ -0,0 +1,78 @@
+// Copyright Epic Games, Inc. All Rights Reserved.
+#include "GroomBindingAssetThumbnailRenderer.h"
+
+#include "Misc/App.h"
+#include "ShowFlags.h"
+#include "SceneInterface.h"
+#include "SceneView.h"
+#include "ThumbnailHelpers.h"
+
+#include "GroomBindingAsset.h"
+
+bool UGroomBindingAssetThumbnailRenderer::CanVisualizeAsset(UObject* Object)
+{
+	if (UGroomBindingAsset* GroomBindingAsset = Cast<UGroomBindingAsset>(Object); GroomBindingAsset->bIsValid)
+	{
+		EGroomBindingMeshType BindingType = GroomBindingAsset->GetGroomBindingType();
+		switch (BindingType)
+		{
+
+		case EGroomBindingMeshType::SkeletalMesh:
+		{
+			if (USkeletalMesh* TargetSkeletalMesh = GroomBindingAsset->GetTargetSkeletalMesh())
+			{
+				return true;
+			}
+			break;
+		}
+		case EGroomBindingMeshType::GeometryCache:
+		{
+			if (UGeometryCache* TargetGeometryCache = GroomBindingAsset->GetTargetGeometryCache())
+			{
+				return true;
+			}
+			break;
+		}
+		default:
+			break;
+		}
+	}
+
+	return false;
+}
+
+void UGroomBindingAssetThumbnailRenderer::Draw(UObject* Object, int32 X, int32 Y, uint32 Width, uint32 Height, FRenderTarget* RenderTarget, FCanvas* Canvas, bool bAdditionalViewFamily)
+{
+	if (UGroomBindingAsset* GroomBindingAsset = Cast<UGroomBindingAsset>(Object); GroomBindingAsset->bIsValid)
+	{
+		if (!ThumbnailScene.IsValid() || !ensure(ThumbnailScene->GetWorld()))
+		{
+			if (ThumbnailScene)
+			{
+				FlushRenderingCommands();
+				ThumbnailScene.Reset();
+			}
+			ThumbnailScene = MakeUnique<FGroomBindingAssetThumbnailScene>();
+		}
+
+		ThumbnailScene->SetGroomBindingAsset(GroomBindingAsset);
+		ThumbnailScene->GetScene()->UpdateSpeedTreeWind(0.0);
+
+		FSceneViewFamilyContext ViewFamily(FSceneViewFamily::ConstructionValues(RenderTarget, ThumbnailScene->GetScene(), FEngineShowFlags(ESFIM_Game))
+			.SetTime(UThumbnailRenderer::GetTime())
+			.SetAdditionalViewFamily(bAdditionalViewFamily));
+
+		ViewFamily.EngineShowFlags.DisableAdvancedFeatures();
+		ViewFamily.EngineShowFlags.MotionBlur = 0;
+		ViewFamily.EngineShowFlags.LOD = 0;
+
+		RenderViewFamily(Canvas, &ViewFamily, ThumbnailScene->CreateView(&ViewFamily, X, Y, Width, Height));
+		ThumbnailScene->CleanupSceneAfterThumbnailRender();
+	}
+}
+
+void UGroomBindingAssetThumbnailRenderer::BeginDestroy()
+{
+	ThumbnailScene.Reset();
+	Super::BeginDestroy();
+}
diff --git a/Engine/Plugins/Runtime/HairStrands/Source/HairStrandsEditor/Private/GroomBindingAssetThumbnailRenderer.h b/Engine/Plugins/Runtime/HairStrands/Source/HairStrandsEditor/Private/GroomBindingAssetThumbnailRenderer.h
new file mode 100644
--- /dev/null
+++ b/Engine/Plugins/Runtime/HairStrands/Source/HairStrandsEditor/Private/GroomBindingAssetThumbnailRenderer.h
@@ -0,0 +1,31 @@
+// Copyright Epic Games, Inc. All Rights Reserved.
+#pragma once
+#include "GroomBindingAssetThumbnailScene.h"
+#include "Templates/UniquePtr.h"
+#include "ThumbnailRendering/DefaultSizedThumbnailRenderer.h"
+#include "UObject/ObjectMacros.h"
+#include "GroomBindingAssetThumbnailRenderer.generated.h"
+
+
+class FCanvas;
+class FRenderTarget;
+
+UCLASS(config = Editor, MinimalAPI)
+class UGroomBindingAssetThumbnailRenderer : public UDefaultSizedThumbnailRenderer
+{
+	GENERATED_BODY()
+
+public:
+	// Begin UThumbnailRenderer Object
+	HAIRSTRANDSEDITOR_API virtual bool CanVisualizeAsset(UObject* Object);
+
+	HAIRSTRANDSEDITOR_API virtual void Draw(UObject* Object, int32 X, int32 Y, uint32 Width, uint32 Height, FRenderTarget* RenderTarget, FCanvas* Canvas, bool bAdditionalViewFamily) override;
+	// End UThumbnailRenderer Object
+
+	// UObject implementation
+	HAIRSTRANDSEDITOR_API virtual void BeginDestroy() override;
+
+private:
+	TUniquePtr<FGroomBindingAssetThumbnailScene> ThumbnailScene;
+};
+
diff --git a/Engine/Plugins/Runtime/HairStrands/Source/HairStrandsEditor/Private/GroomBindingAssetThumbnailScene.cpp b/Engine/Plugins/Runtime/HairStrands/Source/HairStrandsEditor/Private/GroomBindingAssetThumbnailScene.cpp
new file mode 100644
--- /dev/null
+++ b/Engine/Plugins/Runtime/HairStrands/Source/HairStrandsEditor/Private/GroomBindingAssetThumbnailScene.cpp
@@ -0,0 +1,150 @@
+// Copyright Epic Games, Inc. All Rights Reserved.
+#include "GroomBindingAssetThumbnailScene.h"
+
+#include "GroomActor.h"
+#include "GroomBindingAsset.h"
+#include "GroomComponent.h"
+
+#include "Engine/SkeletalMesh.h"
+#include "Components/SkeletalMeshComponent.h"
+#include "Animation/SkeletalMeshActor.h"
+
+#include "GeometryCache.h"
+#include "GeometryCacheActor.h"
+#include "GeometryCacheComponent.h"
+
+#include "ThumbnailRendering/SceneThumbnailInfo.h"
+
+namespace GroomBindingAssetThumbnailScene::Private
+{
+	template<typename ActorClassType>
+	ActorClassType* SpawnActorInThumbnailScene(FGroomBindingAssetThumbnailScene& ThumbnailScene)
+	{
+		FActorSpawnParameters SpawnInfo;
+		SpawnInfo.SpawnCollisionHandlingOverride = ESpawnActorCollisionHandlingMethod::AlwaysSpawn;
+		SpawnInfo.bNoFail = true;
+		SpawnInfo.ObjectFlags = RF_Transient;
+		ActorClassType* SpawnedActor = ThumbnailScene.GetWorld()->SpawnActor<ActorClassType>(SpawnInfo);
+		SpawnedActor->GetRootComponent()->SetCanEverAffectNavigation(false);
+		SpawnedActor->GetRootComponent()->SetMobility(EComponentMobility::Movable);
+		SpawnedActor->SetActorEnableCollision(false);
+		return SpawnedActor;
+	}
+}
+
+FGroomBindingAssetThumbnailScene::FGroomBindingAssetThumbnailScene()
+	:FThumbnailPreviewScene()
+{
+	using namespace GroomBindingAssetThumbnailScene::Private;
+	
+	bForceAllUsedMipsResident = false;
+
+	PreviewGroomActor = SpawnActorInThumbnailScene<AGroomActor>(*this);
+	PreviewSkeletalMeshActor = SpawnActorInThumbnailScene<ASkeletalMeshActor>(*this);
+	PreviewGeometryCacheActor = SpawnActorInThumbnailScene<AGeometryCacheActor>(*this);
+	
+}
+
+void FGroomBindingAssetThumbnailScene::SetGroomBindingAsset(UGroomBindingAsset* GroomBindingAsset)
+{
+	using namespace GroomBindingAssetThumbnailScene::Private;
+	if (ensure(GroomBindingAsset))
+	{
+		UGroomAsset* GroomAsset = GroomBindingAsset->GetGroom();
+		PreviewGroomActor->GetGroomComponent()->SetGroomAsset(GroomAsset);
+
+		EGroomBindingMeshType BindingType = GroomBindingAsset->GetGroomBindingType();
+		switch (BindingType)
+		{
+
+		case EGroomBindingMeshType::SkeletalMesh:
+		{
+			if (USkeletalMesh* TargetSkeletalMesh = GroomBindingAsset->GetTargetSkeletalMesh())
+			{
+				PreviewSkeletalMeshActor->GetSkeletalMeshComponent()->SetSkeletalMesh(TargetSkeletalMesh);
+				PreviewRootActor = PreviewSkeletalMeshActor;
+			}
+			
+			break;
+		}
+		case EGroomBindingMeshType::GeometryCache:
+		{
+			if (UGeometryCache* TargetGeometryCache = GroomBindingAsset->GetTargetGeometryCache())
+			{
+				PreviewGeometryCacheActor->GetGeometryCacheComponent()->SetGeometryCache(TargetGeometryCache);
+				PreviewRootActor = PreviewGeometryCacheActor;
+			}
+			break;
+		}
+		default:
+			break;
+		}
+
+		if (PreviewRootActor != nullptr)
+		{
+			PreviewGroomActor->GetGroomComponent()->UpdateBounds();
+			PreviewGroomActor->AttachToActor(PreviewRootActor, FAttachmentTransformRules::SnapToTargetNotIncludingScale);
+		}
+		else
+		{
+			PreviewRootActor = PreviewGroomActor;
+		}
+
+		
+		PreviewRootActor->SetActorLocation(FVector(0, 0, 0), false);
+		PreviewRootActor->GetRootComponent()->UpdateBounds();
+			
+		// Center the mesh at the world origin then offset to put it on top of the plane
+		const float BoundsZOffset = GetBoundsZOffset(PreviewRootActor->GetRootComponent()->Bounds);
+		PreviewRootActor->SetActorLocation(-PreviewRootActor->GetRootComponent()->Bounds.Origin + FVector(0, 0, BoundsZOffset), false);
+	
+		PreviewGroomActor->GetGroomComponent()->MarkRenderStateDirty();
+		if (PreviewRootActor != PreviewGroomActor)
+		{
+			PreviewRootActor->GetRootComponent()->MarkRenderStateDirty();
+		}
+	}
+}
+
+void FGroomBindingAssetThumbnailScene::CleanupSceneAfterThumbnailRender()
+{
+	if (PreviewRootActor != nullptr && PreviewGroomActor != PreviewRootActor)
+	{
+		PreviewGroomActor->DetachFromActor(FDetachmentTransformRules::KeepRelativeTransform);
+	}
+
+	PreviewRootActor = nullptr;
+
+	PreviewGroomActor->GetGroomComponent()->SetGroomAsset(nullptr);
+	PreviewGroomActor->GetGroomComponent()->MarkRenderStateDirty();
+
+	PreviewSkeletalMeshActor->GetSkeletalMeshComponent()->SetSkeletalMesh(nullptr);
+	PreviewSkeletalMeshActor->GetSkeletalMeshComponent()->MarkRenderStateDirty();
+
+	PreviewGeometryCacheActor->GetGeometryCacheComponent()->SetGeometryCache(nullptr);
+	PreviewGeometryCacheActor->GetGeometryCacheComponent()->MarkRenderStateDirty();
+}
+
+void FGroomBindingAssetThumbnailScene::GetViewMatrixParameters(const float InFOVDegrees, FVector& OutOrigin, float& OutOrbitPitch, float& OutOrbitYaw, float& OutOrbitZoom) const
+{
+	check(PreviewGroomActor);
+	check(PreviewGroomActor->GetGroomComponent());
+	check(PreviewGroomActor->GetGroomComponent()->GroomAsset);
+
+	check(PreviewRootActor);
+	check(PreviewRootActor->GetRootComponent());
+
+	const float HalfFOVRadians = FMath::DegreesToRadians<float>(!FMath::IsNearlyZero(InFOVDegrees, SMALL_NUMBER) ? InFOVDegrees : 5.0f) * 0.5f;
+	
+	// Add extra size to view slightly outside of the sphere to compensate for perspective
+	const float HalfMeshSize = static_cast<float>(PreviewRootActor->GetRootComponent()->Bounds.SphereRadius * 1.15);
+	const float BoundsZOffset = GetBoundsZOffset(PreviewRootActor->GetRootComponent()->Bounds);
+	const float TargetDistance = HalfMeshSize / FMath::Tan(HalfFOVRadians);
+
+	USceneThumbnailInfo* ThumbnailInfo = USceneThumbnailInfo::StaticClass()->GetDefaultObject<USceneThumbnailInfo>();
+	
+	OutOrigin = FVector(0, 0, -BoundsZOffset);
+	OutOrbitPitch = ThumbnailInfo->OrbitPitch;
+	OutOrbitYaw = ThumbnailInfo->OrbitYaw;
+	OutOrbitZoom = TargetDistance + ThumbnailInfo->OrbitZoom;
+}
diff --git a/Engine/Plugins/Runtime/HairStrands/Source/HairStrandsEditor/Private/GroomBindingAssetThumbnailScene.h b/Engine/Plugins/Runtime/HairStrands/Source/HairStrandsEditor/Private/GroomBindingAssetThumbnailScene.h
new file mode 100644
--- /dev/null
+++ b/Engine/Plugins/Runtime/HairStrands/Source/HairStrandsEditor/Private/GroomBindingAssetThumbnailScene.h
@@ -0,0 +1,41 @@
+// Copyright Epic Games, Inc. All Rights Reserved.
+#pragma once
+#include "ThumbnailHelpers.h"
+#include "UObject/ObjectPtr.h"
+
+class AActor;
+class AGeometryCacheActor;
+class AGroomActor;
+class AGroomActor;
+class ASkeletalMeshActor;
+class UGroomBindingAsset;
+
+class FGroomBindingAssetThumbnailScene : public FThumbnailPreviewScene
+{
+public:
+	/** Constructor */
+	HAIRSTRANDSEDITOR_API FGroomBindingAssetThumbnailScene();
+
+	/** Sets the groom binding to use in the next CreateView() */
+	HAIRSTRANDSEDITOR_API void SetGroomBindingAsset(UGroomBindingAsset* GroomBindingAsset);
+
+	/** Sets the groom binding to use in the next CreateView() */
+	HAIRSTRANDSEDITOR_API void CleanupSceneAfterThumbnailRender();
+
+protected:
+	// FThumbnailPreviewScene implementation
+	HAIRSTRANDSEDITOR_API virtual void GetViewMatrixParameters(const float InFOVDegrees, FVector& OutOrigin, float& OutOrbitPitch, float& OutOrbitYaw, float& OutOrbitZoom) const override;
+
+private:
+	/** The groom actor used to display all groom asset thumbnails */
+	TObjectPtr<AGroomActor> PreviewGroomActor;
+
+	/** The skeletal mesh actor for which the binding is created that will be previewed */
+	TObjectPtr<ASkeletalMeshActor> PreviewSkeletalMeshActor;
+
+	/** The geometry cache actor for which the binding is created that will be previewed */
+	TObjectPtr<AGeometryCacheActor> PreviewGeometryCacheActor;
+
+	/** Actor used as the root of the scene for thumbnail */
+	TObjectPtr<AActor> PreviewRootActor = nullptr;
+};
diff --git a/Engine/Plugins/Runtime/HairStrands/Source/HairStrandsEditor/Private/HairStrandsEditor.cpp b/Engine/Plugins/Runtime/HairStrands/Source/HairStrandsEditor/Private/HairStrandsEditor.cpp
--- a/Engine/Plugins/Runtime/HairStrands/Source/HairStrandsEditor/Private/HairStrandsEditor.cpp
+++ b/Engine/Plugins/Runtime/HairStrands/Source/HairStrandsEditor/Private/HairStrandsEditor.cpp
@@ -13,7 +13,9 @@
 #include "ToolMenus.h"
 
 #include "GroomAsset.h"
+#include "GroomAssetThumbnailRenderer.h"
 #include "GroomBindingAsset.h"
+#include "GroomBindingAssetThumbnailRenderer.h"
 #include "GroomBindingDetailsCustomization.h"
 #include "GroomCacheImportOptions.h"
 #include "GroomCacheImportSettingsCustomization.h"
@@ -22,8 +24,8 @@
 #include "GroomCreateBindingOptions.h"
 #include "GroomEditorCommands.h"
 #include "GroomEditorMode.h"
-#include "GroomPluginSettings.h"
 #include "GroomImportOptions.h"
+#include "GroomPluginSettings.h"
 
 #include "AssetRegistry/AssetRegistryModule.h"
 #include "FileHelpers.h"
@@ -31,6 +33,7 @@
 #include "ISettingsSection.h"
 #include "ISequencerModule.h"
 #include "Delegates/DelegateCombinations.h"
+#include "ThumbnailRendering/ThumbnailManager.h"
 #include "UObject/UObjectIterator.h"
 
 IMPLEMENT_MODULE(FGroomEditor, HairStrandsEditor);
@@ -148,6 +151,9 @@ void FGroomEditor::StartupModule()
 			GetMutableDefault<UGroomPluginSettings>()
 		);
 	}
+
+	UThumbnailManager::Get().RegisterCustomRenderer(UGroomAsset::StaticClass(), UGroomAssetThumbnailRenderer::StaticClass());
+	UThumbnailManager::Get().RegisterCustomRenderer(UGroomBindingAsset::StaticClass(), UGroomBindingAssetThumbnailRenderer::StaticClass());
 }
 
 void FGroomEditor::ShutdownModule()
@@ -180,6 +186,9 @@ void FGroomEditor::ShutdownModule()
 			PropertyModule->UnregisterCustomClassLayout(UGroomCreateBindingOptions::StaticClass()->GetFName());
 			PropertyModule->UnregisterCustomPropertyTypeLayout(FGroomCacheImportSettings::StaticStruct()->GetFName());
 		}
+
+		UThumbnailManager::Get().UnregisterCustomRenderer(UGroomAsset::StaticClass());
+		UThumbnailManager::Get().UnregisterCustomRenderer(UGroomBindingAsset::StaticClass());
 	}
 
 	if (StyleSet.IsValid())
