Improve groom dataflow  rendering to better integrate with dataflow outliner

diff --git a/Engine/Plugins/Experimental/HairCardGenerator/Source/HairCardGeneratorDataflow/Private/HairCardDataflowRendering.cpp b/Engine/Plugins/Experimental/HairCardGenerator/Source/HairCardGeneratorDataflow/Private/HairCardDataflowRendering.cpp
--- a/Engine/Plugins/Experimental/HairCardGenerator/Source/HairCardGeneratorDataflow/Private/HairCardDataflowRendering.cpp
+++ b/Engine/Plugins/Experimental/HairCardGenerator/Source/HairCardGeneratorDataflow/Private/HairCardDataflowRendering.cpp
@@ -5,6 +5,7 @@
 #include "GenerateCardsClumpsNode.h"
 #include "GenerateCardsTexturesNode.h"
 #include "GroomCollectionFacades.h"
+#include "GeometryCollection/Facades/CollectionRenderingFacade.h"
 
 namespace UE::CardGen::Private
 {
@@ -190,7 +191,7 @@ namespace UE::CardGen::Private
 
 			for(int32 TextureIndex = 0; TextureIndex < NumTextures; ++TextureIndex)
 			{
-				TextureColors[TextureIndex] = FLinearColor::MakeRandomColor();
+				TextureColors[TextureIndex] = FLinearColor::IntToDistinctColor(TextureIndex, 0.75f, 1.0f, 90.0f);
 			}
 
 			const int32 NumVertices = VertexCardsArray.Num();
@@ -208,6 +209,25 @@ namespace UE::CardGen::Private
 
 	Dataflow::FRenderKey FCardsClumpsRenderingCallbacks::RenderKey = { "ClumpsRender", FName("FCardsCollection") };
 	
+	int32 FCardsClumpsRenderingCallbacks::GetGroupAttribute(const UE::Groom::FGroomStrandsFacade& StrandsFacade, FString& GroupAttribute, FString& GroupName) const
+	{
+		GroupAttribute = FGenerateCardsClumpsNode::CurveClumpIndicesAttribute.ToString();
+		GroupAttribute.AppendInt(GGroomDataflowCardsLod);
+
+		FString NumClumpsLOD = FGenerateCardsClumpsNode::ObjectNumClumpsAttribute.ToString();
+		NumClumpsLOD.AppendInt(GGroomDataflowCardsLod);
+
+		GroupName = FString("Clump");
+
+		const FManagedArrayCollection& GroomCollection = StrandsFacade.GetManagedArrayCollection();
+		if(GroomCollection.HasAttribute(FName(NumClumpsLOD), UE::Groom::FGroomStrandsFacade::ObjectsGroup))
+		{
+			const TManagedArray<int32>& NumClumps = GroomCollection.GetAttribute<int32>(FName(NumClumpsLOD), UE::Groom::FGroomStrandsFacade::ObjectsGroup);
+			return (NumClumps.Num() > 0) ? NumClumps[0] : 0;
+		}
+		return 0;
+	}
+	
 	void FCardsClumpsRenderingCallbacks::ComputeVertexColors(const UE::Groom::FGroomStrandsFacade& StrandsFacade, TArray<FLinearColor>& VertexColors) const
 	{
 		const FManagedArrayCollection& GroomCollection = StrandsFacade.GetManagedArrayCollection();
@@ -229,7 +249,7 @@ namespace UE::CardGen::Private
 
 			for(int32 ClumpIndex = 0; ClumpIndex < NumClumps[0]; ++ClumpIndex)
 			{
-				ClumpsColors[ClumpIndex] = FLinearColor::MakeRandomColor();
+				ClumpsColors[ClumpIndex] = FLinearColor::IntToDistinctColor(ClumpIndex, 0.75f, 1.0f, 90.0f);
 			}
 
 			TArray<FLinearColor> StrandsColors;
diff --git a/Engine/Plugins/Experimental/HairCardGenerator/Source/HairCardGeneratorDataflow/Public/GenerateCardsClumpsNode.h b/Engine/Plugins/Experimental/HairCardGenerator/Source/HairCardGeneratorDataflow/Public/GenerateCardsClumpsNode.h
--- a/Engine/Plugins/Experimental/HairCardGenerator/Source/HairCardGeneratorDataflow/Public/GenerateCardsClumpsNode.h
+++ b/Engine/Plugins/Experimental/HairCardGenerator/Source/HairCardGeneratorDataflow/Public/GenerateCardsClumpsNode.h
@@ -33,7 +33,7 @@ struct FCardsClumpsSettings
 	int32 NumFlyaways = 10;
 };
 
-/** Generate the textures used to build the cards from the strands */
+/** Generate the clumps used to build the cards from the strands */
 USTRUCT(meta = (DataflowGroom))
 struct FGenerateCardsClumpsNode : public FDataflowNode
 {
diff --git a/Engine/Plugins/Experimental/HairCardGenerator/Source/HairCardGeneratorDataflow/Public/HairCardDataflowRendering.h b/Engine/Plugins/Experimental/HairCardGenerator/Source/HairCardGeneratorDataflow/Public/HairCardDataflowRendering.h
--- a/Engine/Plugins/Experimental/HairCardGenerator/Source/HairCardGeneratorDataflow/Public/HairCardDataflowRendering.h
+++ b/Engine/Plugins/Experimental/HairCardGenerator/Source/HairCardGeneratorDataflow/Public/HairCardDataflowRendering.h
@@ -42,6 +42,9 @@ namespace UE::CardGen::Private
 		}
 		
 		virtual void ComputeVertexColors(const UE::Groom::FGroomStrandsFacade& StrandsFacade, TArray<FLinearColor>& VertexColors) const override;
+		
+		/** Get the group attribute used to create each geometry groups */
+		virtual int32 GetGroupAttribute(const UE::Groom::FGroomStrandsFacade& StrandsFacade, FString& GroupAttribute, FString& GroupName) const override;
 	};
 
 	class FCardsTextureRenderingCallbacks : public FCardsGeometryRenderingCallbacks
diff --git a/Engine/Plugins/Runtime/HairStrands/Source/HairStrandsCore/Private/GroomBuilder.cpp b/Engine/Plugins/Runtime/HairStrands/Source/HairStrandsCore/Private/GroomBuilder.cpp
--- a/Engine/Plugins/Runtime/HairStrands/Source/HairStrandsCore/Private/GroomBuilder.cpp
+++ b/Engine/Plugins/Runtime/HairStrands/Source/HairStrandsCore/Private/GroomBuilder.cpp
@@ -2338,7 +2338,7 @@ bool FGroomBuilder::BuildHairDescriptionGroups(const FHairDescription& HairDescr
 		// Sanity check
 		check(Group.Strands.GetNumCurves() == Group.Info.NumCurves);
 		check(Group.Guides.GetNumCurves() == Group.Info.NumGuides);
-
+		
 		// Update infos
 		Group.Info.NumCurveVertices = Group.Strands.GetNumPoints();
 		Group.Info.NumGuideVertices = Group.Guides.GetNumPoints();
@@ -2365,7 +2365,8 @@ void FGroomBuilder::BuildData(
 	const FHairGroupsInterpolation& InSettings,
 	FHairStrandsDatas& OutRen,
 	FHairStrandsDatas& OutSim,
-	bool bAllowCurveReordering)
+	const bool bAllowCurveReordering,
+	const bool bApplyDecimation)
 {
 	FHairGroupInfo DummyGroupInfo;
 	BuildData(
@@ -2374,7 +2375,8 @@ void FGroomBuilder::BuildData(
 		DummyGroupInfo, 
 		OutRen,
 		OutSim,
-		bAllowCurveReordering);
+		bAllowCurveReordering,
+		bApplyDecimation);
 }
 
 void FGroomBuilder::BuildData(
@@ -2383,14 +2385,15 @@ void FGroomBuilder::BuildData(
 	FHairGroupInfo& OutGroupInfo, 
 	FHairStrandsDatas& OutRen,
 	FHairStrandsDatas& OutSim,
-	bool bAllowCurveReordering)
+	const bool bAllowCurveReordering,
+	const bool bApplyDecimation)
 {
 	TRACE_CPUPROFILER_EVENT_SCOPE(FGroomBuilder::BuildData);
 
 	// Sanitize decimation values. Do not update the 'InSettings' values directly as this would change 
 	// the groom asset and thus would change the DDC key
-	const float CurveDecimation		= FMath::Clamp(InSettings.DecimationSettings.CurveDecimation, 0.f, 1.f);
-	const float VertexDecimation	= FMath::Clamp(InSettings.DecimationSettings.VertexDecimation, 0.f, 1.f);
+	const float CurveDecimation		= bApplyDecimation ? FMath::Clamp(InSettings.DecimationSettings.CurveDecimation, 0.f, 1.f) : 1.0f;
+	const float VertexDecimation	= bApplyDecimation ? FMath::Clamp(InSettings.DecimationSettings.VertexDecimation, 0.f, 1.f) : 1.0f;
 	const float HairToGuideDensity	= FMath::Clamp(InSettings.InterpolationSettings.HairToGuideDensity, 0.f, 1.f);	
 	const bool bCurveReordering 	= bAllowCurveReordering;
 
@@ -2404,7 +2407,7 @@ void FGroomBuilder::BuildData(
 		if (InHairDescriptionGroup.Strands.IsValid())
 		{
 			OutRen = HairStrandsBuilder::BuildInternalData(InHairDescriptionGroup.Strands);
-
+			
 			// Decimate
 			if (CurveDecimation < 1 || VertexDecimation < 1 || bCurveReordering)
 			{
@@ -2412,6 +2415,7 @@ void FGroomBuilder::BuildData(
 				OutRen.Reset();
 				FHairStrandsDecimation::Decimate(FullData, CurveDecimation, VertexDecimation, bCurveReordering, OutRen);
 			}
+			
 			OutGroupInfo.NumCurves			= OutRen.GetNumCurves();
 			OutGroupInfo.NumCurveVertices	= OutRen.GetNumPoints();
 			OutGroupInfo.MaxCurveLength		= GetHairStrandsMaxLength(OutRen);
diff --git a/Engine/Plugins/Runtime/HairStrands/Source/HairStrandsCore/Private/GroomEdit.cpp b/Engine/Plugins/Runtime/HairStrands/Source/HairStrandsCore/Private/GroomEdit.cpp
--- a/Engine/Plugins/Runtime/HairStrands/Source/HairStrandsCore/Private/GroomEdit.cpp
+++ b/Engine/Plugins/Runtime/HairStrands/Source/HairStrandsCore/Private/GroomEdit.cpp
@@ -166,7 +166,7 @@ static void EditableGroomGroupToHairDescription(FHairDescription& OutHairDescrip
 }
 #endif // WITH_EDITORONLY_DATA
 
-void ConvertFromGroomAsset(UGroomAsset* InGroom, FEditableGroom* OutGroom, const bool bAllowCurveReordering)
+void ConvertFromGroomAsset(UGroomAsset* InGroom, FEditableGroom* OutGroom, const bool bAllowCurveReordering, const bool bApplyDecimation)
 {
 #if WITH_EDITORONLY_DATA
 	check(InGroom);
@@ -184,7 +184,7 @@ void ConvertFromGroomAsset(UGroomAsset* InGroom, FEditableGroom* OutGroom, const
 		FHairGroupInfo GroupInfo;
 		FHairStrandsDatas StrandsData;
 		FHairStrandsDatas GuidesData;
-		FGroomBuilder::BuildData(HairDescriptionGroups.HairGroups[GroupIndex], In.GetHairGroupsInterpolation()[GroupIndex], GroupInfo, StrandsData, GuidesData, bAllowCurveReordering);
+		FGroomBuilder::BuildData(HairDescriptionGroups.HairGroups[GroupIndex], In.GetHairGroupsInterpolation()[GroupIndex], GroupInfo, StrandsData, GuidesData, bAllowCurveReordering, bApplyDecimation);
 
 		Out.Groups[GroupIndex].GroupIndex = GroupIndex;
 		Out.Groups[GroupIndex].GroupID = GroupInfo.GroupID;
diff --git a/Engine/Plugins/Runtime/HairStrands/Source/HairStrandsCore/Public/GroomBuilder.h b/Engine/Plugins/Runtime/HairStrands/Source/HairStrandsCore/Public/GroomBuilder.h
--- a/Engine/Plugins/Runtime/HairStrands/Source/HairStrandsCore/Public/GroomBuilder.h
+++ b/Engine/Plugins/Runtime/HairStrands/Source/HairStrandsCore/Public/GroomBuilder.h
@@ -71,14 +71,16 @@ struct HAIRSTRANDSCORE_API FGroomBuilder
 		FHairGroupInfo& OutGroupInfo,
 		FHairStrandsDatas& OutStrands,
 		FHairStrandsDatas& OutGuides,
-		bool bAllowCurveReordering=true);
+		const bool bAllowCurveReordering = true,
+		const bool bApplyDecimation = true);
 
 	static void BuildData(
 		const FHairDescriptionGroup& InHairDescriptionGroup,
 		const FHairGroupsInterpolation& InSettings,
 		FHairStrandsDatas& OutStrands,
 		FHairStrandsDatas& OutGuides,
-		bool bAllowCurveReordering=true);
+		const bool bAllowCurveReordering = true,
+		const bool bApplyDecimation = true);
 
 	// 2.b Build FHairStrandsDatas from a FHairStrandsRawDatas (Strands / Guides)
 	UE_DEPRECATED(5.6, "Please, no longer use this build function, as it does not contains certain build steps (decimation/shuffling)")
diff --git a/Engine/Plugins/Runtime/HairStrands/Source/HairStrandsCore/Public/GroomEdit.h b/Engine/Plugins/Runtime/HairStrands/Source/HairStrandsCore/Public/GroomEdit.h
--- a/Engine/Plugins/Runtime/HairStrands/Source/HairStrandsCore/Public/GroomEdit.h
+++ b/Engine/Plugins/Runtime/HairStrands/Source/HairStrandsCore/Public/GroomEdit.h
@@ -71,7 +71,7 @@ enum EEditableGroomOperations
 	ControlPoints_Modified = 0x2,
 	ControlPoints_Deleted  = 0x4,
 
-	Strands_Added          = 0x8,
+	Strands_Added          = 0x8, 
 	Strands_Modified       = 0x10,
 	Strands_Deleted        = 0x20,
 
@@ -80,7 +80,7 @@ enum EEditableGroomOperations
 };
 
 // Convert a groom asset into an editable groom asset
-HAIRSTRANDSCORE_API void ConvertFromGroomAsset(UGroomAsset* In, FEditableGroom* Out, const bool bAllowCurveReordering = true);
+HAIRSTRANDSCORE_API void ConvertFromGroomAsset(UGroomAsset* In, FEditableGroom* Out, const bool bAllowCurveReordering = true, const bool bApplyDecimation = true);
 
 // Convert an editable groom asset into a groom asset
 // The 'operations' flag indicates what type of modifications have been done onto the editable groom. 
diff --git a/Engine/Plugins/Runtime/HairStrands/Source/HairStrandsDataflow/Private/GetGroomAssetNode.cpp b/Engine/Plugins/Runtime/HairStrands/Source/HairStrandsDataflow/Private/GetGroomAssetNode.cpp
--- a/Engine/Plugins/Runtime/HairStrands/Source/HairStrandsDataflow/Private/GetGroomAssetNode.cpp
+++ b/Engine/Plugins/Runtime/HairStrands/Source/HairStrandsDataflow/Private/GetGroomAssetNode.cpp
@@ -25,7 +25,7 @@ namespace UE::Groom::Private
 
 		if (GroomAsset)
 		{
-			ConvertFromGroomAsset(const_cast<UGroomAsset*>(GroomAsset), &EditGroom, false);
+			ConvertFromGroomAsset(const_cast<UGroomAsset*>(GroomAsset), &EditGroom, false, false);
 
 			TArray<FVector3f> PointRestPositions;
 			TArray<int32> ObjectCurveOffsets;
@@ -35,6 +35,7 @@ namespace UE::Groom::Private
 			uint32 GroupIndex = 0;
 			for(FEditableGroomGroup& GroomGroup : EditGroom.Groups)
 			{
+				
 				for( const typename FacadeType::FEditableType& EditType : FacadeType::GetEditableGroom(GroomGroup))
 				{
 					for(auto& GuidePoint : EditType.ControlPoints)
@@ -44,8 +45,8 @@ namespace UE::Groom::Private
 					CurvePointOffsets.Add(PointRestPositions.Num());
 				}
 				ObjectCurveOffsets.Add(CurvePointOffsets.Num());
-
-				const FString GroupName = GroomAsset->GetName() + TEXT("_") + FacadeType::GroupPrefix.ToString() + TEXT("_") + FString::FromInt(GroupIndex);
+				
+				const FString GroupName = GroomAsset->GetName() + TEXT("_") + FacadeType::GroupPrefix.ToString();
 				ObjectGroupNames.Add(GroupName);
 				++GroupIndex;
 			}
diff --git a/Engine/Plugins/Runtime/HairStrands/Source/HairStrandsDataflow/Private/GroomAssetTerminalNode.cpp b/Engine/Plugins/Runtime/HairStrands/Source/HairStrandsDataflow/Private/GroomAssetTerminalNode.cpp
--- a/Engine/Plugins/Runtime/HairStrands/Source/HairStrandsDataflow/Private/GroomAssetTerminalNode.cpp
+++ b/Engine/Plugins/Runtime/HairStrands/Source/HairStrandsDataflow/Private/GroomAssetTerminalNode.cpp
@@ -19,7 +19,7 @@ namespace UE::Groom::Private
 		{
 			const int32 NextCurve = GuidesFacade.GetObjectCurveOffsets()[ObjectIndex];
 			Group.Guides.SetNum(NextCurve-PrevCurve);
-
+			
 			// Only rebuild the guides since the strands are the same
 			for (FEditableHairGuide& Guide : Group.Guides)
 			{
@@ -33,7 +33,7 @@ namespace UE::Groom::Private
 						FMath::Clamp(static_cast<float>(PointIndex-PrevPoint) / (NextPoint-PrevPoint-1), 0.f, 1.f)});
 				}
 						
-				if(StrandIndex != INDEX_NONE)
+				if((StrandIndex != INDEX_NONE) && (StrandIndex < Group.Strands.Num()))
 				{
 					Guide.GuideID = CurveIndex;
 					Guide.RootUV = Group.Strands[StrandIndex].RootUV;
@@ -132,12 +132,13 @@ void FGroomAssetTerminalDataflowNode::SetAssetValue(TObjectPtr<UObject> Asset, U
 
 			// Build an editable groom asset for the strands
 			FEditableGroom EditGroom;
-			ConvertFromGroomAsset(const_cast<UGroomAsset*>(GroomAsset), &EditGroom, false);
+			ConvertFromGroomAsset(const_cast<UGroomAsset*>(GroomAsset), &EditGroom, false, false);
 
 			if(GuidesFacade.GetNumObjects() == EditGroom.Groups.Num())
 			{
 				// Build the editable guides
 				UE::Groom::Private::BuildEditableGuides(GuidesFacade,EditGroom);
+				 
 
 				// Ensure compilation dependent assets is done
 				FAssetCompilingManager::Get().FinishCompilationForObjects({ GroomAsset });
diff --git a/Engine/Plugins/Runtime/HairStrands/Source/HairStrandsDataflow/Private/GroomDataflowRendering.cpp b/Engine/Plugins/Runtime/HairStrands/Source/HairStrandsDataflow/Private/GroomDataflowRendering.cpp
--- a/Engine/Plugins/Runtime/HairStrands/Source/HairStrandsDataflow/Private/GroomDataflowRendering.cpp
+++ b/Engine/Plugins/Runtime/HairStrands/Source/HairStrandsDataflow/Private/GroomDataflowRendering.cpp
@@ -4,6 +4,7 @@
 #include "GroomCollectionFacades.h"
 #include "Engine/SkeletalMesh.h"
 #include "GeometryCollection/GeometryCollection.h"
+#include "GeometryCollection/Facades/CollectionRenderingFacade.h"
 
 namespace UE::Groom
 {
@@ -106,6 +107,89 @@ namespace UE::Groom
 		}
 	}
 
+	template<typename FacadeType>
+	FORCEINLINE void BuildRenderingGroups(const FacadeType& GroomFacade, TArray<TArray<FVector3f>>& GroupVertexPositions,
+		TArray<TArray<FVector3f>>& GroupVertexNormals, TArray<TArray<FIntVector>>& GroupFaceVertices, const FString& GroupAttribute, const int32)
+	{
+		const int32 NumCurves = GroomFacade.GetNumCurves();
+
+		const FManagedArrayCollection& GroomCollection = GroomFacade.GetManagedArrayCollection();
+		if(GroomCollection.HasAttribute(FName(GroupAttribute), FacadeType::CurvesGroup))
+		{
+			const TManagedArray<int32>& CurveGroups = GroomCollection.GetAttribute<int32>(FName(GroupAttribute), FacadeType::CurvesGroup);
+
+			int32 NumGroups = 0;
+			for(int32 CurveIndex = 0; CurveIndex < NumCurves; ++CurveIndex)
+			{
+				NumGroups = FMath::Max(NumGroups, CurveGroups[CurveIndex]);
+			}
+			++NumGroups;
+			
+			GroupVertexPositions.SetNum(NumGroups);
+			GroupVertexNormals.SetNum(NumGroups);
+			GroupFaceVertices.SetNum(NumGroups);
+
+			TArray<int32> GroupsPoints;
+			GroupsPoints.Init(0, NumGroups);
+			
+			TArray<int32> GroupsCurves;
+			GroupsCurves.Init(0, NumGroups);
+
+			int32 PointOffset = 0;
+			for(int32 CurveIndex = 0; CurveIndex < NumCurves; ++CurveIndex)
+			{
+				const int32 CurvePoints =  GroomFacade.GetCurvePointOffsets()[CurveIndex] - PointOffset;
+				GroupsPoints[CurveGroups[CurveIndex]] += CurvePoints;
+				++GroupsCurves[CurveGroups[CurveIndex]];
+				PointOffset = GroomFacade.GetCurvePointOffsets()[CurveIndex];
+			}
+
+			for(int32 GroupIndex = 0; GroupIndex < NumGroups; ++GroupIndex)
+			{
+				GroupVertexPositions[GroupIndex].SetNum(GroupsPoints[GroupIndex] * 2);
+				GroupVertexNormals[GroupIndex].SetNum(GroupsPoints[GroupIndex] * 2);
+				GroupFaceVertices[GroupIndex].SetNum((GroupsPoints[GroupIndex]-GroupsCurves[GroupIndex]) * 2);
+			}
+
+			GroupsPoints.Init(0, NumGroups);
+			GroupsCurves.Init(0, NumGroups);
+			
+			PointOffset = 0;
+			for(int32 CurveIndex = 0; CurveIndex < NumCurves; ++CurveIndex)
+			{
+				const int32 GroupIndex = CurveGroups[CurveIndex];
+				
+				TArray<FVector3f>& VertexPositions = GroupVertexPositions[GroupIndex];
+				TArray<FVector3f>& VertexNormals = GroupVertexNormals[GroupIndex];
+				TArray<FIntVector>& FaceVertices = GroupFaceVertices[GroupIndex];
+				
+				for (int32 PointIndex = PointOffset, PointEnd = GroomFacade.GetCurvePointOffsets()[CurveIndex];
+						   PointIndex < PointEnd; ++PointIndex)
+				{
+					// Build 2 vertices for each points
+					const int32 EdgeIndex = (PointIndex == (PointEnd-1)) ? PointIndex-CurveIndex-1 : PointIndex-CurveIndex;
+										
+					const FVector3f SideVector = GroomFacade.GetEdgeRestOrientations()[EdgeIndex].GetAxisX();
+					BuildVertexPositions(GroupsPoints[GroupIndex], GroomFacade.GetPointRestPositions()[PointIndex], SideVector, GGroomDataflowRenderingThickness, VertexPositions);
+
+					const FVector3f NormalVector = GroomFacade.GetEdgeRestOrientations()[EdgeIndex].GetAxisZ();
+					BuildVertexNormals(GroupsPoints[GroupIndex], NormalVector, VertexNormals);
+							
+					// Build 2 faces for each edges perpendicular to the up-vector curve
+					if(PointIndex < (PointEnd-1))
+					{
+						BuildFaceVertices(GroupsPoints[GroupIndex], GroupsCurves[GroupIndex], FaceVertices);
+					}
+
+					++GroupsPoints[GroupIndex];
+					
+				}
+				++GroupsCurves[GroupIndex];
+				PointOffset = GroomFacade.GetCurvePointOffsets()[CurveIndex];
+			}
+		}
+	}
+
 	template<typename FacadeType>
 	FORCEINLINE void BuildGeometryDatas(const FacadeType& GroomFacade, TArray<float>& VertexPositions, TArray<int32>& FaceVertices, TArray<int32>& ObjectFaceOffsets, TArray<int32>& ObjectVertexOffsets)
 	{
@@ -215,6 +299,54 @@ namespace UE::Groom
 			}
 		}
 	}
+
+	template<typename FacadeType>
+	FORCEINLINE void RenderGroupCollection(GeometryCollection::Facades::FRenderingFacade& RenderCollection, const Dataflow::FGraphRenderingState& State,
+		const TFunction<int32(const FacadeType&, FString&, FString&)>& GroupLambda)
+	{
+		if (State.GetRenderOutputs().Num())
+		{
+			const FManagedArrayCollection Default;
+			checkf(State.GetRenderOutputs().Num() == 1, TEXT("Expected FGraphRenderingState object to have one render output"));
+			const FName PrimaryOutput = State.GetRenderOutputs()[0];
+			FManagedArrayCollection Collection = State.GetValue<FManagedArrayCollection>(PrimaryOutput, Default);
+
+			const TSharedRef<const FManagedArrayCollection> GroomCollection = MakeShared<FManagedArrayCollection>(MoveTemp(Collection));
+			const FacadeType GroomFacade(*GroomCollection);
+			if (GroomFacade.IsValid())
+			{
+				FString GroupAttribute, GroupName;
+				const int32 NumGroups = GroupLambda(GroomFacade, GroupAttribute, GroupName);
+				
+				TArray<TArray<FVector3f>> VertexPositions;
+				TArray<TArray<FVector3f>> VertexNormals;
+				TArray<TArray<FIntVector>> FaceIndices;
+
+				if (State.GetViewMode().GetName() == Dataflow::FDataflowConstruction3DViewMode::Name)
+				{
+					BuildRenderingGroups(GroomFacade, VertexPositions, VertexNormals, FaceIndices, GroupAttribute, NumGroups);
+				}
+				else
+				{
+					UE_LOG(LogTemp, Error, TEXT("Invalid View Mode for Groom dataflow rendering"));
+				}
+
+				for(int32 GroupIndex = 0; GroupIndex < NumGroups; ++GroupIndex)
+				{
+					const FLinearColor GroupColor = FLinearColor::IntToDistinctColor(GroupIndex, 0.75f, 1.0f, 90.0f);
+					
+					TArray<FLinearColor> VertexColors;
+					VertexColors.Init(GroupColor, VertexPositions[GroupIndex].Num());
+					
+					const FString GeometryName = GroomFacade.GetObjectGroupNames().Last() + TEXT("_") + GroupName + TEXT("_") + FString::FromInt(GroupIndex);
+					const int32 GeometryIndex = RenderCollection.StartGeometryGroup(GeometryName);
+					RenderCollection.AddSurface(MoveTemp(VertexPositions[GroupIndex]),
+						MoveTemp(FaceIndices[GroupIndex]), MoveTemp(VertexNormals[GroupIndex]), MoveTemp(VertexColors));
+					RenderCollection.EndGeometryGroup(GeometryIndex);
+				}
+			}
+		}
+	}
 	
 	Dataflow::FRenderKey FGroomGuidesRenderingCallbacks::RenderKey = { "GuidesRender", FName("FGroomCollection") };
 
@@ -266,7 +398,7 @@ namespace UE::Groom
 		{
 			if(LodIndices[CurveIndex] >= GGroomDataflowGuidesLod)
 			{
-				LodColors[CurveIndex] = FLinearColor::MakeRandomColor();
+				LodColors[CurveIndex] = FLinearColor::IntToDistinctColor(CurveIndex, 0.75f, 1.0f, 90.0f);
 			}
 		}
 		for(int32 CurveIndex = 0, CurveEnd = GuidesFacade.GetNumCurves(); CurveIndex < CurveEnd; ++CurveIndex)
@@ -295,12 +427,19 @@ namespace UE::Groom
 	
 	void FGroomGuidesRenderingCallbacks::Render(GeometryCollection::Facades::FRenderingFacade& RenderCollection, const Dataflow::FGraphRenderingState& State)
 	{
-		RenderGroomCollection<FGroomGuidesFacade>(RenderCollection, State,
-			[this](const FGroomGuidesFacade& GuidesFacade, TArray<FLinearColor>& VertexColors)
+		RenderGroupCollection<FGroomGuidesFacade>(RenderCollection, State,
+			[this](const FGroomGuidesFacade& GuidesFacade, FString& GroupAttribute, FString& GroupName) -> int32
 			{
-				ComputeVertexColors(GuidesFacade, VertexColors);
+				return GetGroupAttribute(GuidesFacade, GroupAttribute, GroupName);
 			});
 	}
+	
+	int32 FGroomGuidesRenderingCallbacks::GetGroupAttribute(const FGroomGuidesFacade& StrandsFacade, FString& GroupAttribute, FString& GroupName) const
+	{
+		GroupAttribute = FGroomGuidesFacade::CurveObjectIndicesAttribute.ToString();
+		GroupName = FString("Group");
+		return StrandsFacade.GetNumObjects();
+	}
 
 	void FGroomGuidesRenderingCallbacks::ComputeVertexColors(const FGroomGuidesFacade& GuidesFacade, TArray<FLinearColor>& VertexColors) const
 	{
@@ -318,13 +457,20 @@ namespace UE::Groom
 	
 	void FGroomStrandsRenderingCallbacks::Render(GeometryCollection::Facades::FRenderingFacade& RenderCollection, const Dataflow::FGraphRenderingState& State)
 	{
-		RenderGroomCollection<FGroomStrandsFacade>(RenderCollection, State,
-			[this](const FGroomStrandsFacade& StrandsFacade, TArray<FLinearColor>& VertexColors)
+		RenderGroupCollection<FGroomStrandsFacade>(RenderCollection, State,
+			[this](const FGroomStrandsFacade& StrandsFacade, FString& GroupAttribute, FString& GroupName) -> int32
 			{
-				ComputeVertexColors(StrandsFacade, VertexColors);
+				return GetGroupAttribute(StrandsFacade, GroupAttribute, GroupName);
 			});
 	}
 
+	int32 FGroomStrandsRenderingCallbacks::GetGroupAttribute(const FGroomStrandsFacade& StrandsFacade, FString& GroupAttribute, FString& GroupName) const
+	{
+		GroupAttribute = FGroomStrandsFacade::CurveObjectIndicesAttribute.ToString();
+		GroupName = FString("Group");
+		return StrandsFacade.GetNumObjects();
+	}
+
 	void FGroomStrandsRenderingCallbacks::ComputeVertexColors(const FGroomStrandsFacade& StrandsFacade, TArray<FLinearColor>& VertexColors) const
 	{}
 
diff --git a/Engine/Plugins/Runtime/HairStrands/Source/HairStrandsDataflow/Public/GroomDataflowRendering.h b/Engine/Plugins/Runtime/HairStrands/Source/HairStrandsDataflow/Public/GroomDataflowRendering.h
--- a/Engine/Plugins/Runtime/HairStrands/Source/HairStrandsDataflow/Public/GroomDataflowRendering.h
+++ b/Engine/Plugins/Runtime/HairStrands/Source/HairStrandsDataflow/Public/GroomDataflowRendering.h
@@ -34,6 +34,9 @@ protected:
 	//¬ End ICallbackInterface interface
 	
 	virtual void ComputeVertexColors(const FGroomStrandsFacade& StrandsFacade, TArray<FLinearColor>& VertexColors) const;
+
+	/** Get the group attribute used to create each geometry groups */
+	virtual int32 GetGroupAttribute(const FGroomStrandsFacade& StrandsFacade, FString& GroupAttribute, FString& GroupName) const;
 };
 
 /** Guides rendering callback for the dataflow editor */
@@ -58,6 +61,9 @@ protected:
 	virtual void Render(GeometryCollection::Facades::FRenderingFacade& RenderCollection, const Dataflow::FGraphRenderingState& State) override;
 	
 	virtual void ComputeVertexColors(const FGroomGuidesFacade& GuidesFacade, TArray<FLinearColor>& VertexColors) const;
+	
+	/** Get the group attribute used to create each geometry groups */
+	virtual int32 GetGroupAttribute(const FGroomGuidesFacade& StrandsFacade, FString& GroupAttribute, FString& GroupName) const;
     	
 	//¬ End ICallbackInterface interface
 };
