[HairStrands]
* Moved dll storage to methods/variables instead of on types. This is needed to get merged modules working

diff --git a/Engine/Plugins/Runtime/HairStrands/Source/HairStrandsCore/Public/GroomAsset.h b/Engine/Plugins/Runtime/HairStrands/Source/HairStrandsCore/Public/GroomAsset.h
--- a/Engine/Plugins/Runtime/HairStrands/Source/HairStrandsCore/Public/GroomAsset.h
+++ b/Engine/Plugins/Runtime/HairStrands/Source/HairStrandsCore/Public/GroomAsset.h
@@ -27,6 +27,7 @@
 
 #include "GroomAsset.generated.h"
 
+#define UE_API HAIRSTRANDSCORE_API
 
 class UGroomBindingAsset;
 class UAssetUserData;
@@ -44,7 +45,7 @@ enum class EHairGroupInfoFlags : uint8
 };
 
 USTRUCT(BlueprintType)
-struct HAIRSTRANDSCORE_API FHairGroupLODInfo
+struct FHairGroupLODInfo
 {
 	GENERATED_BODY()
 
@@ -56,7 +57,7 @@ struct HAIRSTRANDSCORE_API FHairGroupLODInfo
 };
 
 USTRUCT(BlueprintType)
-struct HAIRSTRANDSCORE_API FHairGroupInfo
+struct FHairGroupInfo
 {
 	GENERATED_BODY()
 
@@ -92,7 +93,7 @@ struct HAIRSTRANDSCORE_API FHairGroupInfo
 };
 
 USTRUCT(BlueprintType)
-struct HAIRSTRANDSCORE_API FHairGroupsMaterial
+struct FHairGroupsMaterial
 {
 	GENERATED_BODY()
 
@@ -421,7 +422,7 @@ struct FHairDescriptionGroups
 };
 
 USTRUCT(BlueprintType)
-struct HAIRSTRANDSCORE_API FHairGroupInfoWithVisibility : public FHairGroupInfo
+struct FHairGroupInfoWithVisibility : public FHairGroupInfo
 {
 	GENERATED_BODY()
 
@@ -451,8 +452,8 @@ struct FHairVertexFactoryTypesPerMaterialData
 /**
  * Implements an asset that can be used to store hair strands
  */
-UCLASS(BlueprintType, AutoExpandCategories = ("HairRendering", "HairPhysics", "HairInterpolation", "HairDataflow"), hidecategories = (Object, "Hidden"))
-class HAIRSTRANDSCORE_API UGroomAsset : public UObject, public IInterface_AssetUserData, public IDataflowContentOwner
+UCLASS(BlueprintType, AutoExpandCategories = ("HairRendering", "HairPhysics", "HairInterpolation", "HairDataflow"), hidecategories = (Object, "Hidden"), MinimalAPI)
+class UGroomAsset : public UObject, public IInterface_AssetUserData, public IDataflowContentOwner
 {
 	GENERATED_UCLASS_BODY()
 
@@ -561,7 +562,7 @@ private:
 public:
 
 	/** Static member to retrieve the member name */
-	static FName GetDataflowSettingsMemberName();
+	static UE_API FName GetDataflowSettingsMemberName();
 	
 	/** Return the Dataflow settings associated to this groom asset if any, const version. */
 	const FGroomDataflowSettings& GetDataflowSettings() const { return DataflowSettings; }
@@ -572,110 +573,110 @@ public:
 	/** Set the Dataflow settings associated to this groom asset if any */
     void SetDataflowSettings(const FGroomDataflowSettings& InDataflowSettings) { DataflowSettings = InDataflowSettings; }
 
-	static FName GetHairGroupsRenderingMemberName();
-	UFUNCTION(BlueprintGetter) TArray<FHairGroupsRendering>& GetHairGroupsRendering();
-	UFUNCTION(BlueprintSetter) void SetHairGroupsRendering(const TArray<FHairGroupsRendering>& In);
-	const TArray<FHairGroupsRendering>& GetHairGroupsRendering() const;
-
-	static FName GetHairGroupsPhysicsMemberName();
-	UFUNCTION(BlueprintGetter) TArray<FHairGroupsPhysics>& GetHairGroupsPhysics();
-	UFUNCTION(BlueprintSetter) void SetHairGroupsPhysics(const TArray<FHairGroupsPhysics>& In);
-	const TArray<FHairGroupsPhysics>& GetHairGroupsPhysics() const;
-
-	static FName GetHairGroupsInterpolationMemberName();
-	UFUNCTION(BlueprintGetter) TArray<FHairGroupsInterpolation>& GetHairGroupsInterpolation();
-	UFUNCTION(BlueprintSetter) void SetHairGroupsInterpolation(const TArray<FHairGroupsInterpolation>& In);
-	const TArray<FHairGroupsInterpolation>& GetHairGroupsInterpolation() const;
-
-	static FName GetHairGroupsLODMemberName();
-	UFUNCTION(BlueprintGetter) TArray<FHairGroupsLOD>& GetHairGroupsLOD();
-	UFUNCTION(BlueprintSetter) void SetHairGroupsLOD(const TArray<FHairGroupsLOD>& In);
-	const TArray<FHairGroupsLOD>& GetHairGroupsLOD() const;
-
-	static FName GetHairGroupsCardsMemberName();
-	UFUNCTION(BlueprintGetter) TArray<FHairGroupsCardsSourceDescription>& GetHairGroupsCards();
-	UFUNCTION(BlueprintSetter) void SetHairGroupsCards(const TArray<FHairGroupsCardsSourceDescription>& In);
-	const TArray<FHairGroupsCardsSourceDescription>& GetHairGroupsCards() const;
-
-	static FName GetHairGroupsMeshesMemberName();
-	UFUNCTION(BlueprintGetter) TArray<FHairGroupsMeshesSourceDescription>& GetHairGroupsMeshes();
-	UFUNCTION(BlueprintSetter) void SetHairGroupsMeshes(const TArray<FHairGroupsMeshesSourceDescription>& In);
-	const TArray<FHairGroupsMeshesSourceDescription>& GetHairGroupsMeshes() const;
-
-	static FName GetHairGroupsMaterialsMemberName();
-	UFUNCTION(BlueprintGetter) TArray<FHairGroupsMaterial>& GetHairGroupsMaterials();
-	UFUNCTION(BlueprintSetter) void SetHairGroupsMaterials(const TArray<FHairGroupsMaterial>& In);
-	const TArray<FHairGroupsMaterial>& GetHairGroupsMaterials() const;
-
-	static FName GetEnableGlobalInterpolationMemberName();
-	UFUNCTION(BlueprintGetter) bool GetEnableGlobalInterpolation() const;
-	UFUNCTION(BlueprintSetter) void SetEnableGlobalInterpolation(bool In);
-
-	static FName GetEnableSimulationCacheMemberName();
-	UFUNCTION(BlueprintGetter) bool GetEnableSimulationCache() const;
-	UFUNCTION(BlueprintSetter) void SetEnableSimulationCache(bool In);
-
-	static FName GetHairInterpolationTypeMemberName();
-	UFUNCTION(BlueprintGetter) EGroomInterpolationType GetHairInterpolationType() const;
-	UFUNCTION(BlueprintSetter) void SetHairInterpolationType(EGroomInterpolationType In);
-
-	static FName GetRiggedSkeletalMeshMemberName();
-	UFUNCTION(BlueprintGetter) USkeletalMesh* GetRiggedSkeletalMesh() const;
-	UFUNCTION(BlueprintSetter) void SetRiggedSkeletalMesh(USkeletalMesh* In);
-
-	static FName GetDeformedGroupSectionsMemberName();
-	UFUNCTION(BlueprintGetter) TArray<int32>& GetDeformedGroupSections();
-	UFUNCTION(BlueprintSetter) void SetDeformedGroupSections(const TArray<int32>& In);
-	const TArray<int32>& GetDeformedGroupSections() const;
-
-	static FName GetMinLODMemberName();
-	FPerPlatformInt GetMinLOD() const;
-	void SetMinLOD(FPerPlatformInt In);
-
-	static FName GetDisableBelowMinLodStrippingMemberName();
-	FPerPlatformBool GetDisableBelowMinLodStripping() const;
-	void SetDisableBelowMinLodStripping(FPerPlatformBool In);
-
-	static FName GetEffectiveLODBiasMemberName();
-	TArray<float>& GetEffectiveLODBias();
-	void SetEffectiveLODBias(const TArray<float>& In);
-	const TArray<float>& GetEffectiveLODBias() const;
-
-	static FName GetHairGroupsPlatformDataMemberName();
-	TArray<FHairGroupPlatformData>& GetHairGroupsPlatformData();
-	const TArray<FHairGroupPlatformData>& GetHairGroupsPlatformData() const;
-	void SetHairGroupsPlatformData(const TArray<FHairGroupPlatformData>& In);
-
-	static FName GetHairGroupsInfoMemberName();
-	TArray<FHairGroupInfoWithVisibility>& GetHairGroupsInfo();
-	const TArray<FHairGroupInfoWithVisibility>& GetHairGroupsInfo() const;
-	void SetHairGroupsInfo(const TArray<FHairGroupInfoWithVisibility>& In);
-
-	const TArray<FHairGroupResources>& GetHairGroupsResources() const;
-	TArray<FHairGroupResources>& GetHairGroupsResources();
-
-	static FName GetLODModeMemberName();
-	EGroomLODMode GetLODMode() const;
-
-	static FName GetAutoLODBiasMemberName();
+	static UE_API FName GetHairGroupsRenderingMemberName();
+	UFUNCTION(BlueprintGetter) UE_API TArray<FHairGroupsRendering>& GetHairGroupsRendering();
+	UFUNCTION(BlueprintSetter) UE_API void SetHairGroupsRendering(const TArray<FHairGroupsRendering>& In);
+	UE_API const TArray<FHairGroupsRendering>& GetHairGroupsRendering() const;
+
+	static UE_API FName GetHairGroupsPhysicsMemberName();
+	UFUNCTION(BlueprintGetter) UE_API TArray<FHairGroupsPhysics>& GetHairGroupsPhysics();
+	UFUNCTION(BlueprintSetter) UE_API void SetHairGroupsPhysics(const TArray<FHairGroupsPhysics>& In);
+	UE_API const TArray<FHairGroupsPhysics>& GetHairGroupsPhysics() const;
+
+	static UE_API FName GetHairGroupsInterpolationMemberName();
+	UFUNCTION(BlueprintGetter) UE_API TArray<FHairGroupsInterpolation>& GetHairGroupsInterpolation();
+	UFUNCTION(BlueprintSetter) UE_API void SetHairGroupsInterpolation(const TArray<FHairGroupsInterpolation>& In);
+	UE_API const TArray<FHairGroupsInterpolation>& GetHairGroupsInterpolation() const;
+
+	static UE_API FName GetHairGroupsLODMemberName();
+	UFUNCTION(BlueprintGetter) UE_API TArray<FHairGroupsLOD>& GetHairGroupsLOD();
+	UFUNCTION(BlueprintSetter) UE_API void SetHairGroupsLOD(const TArray<FHairGroupsLOD>& In);
+	UE_API const TArray<FHairGroupsLOD>& GetHairGroupsLOD() const;
+
+	static UE_API FName GetHairGroupsCardsMemberName();
+	UFUNCTION(BlueprintGetter) UE_API TArray<FHairGroupsCardsSourceDescription>& GetHairGroupsCards();
+	UFUNCTION(BlueprintSetter) UE_API void SetHairGroupsCards(const TArray<FHairGroupsCardsSourceDescription>& In);
+	UE_API const TArray<FHairGroupsCardsSourceDescription>& GetHairGroupsCards() const;
+
+	static UE_API FName GetHairGroupsMeshesMemberName();
+	UFUNCTION(BlueprintGetter) UE_API TArray<FHairGroupsMeshesSourceDescription>& GetHairGroupsMeshes();
+	UFUNCTION(BlueprintSetter) UE_API void SetHairGroupsMeshes(const TArray<FHairGroupsMeshesSourceDescription>& In);
+	UE_API const TArray<FHairGroupsMeshesSourceDescription>& GetHairGroupsMeshes() const;
+
+	static UE_API FName GetHairGroupsMaterialsMemberName();
+	UFUNCTION(BlueprintGetter) UE_API TArray<FHairGroupsMaterial>& GetHairGroupsMaterials();
+	UFUNCTION(BlueprintSetter) UE_API void SetHairGroupsMaterials(const TArray<FHairGroupsMaterial>& In);
+	UE_API const TArray<FHairGroupsMaterial>& GetHairGroupsMaterials() const;
+
+	static UE_API FName GetEnableGlobalInterpolationMemberName();
+	UFUNCTION(BlueprintGetter) UE_API bool GetEnableGlobalInterpolation() const;
+	UFUNCTION(BlueprintSetter) UE_API void SetEnableGlobalInterpolation(bool In);
+
+	static UE_API FName GetEnableSimulationCacheMemberName();
+	UFUNCTION(BlueprintGetter) UE_API bool GetEnableSimulationCache() const;
+	UFUNCTION(BlueprintSetter) UE_API void SetEnableSimulationCache(bool In);
+
+	static UE_API FName GetHairInterpolationTypeMemberName();
+	UFUNCTION(BlueprintGetter) UE_API EGroomInterpolationType GetHairInterpolationType() const;
+	UFUNCTION(BlueprintSetter) UE_API void SetHairInterpolationType(EGroomInterpolationType In);
+
+	static UE_API FName GetRiggedSkeletalMeshMemberName();
+	UFUNCTION(BlueprintGetter) UE_API USkeletalMesh* GetRiggedSkeletalMesh() const;
+	UFUNCTION(BlueprintSetter) UE_API void SetRiggedSkeletalMesh(USkeletalMesh* In);
+
+	static UE_API FName GetDeformedGroupSectionsMemberName();
+	UFUNCTION(BlueprintGetter) UE_API TArray<int32>& GetDeformedGroupSections();
+	UFUNCTION(BlueprintSetter) UE_API void SetDeformedGroupSections(const TArray<int32>& In);
+	UE_API const TArray<int32>& GetDeformedGroupSections() const;
+
+	static UE_API FName GetMinLODMemberName();
+	UE_API FPerPlatformInt GetMinLOD() const;
+	UE_API void SetMinLOD(FPerPlatformInt In);
+
+	static UE_API FName GetDisableBelowMinLodStrippingMemberName();
+	UE_API FPerPlatformBool GetDisableBelowMinLodStripping() const;
+	UE_API void SetDisableBelowMinLodStripping(FPerPlatformBool In);
+
+	static UE_API FName GetEffectiveLODBiasMemberName();
+	UE_API TArray<float>& GetEffectiveLODBias();
+	UE_API void SetEffectiveLODBias(const TArray<float>& In);
+	UE_API const TArray<float>& GetEffectiveLODBias() const;
+
+	static UE_API FName GetHairGroupsPlatformDataMemberName();
+	UE_API TArray<FHairGroupPlatformData>& GetHairGroupsPlatformData();
+	UE_API const TArray<FHairGroupPlatformData>& GetHairGroupsPlatformData() const;
+	UE_API void SetHairGroupsPlatformData(const TArray<FHairGroupPlatformData>& In);
+
+	static UE_API FName GetHairGroupsInfoMemberName();
+	UE_API TArray<FHairGroupInfoWithVisibility>& GetHairGroupsInfo();
+	UE_API const TArray<FHairGroupInfoWithVisibility>& GetHairGroupsInfo() const;
+	UE_API void SetHairGroupsInfo(const TArray<FHairGroupInfoWithVisibility>& In);
+
+	UE_API const TArray<FHairGroupResources>& GetHairGroupsResources() const;
+	UE_API TArray<FHairGroupResources>& GetHairGroupsResources();
+
+	static UE_API FName GetLODModeMemberName();
+	UE_API EGroomLODMode GetLODMode() const;
+
+	static UE_API FName GetAutoLODBiasMemberName();
 	// Return the asset Auto LOD Bias
-	float GetAutoLODBias() const;
+	UE_API float GetAutoLODBias() const;
 
 	// Return the group Auto LOD Bias, which combines both the asset's bias, and the group bias
-	float GetAutoLODBias(int32 InGroupIndex) const;
+	UE_API float GetAutoLODBias(int32 InGroupIndex) const;
 
 public:
 
 	//~ Begin UObject Interface.
-	virtual void GetResourceSizeEx(FResourceSizeEx& CumulativeResourceSize) override;
-	virtual void PostLoad() override;
-	virtual void BeginDestroy() override;
-	virtual void Serialize(FArchive& Ar) override;
+	UE_API virtual void GetResourceSizeEx(FResourceSizeEx& CumulativeResourceSize) override;
+	UE_API virtual void PostLoad() override;
+	UE_API virtual void BeginDestroy() override;
+	UE_API virtual void Serialize(FArchive& Ar) override;
 
 	//~ Begin IDataflowContentOwner interface 
-	virtual TObjectPtr<UDataflowBaseContent> CreateDataflowContent() override;
-	virtual void WriteDataflowContent(const TObjectPtr<UDataflowBaseContent>& DataflowContent) const override;
-	virtual void ReadDataflowContent(const TObjectPtr<UDataflowBaseContent>& DataflowContent) override;
+	UE_API virtual TObjectPtr<UDataflowBaseContent> CreateDataflowContent() override;
+	UE_API virtual void WriteDataflowContent(const TObjectPtr<UDataflowBaseContent>& DataflowContent) const override;
+	UE_API virtual void ReadDataflowContent(const TObjectPtr<UDataflowBaseContent>& DataflowContent) override;
 	//~ End IDataflowContentOwner interface
 
 #if WITH_EDITOR
@@ -684,8 +685,8 @@ public:
 	FOnGroomAsyncLoadFinished& GetOnGroomAsyncLoadFinished() { return OnGroomAsyncLoadFinished; }
 
 	/**  Part of Uobject interface  */
-	virtual bool Modify(bool bAlwaysMarkDirty = true) override;
-	virtual void PostEditChangeProperty(FPropertyChangedEvent& PropertyChangedEvent) override;
+	UE_API virtual bool Modify(bool bAlwaysMarkDirty = true) override;
+	UE_API virtual void PostEditChangeProperty(FPropertyChangedEvent& PropertyChangedEvent) override;
 
 #endif // WITH_EDITOR
 
@@ -696,81 +697,81 @@ public:
 	TObjectPtr<class UAssetImportData> AssetImportData;
 
 	/** Retrieve the asset tags*/
-	virtual void GetAssetRegistryTags(FAssetRegistryTagsContext Context) const override;
+	UE_API virtual void GetAssetRegistryTags(FAssetRegistryTagsContext Context) const override;
 	UE_DEPRECATED(5.4, "Implement the version that takes FAssetRegistryTagsContext instead.")
-	virtual void GetAssetRegistryTags(TArray<FAssetRegistryTag>& OutTags) const override;
+	UE_API virtual void GetAssetRegistryTags(TArray<FAssetRegistryTag>& OutTags) const override;
 
 	/** Part of Uobject interface */
-	virtual void PostInitProperties() override;
+	UE_API virtual void PostInitProperties() override;
 
 #endif // WITH_EDITORONLY_DATA
 
 	bool IsValid() const { return bIsInitialized; }
 
 	// Helper functions for setting options on all hair groups
-	void SetStableRasterization(bool bEnable);
-	void SetScatterSceneLighting(bool Enable);
-	void SetHairWidth(float Width);
+	UE_API void SetStableRasterization(bool bEnable);
+	UE_API void SetScatterSceneLighting(bool Enable);
+	UE_API void SetHairWidth(float Width);
 
 	/** Initialize/Update/Release resources. */
-	void InitResources();
-	void InitGuideResources();
-	void InitStrandsResources();
-	void InitCardsResources();
-	void InitMeshesResources();
+	UE_API void InitResources();
+	UE_API void InitGuideResources();
+	UE_API void InitStrandsResources();
+	UE_API void InitCardsResources();
+	UE_API void InitMeshesResources();
 #if WITH_EDITOR
-	void UpdateResource();
+	UE_API void UpdateResource();
 #endif
-	void ReleaseResource();
-	void ReleaseGuidesResource(uint32 GroupIndex);
-	void ReleaseStrandsResource(uint32 GroupIndex);
-	void ReleaseCardsResource(uint32 GroupIndex);
-	void ReleaseMeshesResource(uint32 GroupIndex);
+	UE_API void ReleaseResource();
+	UE_API void ReleaseGuidesResource(uint32 GroupIndex);
+	UE_API void ReleaseStrandsResource(uint32 GroupIndex);
+	UE_API void ReleaseCardsResource(uint32 GroupIndex);
+	UE_API void ReleaseMeshesResource(uint32 GroupIndex);
 
-	void SetNumGroup(uint32 InGroupCount, bool bResetGroupData=true, bool bResetOtherData=true);
-	void ClearNumGroup(uint32 InGroupCount);
+	UE_API void SetNumGroup(uint32 InGroupCount, bool bResetGroupData=true, bool bResetOtherData=true);
+	UE_API void ClearNumGroup(uint32 InGroupCount);
 
-	bool AreGroupsValid() const;
-	int32 GetNumHairGroups() const;
+	UE_API bool AreGroupsValid() const;
+	UE_API int32 GetNumHairGroups() const;
 
-	int32 GetLODCount() const;
+	UE_API int32 GetLODCount() const;
 #if WITH_EDITORONLY_DATA
-	void StripLODs(const TArray<int32>& LODsToKeep, bool bRebuildResources);
+	UE_API void StripLODs(const TArray<int32>& LODsToKeep, bool bRebuildResources);
 #endif // WITH_EDITORONLY_DATA
 
 	/** Debug data for derived asset generation (strands textures, ...). */
-	bool HasDebugData() const;
-	void CreateDebugData();
+	UE_API bool HasDebugData() const;
+	UE_API void CreateDebugData();
 
 	/** Returns true if the asset has the HairDescription needed to recompute its groom data */
-	bool CanRebuildFromDescription() const;
+	UE_API bool CanRebuildFromDescription() const;
 
 	//~ Begin IInterface_AssetUserData Interface
-	virtual void AddAssetUserData(UAssetUserData* InUserData) override;
-	virtual void RemoveUserDataOfClass(TSubclassOf<UAssetUserData> InUserDataClass) override;
-	virtual UAssetUserData* GetAssetUserDataOfClass(TSubclassOf<UAssetUserData> InUserDataClass) override;
-	virtual const TArray<UAssetUserData*>* GetAssetUserDataArray() const override;
+	UE_API virtual void AddAssetUserData(UAssetUserData* InUserData) override;
+	UE_API virtual void RemoveUserDataOfClass(TSubclassOf<UAssetUserData> InUserDataClass) override;
+	UE_API virtual UAssetUserData* GetAssetUserDataOfClass(TSubclassOf<UAssetUserData> InUserDataClass) override;
+	UE_API virtual const TArray<UAssetUserData*>* GetAssetUserDataArray() const override;
 	//~ End IInterface_AssetUserData Interface
 
-	EGroomGeometryType GetGeometryType(int32 GroupIndex, int32 LODIndex) const;
-	EGroomBindingType GetBindingType(int32 GroupIndex, int32 LODIndex) const;
-	bool IsVisible(int32 GroupIndex, int32 LODIndex) const;
-	bool IsSimulationEnable(int32 GroupIndex, int32 LODIndex) const;
-	bool IsSimulationEnable() const;
-	bool IsDeformationEnable(int32 GroupIndex) const;
-	bool IsGlobalInterpolationEnable(int32 GroupIndex, int32 LODIndex) const;
-	bool NeedsInterpolationData(int32 GroupIndex) const;
-	bool NeedsInterpolationData() const;
+	UE_API EGroomGeometryType GetGeometryType(int32 GroupIndex, int32 LODIndex) const;
+	UE_API EGroomBindingType GetBindingType(int32 GroupIndex, int32 LODIndex) const;
+	UE_API bool IsVisible(int32 GroupIndex, int32 LODIndex) const;
+	UE_API bool IsSimulationEnable(int32 GroupIndex, int32 LODIndex) const;
+	UE_API bool IsSimulationEnable() const;
+	UE_API bool IsDeformationEnable(int32 GroupIndex) const;
+	UE_API bool IsGlobalInterpolationEnable(int32 GroupIndex, int32 LODIndex) const;
+	UE_API bool NeedsInterpolationData(int32 GroupIndex) const;
+	UE_API bool NeedsInterpolationData() const;
 
-	void UpdateHairGroupsInfo();
-	bool HasGeometryType(EGroomGeometryType Type) const;
-	bool HasGeometryType(uint32 GroupIndex, EGroomGeometryType Type) const;
+	UE_API void UpdateHairGroupsInfo();
+	UE_API bool HasGeometryType(EGroomGeometryType Type) const;
+	UE_API bool HasGeometryType(uint32 GroupIndex, EGroomGeometryType Type) const;
 
 	/** Used for PSO precaching of used materials and vertex factories */
-	TArray<FHairVertexFactoryTypesPerMaterialData> CollectVertexFactoryTypesPerMaterialData(EShaderPlatform ShaderPlatform);
+	UE_API TArray<FHairVertexFactoryTypesPerMaterialData> CollectVertexFactoryTypesPerMaterialData(EShaderPlatform ShaderPlatform);
 
 	/** Helper function to return the asset path name, optionally joined with the LOD index if LODIndex > -1. */
-	FName GetAssetPathName(int32 LODIndex = -1);
+	UE_API FName GetAssetPathName(int32 LODIndex = -1);
 	uint32 GetAssetHash() const { return AssetNameHash; }
 
 //private :
@@ -779,11 +780,11 @@ public:
 	FOnGroomAssetResourcesChanged OnGroomAssetResourcesChanged;
 	FOnGroomAsyncLoadFinished OnGroomAsyncLoadFinished;
 
-	void MarkMaterialsHasChanged();
+	UE_API void MarkMaterialsHasChanged();
 
-	void RecreateResources();
-	void ChangeFeatureLevel(ERHIFeatureLevel::Type PendingFeatureLevel);
-	void ChangePlatformLevel(ERHIFeatureLevel::Type PendingFeatureLevel);
+	UE_API void RecreateResources();
+	UE_API void ChangeFeatureLevel(ERHIFeatureLevel::Type PendingFeatureLevel);
+	UE_API void ChangePlatformLevel(ERHIFeatureLevel::Type PendingFeatureLevel);
 #endif
 
 	/** Array of user data stored with the asset */
@@ -791,13 +792,13 @@ public:
 	TArray<TObjectPtr<UAssetUserData>> AssetUserData;
 
 	/* Return the material slot index corresponding to the material name */
-	int32 GetMaterialIndex(FName MaterialSlotName) const;
-	bool IsMaterialSlotNameValid(FName MaterialSlotName) const;
-	bool IsMaterialUsed(int32 MaterialIndex) const;
-	TArray<FName> GetMaterialSlotNames() const;
+	UE_API int32 GetMaterialIndex(FName MaterialSlotName) const;
+	UE_API bool IsMaterialSlotNameValid(FName MaterialSlotName) const;
+	UE_API bool IsMaterialUsed(int32 MaterialIndex) const;
+	UE_API TArray<FName> GetMaterialSlotNames() const;
 
-	bool BuildCardsData();
-	bool BuildMeshesData();
+	UE_API bool BuildCardsData();
+	UE_API bool BuildMeshesData();
 
 	enum EClassDataStripFlag : uint8
 	{
@@ -808,10 +809,10 @@ public:
 		CDSF_MeshesStripped = 16
 	};
 
-	uint8 GenerateClassStripFlags(FArchive& Ar);
+	UE_API uint8 GenerateClassStripFlags(FArchive& Ar);
 
 	/** Enable the simulation cache and recompute the cached derived data */
-	void ValidateSimulationCache();
+	UE_API void ValidateSimulationCache();
 
 private:
 	void ApplyStripFlags(uint8 StripFlags, const class ITargetPlatform* CookTarget);
@@ -849,28 +850,28 @@ public:
 	};
 
 	/** Commits a HairDescription to buffer for serialization */
-	void CommitHairDescription(FHairDescription&& HairDescription, EHairDescriptionType Type);
-	FHairDescription GetHairDescription() const;
+	UE_API void CommitHairDescription(FHairDescription&& HairDescription, EHairDescriptionType Type);
+	UE_API FHairDescription GetHairDescription() const;
 
 	/** Get/Build render & guides data based on the hair description and interpolation settings */
-	bool GetHairStrandsDatas(const int32 GroupIndex, FHairStrandsDatas& OutStrandsData, FHairStrandsDatas& OutGuidesData);
-	bool GetHairCardsGuidesDatas(const int32 GroupIndex, const int32 LODIndex, FHairStrandsDatas& OutCardsGuidesData);
+	UE_API bool GetHairStrandsDatas(const int32 GroupIndex, FHairStrandsDatas& OutStrandsData, FHairStrandsDatas& OutGuidesData);
+	UE_API bool GetHairCardsGuidesDatas(const int32 GroupIndex, const int32 LODIndex, FHairStrandsDatas& OutCardsGuidesData);
 
 	/** Caches the computed (group) groom data with the given build settings from/to the Derived Data Cache, building it if needed.
 	 *  This function assumes the interpolation settings are properly populated, as they will be used to build the asset.
 	 */
-	bool CacheDerivedDatas();
-	bool CacheDerivedData(uint32 GroupIndex);
-	bool CacheStrandsData(uint32 GroupIndex, FString& OutDerivedDataKey);
-	bool CacheCardsData(uint32 GroupIndex, const FString& StrandsKey);
-	bool CacheMeshesData(uint32 GroupIndex);
-
-	FString GetDerivedDataKey(bool bUseCacheKey=false);
-	FString GetDerivedDataKeyForCards(uint32 GroupIt, const FString& StrandsKey);
-	FString GetDerivedDataKeyForStrands(uint32 GroupIndex);
-	FString GetDerivedDataKeyForMeshes(uint32 GroupIndex);
-
-	const FHairDescriptionGroups& GetHairDescriptionGroups();
+	UE_API bool CacheDerivedDatas();
+	UE_API bool CacheDerivedData(uint32 GroupIndex);
+	UE_API bool CacheStrandsData(uint32 GroupIndex, FString& OutDerivedDataKey);
+	UE_API bool CacheCardsData(uint32 GroupIndex, const FString& StrandsKey);
+	UE_API bool CacheMeshesData(uint32 GroupIndex);
+
+	UE_API FString GetDerivedDataKey(bool bUseCacheKey=false);
+	UE_API FString GetDerivedDataKeyForCards(uint32 GroupIt, const FString& StrandsKey);
+	UE_API FString GetDerivedDataKeyForStrands(uint32 GroupIndex);
+	UE_API FString GetDerivedDataKeyForMeshes(uint32 GroupIndex);
+
+	UE_API const FHairDescriptionGroups& GetHairDescriptionGroups();
 private:
 	
 	bool IsFullyCached();
@@ -896,7 +897,7 @@ private:
 
 #if WITH_EDITOR
 public:
-	void UpdateCachedSettings();
+	UE_API void UpdateCachedSettings();
 private:
 	void SavePendingProceduralAssets();
 
@@ -944,3 +945,5 @@ struct FGroomAssetMemoryStats
 	uint32 GetTotalCPUSize() const;
 	uint32 GetTotalGPUSize() const;
 };
+
+#undef UE_API
diff --git a/Engine/Plugins/Runtime/HairStrands/Source/HairStrandsCore/Public/GroomAssetDataflow.h b/Engine/Plugins/Runtime/HairStrands/Source/HairStrandsCore/Public/GroomAssetDataflow.h
--- a/Engine/Plugins/Runtime/HairStrands/Source/HairStrandsCore/Public/GroomAssetDataflow.h
+++ b/Engine/Plugins/Runtime/HairStrands/Source/HairStrandsCore/Public/GroomAssetDataflow.h
@@ -1,4 +1,4 @@
-﻿// Copyright Epic Games, Inc. All Rights Reserved.
+// Copyright Epic Games, Inc. All Rights Reserved.
 
 #pragma once
 
@@ -7,6 +7,8 @@
 #include "UObject/ObjectMacros.h"
 #include "GroomAssetDataflow.generated.h"
 
+#define UE_API HAIRSTRANDSCORE_API
+
 class UGroomBindingAsset;
 class UAnimationAsset;
 
@@ -14,27 +16,27 @@ class UAnimationAsset;
  * Dataflow content owning dataflow and binding assets that will be used to evaluate the graph
  */
 UCLASS()
-class HAIRSTRANDSCORE_API  UDataflowGroomContent : public UDataflowSkeletalContent
+class UDataflowGroomContent : public UDataflowSkeletalContent
 {
 	GENERATED_BODY()
 
 public:
-	UDataflowGroomContent();
+	UE_API UDataflowGroomContent();
 	virtual ~UDataflowGroomContent() override {}
 	
 #if WITH_EDITOR
-	virtual void PostEditChangeProperty(FPropertyChangedEvent& PropertyChangedEvent) override;
+	UE_API virtual void PostEditChangeProperty(FPropertyChangedEvent& PropertyChangedEvent) override;
 #endif //if WITH_EDITOR
 
 	/** Data flow binding asset accessors */
-	void SetBindingAsset(const TObjectPtr<UGroomBindingAsset>& InBindingAsset);
+	UE_API void SetBindingAsset(const TObjectPtr<UGroomBindingAsset>& InBindingAsset);
 	const TObjectPtr<UGroomBindingAsset>& GetBindingAsset() const { return BindingAsset; }
 
 	//~ UObject interface
-	static void AddReferencedObjects(UObject* InThis, FReferenceCollector& Collector);
+	static UE_API void AddReferencedObjects(UObject* InThis, FReferenceCollector& Collector);
 
 	/** Set all the preview actor exposed properties */
-	virtual void SetActorProperties(TObjectPtr<AActor>& PreviewActor) const override;
+	UE_API virtual void SetActorProperties(TObjectPtr<AActor>& PreviewActor) const override;
 
 protected:
 
@@ -44,17 +46,17 @@ protected:
 };
 
 USTRUCT()
-struct HAIRSTRANDSCORE_API FGroomDataflowSettings
+struct FGroomDataflowSettings
 {
 	GENERATED_BODY()
 
-	FGroomDataflowSettings();
+	UE_API FGroomDataflowSettings();
 
 	/** Return the Dataflow asset name. */
-	static FName GetDataflowAssetMemberName();
+	static UE_API FName GetDataflowAssetMemberName();
 
 	/** Return the Dataflow terminal name. */
-	static FName GetDataflowTerminalMemberName();
+	static UE_API FName GetDataflowTerminalMemberName();
 
 	/** Return the Dataflow asset associated to this groom asset if any, const version. */
 	UDataflow* GetDataflowAsset() const { return DataflowAsset; }
@@ -94,21 +96,21 @@ struct HAIRSTRANDSCORE_API FGroomDataflowSettings
 	void SetRestCollection(FManagedArrayCollection* InRestCollection) { RestCollection = TSharedPtr<FManagedArrayCollection, ESPMode::ThreadSafe>(InRestCollection); }
 
 	/** Custom serialization function for the ManagedArrayCollection */
-	bool Serialize(FArchive& Ar);
+	UE_API bool Serialize(FArchive& Ar);
 	
 #if WITH_EDITORONLY_DATA
 
 	/** Preview scene binding asset setter */
-	void SetPreviewBindingAsset(UGroomBindingAsset* BindingAsset);
+	UE_API void SetPreviewBindingAsset(UGroomBindingAsset* BindingAsset);
 
 	/** Preview scene binding asset setter */
-	UGroomBindingAsset* GetPreviewBindingAsset() const;
+	UE_API UGroomBindingAsset* GetPreviewBindingAsset() const;
 
 	/** Preview scene animation setter */
-	void SetPreviewAnimationAsset(UAnimationAsset* AnimationAsset);
+	UE_API void SetPreviewAnimationAsset(UAnimationAsset* AnimationAsset);
 
 	/** Preview scene animation getter */
-	UAnimationAsset* GetPreviewAnimationAsset() const;
+	UE_API UAnimationAsset* GetPreviewAnimationAsset() const;
 
 #endif
 
@@ -152,4 +154,6 @@ template<> struct TStructOpsTypeTraits<FGroomDataflowSettings> : public TStructO
 	{
 		WithSerializer = true,
 	};
-};
\ No newline at end of file
+};
+
+#undef UE_API
diff --git a/Engine/Plugins/Runtime/HairStrands/Source/HairStrandsCore/Public/GroomCache.h b/Engine/Plugins/Runtime/HairStrands/Source/HairStrandsCore/Public/GroomCache.h
--- a/Engine/Plugins/Runtime/HairStrands/Source/HairStrandsCore/Public/GroomCache.h
+++ b/Engine/Plugins/Runtime/HairStrands/Source/HairStrandsCore/Public/GroomCache.h
@@ -10,52 +10,54 @@
 #include "Serialization/BulkData.h"
 #include "GroomCache.generated.h"
 
+#define UE_API HAIRSTRANDSCORE_API
+
 struct FGroomCacheChunk;
 
 /**
  * Implements an asset that is used to store an animated groom
  */
-UCLASS(BlueprintType)
-class HAIRSTRANDSCORE_API UGroomCache : public UObject, public IInterface_AssetUserData, public IChaosCacheData
+UCLASS(BlueprintType, MinimalAPI)
+class UGroomCache : public UObject, public IInterface_AssetUserData, public IChaosCacheData
 {
 	GENERATED_BODY()
 
 public:
 
 	//~ Begin UObject Interface.
-	virtual void Serialize(FArchive& Ar) override;
-	virtual void GetAssetRegistryTags(FAssetRegistryTagsContext Context) const override;
+	UE_API virtual void Serialize(FArchive& Ar) override;
+	UE_API virtual void GetAssetRegistryTags(FAssetRegistryTagsContext Context) const override;
 	UE_DEPRECATED(5.4, "Implement the version that takes FAssetRegistryTagsContext instead.")
-	virtual void GetAssetRegistryTags(TArray<FAssetRegistryTag>& OutTags) const override;
+	UE_API virtual void GetAssetRegistryTags(TArray<FAssetRegistryTag>& OutTags) const override;
 	//~ End UObject Interface.
 
-	void Initialize(EGroomCacheType Type);
-	int32 GetStartFrame() const;
-	int32 GetEndFrame() const;
+	UE_API void Initialize(EGroomCacheType Type);
+	UE_API int32 GetStartFrame() const;
+	UE_API int32 GetEndFrame() const;
 
 	//~ Begin IChaosCacheData interface
-	virtual float GetDuration() const override;
+	UE_API virtual float GetDuration() const override;
 	//~ End IChaosCacheData interface
 
 	/** Get the frame number at the specified time within the animation range which might not start at 0 */
-	int32 GetFrameNumberAtTime(const float Time, bool bLooping) const;
+	UE_API int32 GetFrameNumberAtTime(const float Time, bool bLooping) const;
 
 	/** Get the (floored) frame index at the specified time with the index 0 being the start of the animation */
-	int32 GetFrameIndexAtTime(const float Time, bool bLooping) const;
+	UE_API int32 GetFrameIndexAtTime(const float Time, bool bLooping) const;
 
 	/** Get the frame indices and interpolation factor between them that correspond to the specified time */
-	void GetFrameIndicesAtTime(float Time, bool bLooping, bool bIsPlayingBackwards, int32 &OutFrameIndex, int32 &OutNextFrameIndex, float &InterpolationFactor);
+	UE_API void GetFrameIndicesAtTime(float Time, bool bLooping, bool bIsPlayingBackwards, int32 &OutFrameIndex, int32 &OutNextFrameIndex, float &InterpolationFactor);
 
 	/** Get the frame indices that correspond to the specified time range */
-	void GetFrameIndicesForTimeRange(float StartTime, float EndTime, bool Looping, TArray<int32>& OutFrameIndices);
+	UE_API void GetFrameIndicesForTimeRange(float StartTime, float EndTime, bool Looping, TArray<int32>& OutFrameIndices);
 
-	bool GetGroomDataAtTime(float Time, bool bLooping, FGroomCacheAnimationData& AnimData);
-	bool GetGroomDataAtFrameIndex(int32 FrameIndex, FGroomCacheAnimationData& AnimData);
+	UE_API bool GetGroomDataAtTime(float Time, bool bLooping, FGroomCacheAnimationData& AnimData);
+	UE_API bool GetGroomDataAtFrameIndex(int32 FrameIndex, FGroomCacheAnimationData& AnimData);
 
-	void SetGroomAnimationInfo(const FGroomAnimationInfo& AnimInfo);
+	UE_API void SetGroomAnimationInfo(const FGroomAnimationInfo& AnimInfo);
 	const FGroomAnimationInfo& GetGroomAnimationInfo() const { return GroomCacheInfo.AnimationInfo; }
 
-	EGroomCacheType GetType() const;
+	UE_API EGroomCacheType GetType() const;
 
 	TArray<FGroomCacheChunk>& GetChunks() { return Chunks; }
 
@@ -73,10 +75,10 @@ public:
 	TArray<TObjectPtr<UAssetUserData>> AssetUserData;
 
 	//~ Begin IInterface_AssetUserData Interface
-	virtual void AddAssetUserData(UAssetUserData* InUserData) override;
-	virtual void RemoveUserDataOfClass(TSubclassOf<UAssetUserData> InUserDataClass) override;
-	virtual UAssetUserData* GetAssetUserDataOfClass(TSubclassOf<UAssetUserData> InUserDataClass) override;
-	virtual const TArray<UAssetUserData*>* GetAssetUserDataArray() const override;
+	UE_API virtual void AddAssetUserData(UAssetUserData* InUserData) override;
+	UE_API virtual void RemoveUserDataOfClass(TSubclassOf<UAssetUserData> InUserDataClass) override;
+	UE_API virtual UAssetUserData* GetAssetUserDataOfClass(TSubclassOf<UAssetUserData> InUserDataClass) override;
+	UE_API virtual const TArray<UAssetUserData*>* GetAssetUserDataArray() const override;
 	//~ End IInterface_AssetUserData Interface
 
 protected:
@@ -107,13 +109,13 @@ struct FGroomCacheChunk
 };
 
 /** Proxy that processes the HairGroupData into GroomCacheChunks that contain the groom animation data */
-class HAIRSTRANDSCORE_API FGroomCacheProcessor
+class FGroomCacheProcessor
 {
 public:
-	FGroomCacheProcessor(EGroomCacheType InType, EGroomCacheAttributes InAttributes);
+	UE_API FGroomCacheProcessor(EGroomCacheType InType, EGroomCacheAttributes InAttributes);
 
-	void AddGroomSample(TArray<FGroomCacheInputData>&& GroomData);
-	void TransferChunks(UGroomCache* GroomCache);
+	UE_API void AddGroomSample(TArray<FGroomCacheInputData>&& GroomData);
+	UE_API void TransferChunks(UGroomCache* GroomCache);
 	EGroomCacheType GetType() const { return Type; }
 
 private:
@@ -131,3 +133,5 @@ namespace UE::Groom
 	HAIRSTRANDSCORE_API bool BuildGroupsData(const FHairDescription& HairDescription, const TArray<FHairGroupPlatformData>& PlatformData,
 		TArray<FHairGroupInfoWithVisibility>& VisibilityData, const TArray<FHairGroupsInterpolation>& InterpolationData, TArray<FGroomCacheInputData>& GroupsData);
 }
+
+#undef UE_API
\ No newline at end of file
diff --git a/Engine/Plugins/Runtime/HairStrands/Source/HairStrandsCore/Public/GroomComponent.h b/Engine/Plugins/Runtime/HairStrands/Source/HairStrandsCore/Public/GroomComponent.h
--- a/Engine/Plugins/Runtime/HairStrands/Source/HairStrandsCore/Public/GroomComponent.h
+++ b/Engine/Plugins/Runtime/HairStrands/Source/HairStrandsCore/Public/GroomComponent.h
@@ -16,14 +16,16 @@
 
 #include "GroomComponent.generated.h"
 
+#define UE_API HAIRSTRANDSCORE_API
+
 class UGroomCache;
 class UMeshDeformer;
 class UMeshDeformerInstance;
 class UMeshDeformerInstanceSettings;
 class UGroomComponent;
 
-UCLASS(HideCategories = (Object, Physics, Activation, Mobility, "Components|Activation"), editinlinenew, meta = (BlueprintSpawnableComponent), ClassGroup = Rendering)
-class HAIRSTRANDSCORE_API UGroomComponent : public UMeshComponent, public ILODSyncInterface, public INiagaraPhysicsAssetDICollectorInterface
+UCLASS(HideCategories = (Object, Physics, Activation, Mobility, "Components|Activation"), editinlinenew, meta = (BlueprintSpawnableComponent), ClassGroup = Rendering, MinimalAPI)
+class UGroomComponent : public UMeshComponent, public ILODSyncInterface, public INiagaraPhysicsAssetDICollectorInterface
 {
 	GENERATED_UCLASS_BODY()
 
@@ -101,163 +103,163 @@ public:
 	FMatrix	PrevBoneMatrix;
 
 	/** Update Niagara components */
-	void UpdateHairSimulation();
+	UE_API void UpdateHairSimulation();
 
 	/** Release Niagara components */
-	void ReleaseHairSimulation();
-	void ReleaseHairSimulation(const int32 GroupIndex);
+	UE_API void ReleaseHairSimulation();
+	UE_API void ReleaseHairSimulation(const int32 GroupIndex);
 
 	/** Create per Group/LOD the Niagara component */
-	void CreateHairSimulation(const int32 GroupIndex, const int32 LODIndex);
+	UE_API void CreateHairSimulation(const int32 GroupIndex, const int32 LODIndex);
 
 	/** Enable/Disable hair simulation while transitioning from one LOD to another one */
-	void SwitchSimulationLOD(const int32 PreviousLOD, const int32 CurrentLOD, const EHairLODSelectionType InLODSelectionType);
+	UE_API void SwitchSimulationLOD(const int32 PreviousLOD, const int32 CurrentLOD, const EHairLODSelectionType InLODSelectionType);
 
 	/** Check if the simulation is enabled or not */
-	bool IsSimulationEnable(int32 GroupIndex, int32 LODIndex) const;
+	UE_API bool IsSimulationEnable(int32 GroupIndex, int32 LODIndex) const;
 
 	/** Check if the deformation is enabled or not (from the rigged guides or from the solver deformer) */
-	bool IsDeformationEnable(int32 GroupIndex) const;
+	UE_API bool IsDeformationEnable(int32 GroupIndex) const;
 
 	/** Update Group Description */
-	void UpdateHairGroupsDesc();
-	void UpdateHairGroupsDescAndInvalidateRenderState(bool bInvalidate=true);
+	UE_API void UpdateHairGroupsDesc();
+	UE_API void UpdateHairGroupsDescAndInvalidateRenderState(bool bInvalidate=true);
 
 	/** Update simulated groups */
-	void UpdateSimulatedGroups();
+	UE_API void UpdateSimulatedGroups();
 
 	//~ Begin UObject Interface.
-	virtual void GetResourceSizeEx(FResourceSizeEx& CumulativeResourceSize) override;
-	uint32 GetResourcesSize() const; 
+	UE_API virtual void GetResourceSizeEx(FResourceSizeEx& CumulativeResourceSize) override;
+	UE_API uint32 GetResourcesSize() const; 
 	//~ End UObject Interface.
 
 	//~ Begin UActorComponent Interface.
-	virtual void OnRegister() override;
-	virtual void OnUnregister() override;
-	virtual void OnComponentDestroyed(bool bDestroyingHierarchy) override;
-	virtual void BeginDestroy() override;
-	virtual void FinishDestroy() override;
-	virtual void OnAttachmentChanged() override;
-	virtual void DetachFromComponent(const FDetachmentTransformRules& DetachmentRules) override;
-	virtual void TickComponent(float DeltaTime, enum ELevelTick TickType, FActorComponentTickFunction *ThisTickFunction) override;
-
-	virtual void CreateRenderState_Concurrent(FRegisterComponentContext* Context) override;
-	virtual void SendRenderDynamicData_Concurrent() override;
-	virtual void DestroyRenderState_Concurrent() override;
-	virtual bool RequiresGameThreadEndOfFrameRecreate() const override;
+	UE_API virtual void OnRegister() override;
+	UE_API virtual void OnUnregister() override;
+	UE_API virtual void OnComponentDestroyed(bool bDestroyingHierarchy) override;
+	UE_API virtual void BeginDestroy() override;
+	UE_API virtual void FinishDestroy() override;
+	UE_API virtual void OnAttachmentChanged() override;
+	UE_API virtual void DetachFromComponent(const FDetachmentTransformRules& DetachmentRules) override;
+	UE_API virtual void TickComponent(float DeltaTime, enum ELevelTick TickType, FActorComponentTickFunction *ThisTickFunction) override;
+
+	UE_API virtual void CreateRenderState_Concurrent(FRegisterComponentContext* Context) override;
+	UE_API virtual void SendRenderDynamicData_Concurrent() override;
+	UE_API virtual void DestroyRenderState_Concurrent() override;
+	UE_API virtual bool RequiresGameThreadEndOfFrameRecreate() const override;
 	//~ End UActorComponent Interface.
 
 	//~ Begin USceneComponent Interface.
-	virtual FBoxSphereBounds CalcBounds(const FTransform& LocalToWorld) const override;
+	UE_API virtual FBoxSphereBounds CalcBounds(const FTransform& LocalToWorld) const override;
 	//~ End USceneComponent Interface.
 
 	//~ Begin UPrimitiveComponent Interface.
-	virtual FPrimitiveSceneProxy* CreateSceneProxy() override;
-	virtual void CollectPSOPrecacheData(const FPSOPrecacheParams& BasePrecachePSOParams, FMaterialInterfacePSOPrecacheParamsList& OutParams) override;
+	UE_API virtual FPrimitiveSceneProxy* CreateSceneProxy() override;
+	UE_API virtual void CollectPSOPrecacheData(const FPSOPrecacheParams& BasePrecachePSOParams, FMaterialInterfacePSOPrecacheParamsList& OutParams) override;
 	//~ End UPrimitiveComponent Interface.
 
 	//~ Begin UMeshComponent Interface.
-	virtual void PostLoad() override;
+	UE_API virtual void PostLoad() override;
 	//~ End UMeshComponent Interface.
 
 	/** Return the guide hairs rest resources*/
-	FHairStrandsRestResource* GetGuideStrandsRestResource(uint32 GroupIndex);
+	UE_API FHairStrandsRestResource* GetGuideStrandsRestResource(uint32 GroupIndex);
 
 	/** Return the guide hairs deformed resources*/
-	FHairStrandsDeformedResource* GetGuideStrandsDeformedResource(uint32 GroupIndex);
+	UE_API FHairStrandsDeformedResource* GetGuideStrandsDeformedResource(uint32 GroupIndex);
 
 	/** Return the guide hairs root resources*/
-	FHairStrandsRestRootResource* GetGuideStrandsRestRootResource(uint32 GroupIndex);
-	FHairStrandsDeformedRootResource* GetGuideStrandsDeformedRootResource(uint32 GroupIndex);
+	UE_API FHairStrandsRestRootResource* GetGuideStrandsRestRootResource(uint32 GroupIndex);
+	UE_API FHairStrandsDeformedRootResource* GetGuideStrandsDeformedRootResource(uint32 GroupIndex);
 
 #if WITH_EDITOR
-	virtual void CheckForErrors() override;
-	virtual void PreEditChange(FProperty* PropertyAboutToChange) override;
-	virtual void PostEditChangeProperty(FPropertyChangedEvent& PropertyChangedEvent) override;
-	virtual bool CanEditChange(const FProperty* InProperty) const override;
-	void ValidateMaterials(bool bMapCheck) const;
-	void Invalidate();
-	void InvalidateAndRecreate();
-
-	virtual void PreFeatureLevelChange(ERHIFeatureLevel::Type PendingFeatureLevel) override;
-	void HandlePlatformPreviewChanged(ERHIFeatureLevel::Type InFeatureLevel);
-	void HandleFeatureLevelChanged(ERHIFeatureLevel::Type InFeatureLevel);
+	UE_API virtual void CheckForErrors() override;
+	UE_API virtual void PreEditChange(FProperty* PropertyAboutToChange) override;
+	UE_API virtual void PostEditChangeProperty(FPropertyChangedEvent& PropertyChangedEvent) override;
+	UE_API virtual bool CanEditChange(const FProperty* InProperty) const override;
+	UE_API void ValidateMaterials(bool bMapCheck) const;
+	UE_API void Invalidate();
+	UE_API void InvalidateAndRecreate();
+
+	UE_API virtual void PreFeatureLevelChange(ERHIFeatureLevel::Type PendingFeatureLevel) override;
+	UE_API void HandlePlatformPreviewChanged(ERHIFeatureLevel::Type InFeatureLevel);
+	UE_API void HandleFeatureLevelChanged(ERHIFeatureLevel::Type InFeatureLevel);
 #endif
 
-	void PostCompilation();
+	UE_API void PostCompilation();
 	
 	//~ Begin IInterface_AsyncCompilation Interface.
-	virtual bool IsCompiling() const override;
+	UE_API virtual bool IsCompiling() const override;
 	//~ End IInterface_AsyncCompilation Interface.
 
 	/* Accessor function for changing Groom asset from blueprint/sequencer */
 	UFUNCTION(BlueprintCallable, Category = "Groom")
-	void SetGroomAsset(UGroomAsset* Asset);
+	UE_API void SetGroomAsset(UGroomAsset* Asset);
 
 	/* Accessor function for changing Groom binding asset from blueprint/sequencer */
 	UFUNCTION(BlueprintCallable, Category = "Groom")
-	void SetBindingAsset(UGroomBindingAsset* InBinding);
+	UE_API void SetBindingAsset(UGroomBindingAsset* InBinding);
 
 	/* Accessor function for changing Groom physics asset from blueprint/sequencer */
 	UFUNCTION(BlueprintCallable, Category = "Simulation")
-	void SetPhysicsAsset(UPhysicsAsset* InPhysicsAsset);
+	UE_API void SetPhysicsAsset(UPhysicsAsset* InPhysicsAsset);
 
 	/* Change the MeshDeformer that is used for this Component. */
 	UFUNCTION(BlueprintCallable, Category = "Components|SkinnedMesh")
-	void SetMeshDeformer(UMeshDeformer* InMeshDeformer);
+	UE_API void SetMeshDeformer(UMeshDeformer* InMeshDeformer);
 
 	/* Add a skeletal mesh to the collision components */
 	UFUNCTION(BlueprintCallable, Category = "Simulation")
-	void AddCollisionComponent(USkeletalMeshComponent* SkeletalMeshComponent);
+	UE_API void AddCollisionComponent(USkeletalMeshComponent* SkeletalMeshComponent);
 	
 	/* Reset the collision components */
 	UFUNCTION(BlueprintCallable, Category = "Simulation")
-	void ResetCollisionComponents();
+	UE_API void ResetCollisionComponents();
 
 	/* Accessor function for changing the enable simulation flag from blueprint/sequencer */
 	UFUNCTION(BlueprintCallable, Category = "Simulation")
-	void SetEnableSimulation(bool bInEnableSimulation);
+	UE_API void SetEnableSimulation(bool bInEnableSimulation);
 
 	/* Reset the simulation, if enabled */
 	UFUNCTION(BlueprintCallable, Category = "Simulation")
-	void ResetSimulation();
+	UE_API void ResetSimulation();
 	
 	/* Given the group index return the matching niagara component */
 	UFUNCTION(BlueprintCallable, Category = "Simulation")
-	UNiagaraComponent* GetNiagaraComponent(const int32 GroupIndex) {return NiagaraComponents.IsValidIndex(GroupIndex) ? NiagaraComponents[GroupIndex] : nullptr;}
+	UE_API UNiagaraComponent* GetNiagaraComponent(const int32 GroupIndex) {return NiagaraComponents.IsValidIndex(GroupIndex) ? NiagaraComponents[GroupIndex] : nullptr;}
 
 	/* Accessor function for changing hair length scale from blueprint/sequencer */
 	UFUNCTION(BlueprintCallable, Category = "Groom")
-	void SetHairLengthScale(float Scale);
+	UE_API void SetHairLengthScale(float Scale);
 
 	UFUNCTION(BlueprintCallable, Category = "Groom")
-	void SetHairLengthScaleEnable(bool bEnable);
+	UE_API void SetHairLengthScaleEnable(bool bEnable);
 
 	UFUNCTION(BlueprintCallable, Category = "Groom")
-	bool GetIsHairLengthScaleEnabled();
-
-	void SetStableRasterization(bool bEnable);
-	void SetGroomAsset(UGroomAsset* Asset, UGroomBindingAsset* InBinding, const bool bUpdateSimulation = true);
-	void SetHairRootScale(float Scale);
-	void SetHairWidth(float HairWidth);
-	void SetScatterSceneLighting(bool Enable);
-	void SetBinding(UGroomBindingAsset* InBinding);
-	void SetUseCards(bool InbUseCards);
+	UE_API bool GetIsHairLengthScaleEnabled();
+
+	UE_API void SetStableRasterization(bool bEnable);
+	UE_API void SetGroomAsset(UGroomAsset* Asset, UGroomBindingAsset* InBinding, const bool bUpdateSimulation = true);
+	UE_API void SetHairRootScale(float Scale);
+	UE_API void SetHairWidth(float HairWidth);
+	UE_API void SetScatterSceneLighting(bool Enable);
+	UE_API void SetBinding(UGroomBindingAsset* InBinding);
+	UE_API void SetUseCards(bool InbUseCards);
 	void SetValidation(bool bEnable) { bValidationEnable = bEnable; }
 
 	///~ Begin ILODSyncInterface Interface.
-	virtual int32 GetDesiredSyncLOD() const override;
-	virtual int32 GetBestAvailableLOD() const override;
-	virtual void SetForceStreamedLOD(int32 LODIndex) override;
-	virtual void SetForceRenderedLOD(int32 LODIndex) override;
-	virtual int32 GetNumSyncLODs() const override;
-	virtual int32 GetForceStreamedLOD() const override;
-	virtual int32 GetForceRenderedLOD() const override;
+	UE_API virtual int32 GetDesiredSyncLOD() const override;
+	UE_API virtual int32 GetBestAvailableLOD() const override;
+	UE_API virtual void SetForceStreamedLOD(int32 LODIndex) override;
+	UE_API virtual void SetForceRenderedLOD(int32 LODIndex) override;
+	UE_API virtual int32 GetNumSyncLODs() const override;
+	UE_API virtual int32 GetForceStreamedLOD() const override;
+	UE_API virtual int32 GetForceRenderedLOD() const override;
 	//~ End ILODSyncInterface
 
-	int32 GetNumLODs() const;
-	int32 GetForcedLOD() const;
-	void  SetForcedLOD(int32 LODIndex);
+	UE_API int32 GetNumLODs() const;
+	UE_API int32 GetForcedLOD() const;
+	UE_API void  SetForcedLOD(int32 LODIndex);
 
 	/** Groom's groups info. */
 	UPROPERTY(EditAnywhere, BlueprintReadWrite, Category = "Groom")
@@ -274,14 +276,14 @@ public:
 	bool ContainsGroupInstance(FHairGroupInstance* Instance) const { return HairGroupInstances.Contains(Instance); }
 
 	//~ Begin UPrimitiveComponent Interface
-	EHairGeometryType GetMaterialGeometryType(int32 ElementIndex) const;
-	UMaterialInterface* GetMaterial(int32 ElementIndex, EHairGeometryType GeometryType, bool bUseDefaultIfIncompatible) const;
-	virtual UMaterialInterface* GetMaterial(int32 ElementIndex) const override;
-	virtual int32 GetMaterialIndex(FName MaterialSlotName) const override;
-	virtual TArray<FName> GetMaterialSlotNames() const override;
-	virtual bool IsMaterialSlotNameValid(FName MaterialSlotName) const override;
-	virtual void GetUsedMaterials(TArray<UMaterialInterface*>& OutMaterials, bool bGetDebugMaterials = false) const override;
-	virtual int32 GetNumMaterials() const override;
+	UE_API EHairGeometryType GetMaterialGeometryType(int32 ElementIndex) const;
+	UE_API UMaterialInterface* GetMaterial(int32 ElementIndex, EHairGeometryType GeometryType, bool bUseDefaultIfIncompatible) const;
+	UE_API virtual UMaterialInterface* GetMaterial(int32 ElementIndex) const override;
+	UE_API virtual int32 GetMaterialIndex(FName MaterialSlotName) const override;
+	UE_API virtual TArray<FName> GetMaterialSlotNames() const override;
+	UE_API virtual bool IsMaterialSlotNameValid(FName MaterialSlotName) const override;
+	UE_API virtual void GetUsedMaterials(TArray<UMaterialInterface*>& OutMaterials, bool bGetDebugMaterials = false) const override;
+	UE_API virtual int32 GetNumMaterials() const override;
 	//~ End UPrimitiveComponent Interface
 
 #if WITH_EDITORONLY_DATA
@@ -294,16 +296,16 @@ public:
 
 	/* Accessor function for changing GroomCache asset from blueprint/sequencer */
 	UFUNCTION(BlueprintCallable, Category = "Groom")
-	void SetGroomCache(UGroomCache* InGroomCache);
+	UE_API void SetGroomCache(UGroomCache* InGroomCache);
 
-	float GetGroomCacheDuration() const;
+	UE_API float GetGroomCacheDuration() const;
 
-	void SetManualTick(bool bInManualTick);
-	bool GetManualTick() const;
-	void TickAtThisTime(const float Time, bool bInIsRunning, bool bInBackwards, bool bInIsLooping);
+	UE_API void SetManualTick(bool bInManualTick);
+	UE_API bool GetManualTick() const;
+	UE_API void TickAtThisTime(const float Time, bool bInIsRunning, bool bInBackwards, bool bInIsLooping);
 
-	void ResetAnimationTime();
-	float GetAnimationTime() const;
+	UE_API void ResetAnimationTime();
+	UE_API float GetAnimationTime() const;
 	bool IsLooping() const { return bLooping; }
 
 	/** Set the groom solver onto the component */
@@ -313,10 +315,10 @@ public:
 	UMeshComponent* GetGroomSolver() const {return GroomSolver.Get();}
 
 	/** Build the local simulation transform that could be used in strands simulation */
-	void BuildSimulationTransform(FTransform& SimulationTransform) const;
+	UE_API void BuildSimulationTransform(FTransform& SimulationTransform) const;
 
 	//~ Begin INiagaraPhysicsAssetDICollectorInterface Interface
-	virtual UPhysicsAsset* BuildAndCollect(FTransform& BoneTransform, TArray<TWeakObjectPtr<USkeletalMeshComponent>>& SourceComponents, TArray<TWeakObjectPtr<UPhysicsAsset>>& PhysicsAssets) const override;
+	UE_API virtual UPhysicsAsset* BuildAndCollect(FTransform& BoneTransform, TArray<TWeakObjectPtr<USkeletalMeshComponent>>& SourceComponents, TArray<TWeakObjectPtr<UPhysicsAsset>>& PhysicsAssets) const override;
 	//~ End INiagaraPhysicsAssetDICollectorInterface Interface
 
 private:
@@ -378,11 +380,11 @@ private:
 
 #if WITH_EDITORONLY_DATA
 /** Used to recreate render context for all GroomComponents that use a given GroomAsset */
-class HAIRSTRANDSCORE_API FGroomComponentRecreateRenderStateContext
+class FGroomComponentRecreateRenderStateContext
 {
 public:
-	FGroomComponentRecreateRenderStateContext(UGroomAsset* GroomAsset);
-	~FGroomComponentRecreateRenderStateContext();
+	UE_API FGroomComponentRecreateRenderStateContext(UGroomAsset* GroomAsset);
+	UE_API ~FGroomComponentRecreateRenderStateContext();
 
 private:
 	TArray<UGroomComponent*> GroomComponents;
@@ -405,4 +407,4 @@ struct FGroomComponentMemoryStats
 	uint32 GetTotalSize() const;
 };
 
-
+#undef UE_API 
\ No newline at end of file
diff --git a/Engine/Plugins/Runtime/HairStrands/Source/HairStrandsCore/Public/HairStrandsDatas.h b/Engine/Plugins/Runtime/HairStrands/Source/HairStrandsCore/Public/HairStrandsDatas.h
--- a/Engine/Plugins/Runtime/HairStrands/Source/HairStrandsCore/Public/HairStrandsDatas.h
+++ b/Engine/Plugins/Runtime/HairStrands/Source/HairStrandsCore/Public/HairStrandsDatas.h
@@ -12,6 +12,8 @@
 #include "IO/IoDispatcher.h"
 #include "Memory/SharedBuffer.h"
 
+#define UE_API HAIRSTRANDSCORE_API
+
 DECLARE_LOG_CATEGORY_EXTERN(LogHairStrands, Log, All);
 
 ///////////////////////////////////////////////////////////////////////////////////////////////////
@@ -217,7 +219,7 @@ struct FHairStrandsRaytracingFormat
 };
 
 /** Hair strands index format */
-struct HAIRSTRANDSCORE_API FHairStrandsIndexFormat
+struct FHairStrandsIndexFormat
 {
 	using Type = uint32;
 	using BulkType = uint32;
@@ -228,7 +230,7 @@ struct HAIRSTRANDSCORE_API FHairStrandsIndexFormat
 };
 
 /** Hair strands RBF sample index format */
-struct HAIRSTRANDSCORE_API FHairStrandsRBFSampleIndexFormat
+struct FHairStrandsRBFSampleIndexFormat
 {
 	using Type = uint32;
 	using BulkType = uint32;
@@ -447,18 +449,18 @@ struct FHairBulkContainer
 	bool IsBulkDataEmptyOrLoaded()const { return Data.GetBulkDataSize() == 0 || Data.IsBulkDataLoaded(); }
 };
 
-struct HAIRSTRANDSCORE_API FHairStrandsBulkCommon
+struct FHairStrandsBulkCommon
 {
 	virtual ~FHairStrandsBulkCommon() { }
-	void Serialize(FArchive& Ar, UObject* Owner);
+	UE_API void Serialize(FArchive& Ar, UObject* Owner);
 	virtual void SerializeHeader(FArchive& Ar, UObject* Owner) = 0;
-	void SerializeData(FArchive& Ar, UObject* Owner);
+	UE_API void SerializeData(FArchive& Ar, UObject* Owner);
 
-	void Write_DDC(UObject* Owner, TArray<UE::DerivedData::FCachePutValueRequest>& Out);
-	void Read_DDC(FHairStreamingRequest* In, TArray<UE::DerivedData::FCacheGetChunkRequest>& Out);
-	void Write_IO(UObject* Owner, FArchive& Out);
-	void Read_IO(FHairStreamingRequest* In, FBulkDataBatchRequest& Out);
-	void Unload(FHairStreamingRequest* In);
+	UE_API void Write_DDC(UObject* Owner, TArray<UE::DerivedData::FCachePutValueRequest>& Out);
+	UE_API void Read_DDC(FHairStreamingRequest* In, TArray<UE::DerivedData::FCacheGetChunkRequest>& Out);
+	UE_API void Write_IO(UObject* Owner, FArchive& Out);
+	UE_API void Read_IO(FHairStreamingRequest* In, FBulkDataBatchRequest& Out);
+	UE_API void Unload(FHairStreamingRequest* In);
 
 	struct FQuery
 	{
@@ -493,18 +495,18 @@ struct HAIRSTRANDSCORE_API FHairStrandsBulkCommon
 // Strands Data/Bulk
 
 /** Hair strands points attribute */
-struct HAIRSTRANDSCORE_API FHairStrandsPoints
+struct FHairStrandsPoints
 {
 	/** Set the number of points */
-	void SetNum(const uint32 NumPoints, uint32 InAttributes);
+	UE_API void SetNum(const uint32 NumPoints, uint32 InAttributes);
 
 	/** Reset the points to 0 */
-	void Reset();
+	UE_API void Reset();
 
 	/** Get the number of points */
 	uint32 Num() const { return PointsPosition.Num();  }
 
-	bool HasAttribute(EHairAttribute In) const;
+	UE_API bool HasAttribute(EHairAttribute In) const;
 
 	/** Points position in local space */
 	TArray<FVector3f> PointsPosition;
@@ -526,20 +528,20 @@ struct HAIRSTRANDSCORE_API FHairStrandsPoints
 };
 
 /** Hair strands Curves attribute */
-struct HAIRSTRANDSCORE_API FHairStrandsCurves
+struct FHairStrandsCurves
 {
 	/** Set the number of Curves */
-	void SetNum(const uint32 NumPoints, uint32 InAttributes);
+	UE_API void SetNum(const uint32 NumPoints, uint32 InAttributes);
 
 	/** Reset the curves to 0 */
-	void Reset();
+	UE_API void Reset();
 
 	/** Get the number of Curves */
 	uint32 Num() const { return CurvesCount.Num(); }
 
-	bool HasPrecomputedWeights() const;
+	UE_API bool HasPrecomputedWeights() const;
 
-	bool HasAttribute(EHairAttribute In) const;
+	UE_API bool HasAttribute(EHairAttribute In) const;
 
 	/** Number of points per rod */
 	TArray<uint16> CurvesCount;
@@ -573,7 +575,7 @@ struct HAIRSTRANDSCORE_API FHairStrandsCurves
 };
 
 /** Hair strands datas that are stored on CPU */
-struct HAIRSTRANDSCORE_API FHairStrandsCommonDatas
+struct FHairStrandsCommonDatas
 {
 	/* Get the total number of points */
 	uint32 GetNumPoints() const { return StrandsPoints.Num(); }
@@ -581,10 +583,10 @@ struct HAIRSTRANDSCORE_API FHairStrandsCommonDatas
 	/* Get the total number of Curves */
 	uint32 GetNumCurves() const { return StrandsCurves.Num(); }
 
-	uint32 GetAttributes() const;
-	uint32 GetAttributeFlags() const;
+	UE_API uint32 GetAttributes() const;
+	UE_API uint32 GetAttributeFlags() const;
 
-	void Reset();
+	UE_API void Reset();
 
 	bool IsValid() const { return StrandsCurves.Num() > 0 && StrandsPoints.Num() > 0; }
 
@@ -597,20 +599,20 @@ struct HAIRSTRANDSCORE_API FHairStrandsCommonDatas
 
 /** Hair strands raw datas that are stored on CPU */
 /** Unlike FHairStrandsDatas, this is raw datas, i.e. not processed by any built step */
-struct HAIRSTRANDSCORE_API FHairStrandsRawDatas : public FHairStrandsCommonDatas
+struct FHairStrandsRawDatas : public FHairStrandsCommonDatas
 {
-	void Reset();
+	UE_API void Reset();
 };
 
 /** Built hair strands data (decimated/shuffled/...) */
-struct HAIRSTRANDSCORE_API FHairStrandsDatas : public FHairStrandsCommonDatas
+struct FHairStrandsDatas : public FHairStrandsCommonDatas
 {
-	void Reset();
+	UE_API void Reset();
 
 	/* Copy a point or a curve from In to Out data */
-	static void CopyCurve(const FHairStrandsDatas& In, FHairStrandsDatas& Out, uint32 InAttributes, uint32 InIndex, uint32 OutIndex);
-	static void CopyPoint(const FHairStrandsDatas& In, FHairStrandsDatas& Out, uint32 InAttributes, uint32 InIndex, uint32 OutIndex);
-	static void CopyPointLerp(const FHairStrandsDatas& In, FHairStrandsDatas& Out, uint32 InAttributes, uint32 InIndex0, uint32 InIndex1, float InAlpha, uint32 OutIndex);
+	static UE_API void CopyCurve(const FHairStrandsDatas& In, FHairStrandsDatas& Out, uint32 InAttributes, uint32 InIndex, uint32 OutIndex);
+	static UE_API void CopyPoint(const FHairStrandsDatas& In, FHairStrandsDatas& Out, uint32 InAttributes, uint32 InIndex, uint32 OutIndex);
+	static UE_API void CopyPointLerp(const FHairStrandsDatas& In, FHairStrandsDatas& Out, uint32 InAttributes, uint32 InIndex0, uint32 InIndex1, float InAlpha, uint32 OutIndex);
 
 	/** The Standard Hair Density */
 	float HairDensity = 1;
@@ -619,7 +621,7 @@ struct HAIRSTRANDSCORE_API FHairStrandsDatas : public FHairStrandsCommonDatas
 	FBox3f BoundingBox = FBox3f(EForceInit::ForceInit);
 };
 
-struct HAIRSTRANDSCORE_API FHairStrandsBulkData : FHairStrandsBulkCommon
+struct FHairStrandsBulkData : FHairStrandsBulkCommon
 {
 	enum EDataFlags
 	{
@@ -630,17 +632,17 @@ struct HAIRSTRANDSCORE_API FHairStrandsBulkData : FHairStrandsBulkCommon
 		DataFlags_HasTrimmedPoint = 0x10u,		// Source data had more points per curve than the limit and some points were trimmed
 	};
 
-	virtual void SerializeHeader(FArchive& Ar, UObject* Owner) override;
-	virtual uint32 GetResourceCount() const override;
-	virtual void GetResources(FQuery& Out) override;
-	virtual void GetResourceVersion(FArchive& Ar) const override;
+	UE_API virtual void SerializeHeader(FArchive& Ar, UObject* Owner) override;
+	UE_API virtual uint32 GetResourceCount() const override;
+	UE_API virtual void GetResources(FQuery& Out) override;
+	UE_API virtual void GetResourceVersion(FArchive& Ar) const override;
 
 	bool IsValid() const { return Header.CurveCount > 0 && Header.PointCount > 0; }
-	void Reset();
-	virtual void ResetLoadedSize() override;
-	float GetCoverageScale(float InCurvePercentage /*[0..1]*/) const;
+	UE_API void Reset();
+	UE_API virtual void ResetLoadedSize() override;
+	UE_API float GetCoverageScale(float InCurvePercentage /*[0..1]*/) const;
 #if !WITH_EDITORONLY_DATA
-	bool IsCookedBulkDataLoadedForCPUAccess() const;
+	UE_API bool IsCookedBulkDataLoadedForCPUAccess() const;
 #endif
 
 	uint32 GetNumCurves() const { return Header.CurveCount;  };
@@ -649,7 +651,7 @@ struct HAIRSTRANDSCORE_API FHairStrandsBulkData : FHairStrandsBulkCommon
 	float  GetMaxRadius() const { return Header.MaxRadius; }
 	FVector GetPositionOffset() const { return Header.BoundingBox.GetCenter(); }
 	const FBox& GetBounds() const { return Header.BoundingBox; }
-	uint32 GetSize() const;
+	UE_API uint32 GetSize() const;
 
 	uint32 GetCurveAttributeSizeInBytes(uint32 InCurveCount=HAIR_MAX_NUM_CURVE_PER_GROUP) const			{ return InCurveCount > 0 ? FMath::DivideAndRoundUp(FMath::Min(Header.CurveCount, InCurveCount), Header.Strides.CurveAttributeChunkElementCount) * Header.Strides.CurveAttributeChunkStride : 0; }
 	uint32 GetPointAttributeSizeInBytes(uint32 InPointCount=HAIR_MAX_NUM_POINT_PER_GROUP) const			{ return InPointCount > 0 ? FMath::DivideAndRoundUp(FMath::Min(Header.PointCount, InPointCount), Header.Strides.PointAttributeChunkElementCount) * Header.Strides.PointAttributeChunkStride : 0; }
@@ -720,13 +722,13 @@ struct HAIRSTRANDSCORE_API FHairStrandsBulkData : FHairStrandsBulkCommon
 // Interpolation Data/Bulk
 
 /** Hair strands points interpolation attributes */
-struct HAIRSTRANDSCORE_API FHairStrandsInterpolationDatas
+struct FHairStrandsInterpolationDatas
 {
 	/** Set the number of interpolated points */
-	void SetNum(const uint32 InCurveCount, const uint32 InPointCount);
+	UE_API void SetNum(const uint32 InCurveCount, const uint32 InPointCount);
 
 	/** Reset the interpolated points to 0 */
-	void Reset();
+	UE_API void Reset();
 
 	/** Get the number of interpolated points */
 	uint32 GetPointCount() const { return PointSimIndices.Num(); }
@@ -753,7 +755,7 @@ struct HAIRSTRANDSCORE_API FHairStrandsInterpolationDatas
 	bool bUseUniqueGuide = false;
 };
 
-struct HAIRSTRANDSCORE_API FHairStrandsInterpolationBulkData : FHairStrandsBulkCommon
+struct FHairStrandsInterpolationBulkData : FHairStrandsBulkCommon
 {
 	enum EDataFlags
 	{
@@ -761,16 +763,16 @@ struct HAIRSTRANDSCORE_API FHairStrandsInterpolationBulkData : FHairStrandsBulkC
 		DataFlags_HasSingleGuideData = 2,
 	};
 
-	void Reset();
-	virtual void ResetLoadedSize() override;
-	virtual void SerializeHeader(FArchive& Ar, UObject* Owner) override;
+	UE_API void Reset();
+	UE_API virtual void ResetLoadedSize() override;
+	UE_API virtual void SerializeHeader(FArchive& Ar, UObject* Owner) override;
 
-	virtual uint32 GetResourceCount() const override;
-	virtual void GetResources(FQuery& Out) override;
+	UE_API virtual uint32 GetResourceCount() const override;
+	UE_API virtual void GetResources(FQuery& Out) override;
 	uint32 GetPointCount() const { return Header.PointCount; };
-	uint32 GetSize() const;
+	UE_API uint32 GetSize() const;
 #if !WITH_EDITORONLY_DATA
-	bool IsCookedBulkDataLoadedForCPUAccess() const;
+	UE_API bool IsCookedBulkDataLoadedForCPUAccess() const;
 #endif
 
 	struct FHeader
@@ -796,9 +798,9 @@ struct HAIRSTRANDSCORE_API FHairStrandsInterpolationBulkData : FHairStrandsBulkC
 ///////////////////////////////////////////////////////////////////////////////////////////////////
 // Cluster Data/Bulk
 
-struct HAIRSTRANDSCORE_API FHairStrandsClusterData
+struct FHairStrandsClusterData
 {
-	void Reset();
+	UE_API void Reset();
 	bool IsValid() const { return ClusterCount > 0 && PointCount > 0; }
 
 	/* LOD info for the various clusters for LOD management on GPU */
@@ -813,18 +815,18 @@ struct HAIRSTRANDSCORE_API FHairStrandsClusterData
 	uint32 CurveCount = 0;			// Number of curves
 };
 
-struct HAIRSTRANDSCORE_API FHairStrandsClusterBulkData : FHairStrandsBulkCommon
+struct FHairStrandsClusterBulkData : FHairStrandsBulkCommon
 {
-	void Reset();
-	virtual void ResetLoadedSize() override;
+	UE_API void Reset();
+	UE_API virtual void ResetLoadedSize() override;
 
-	virtual void SerializeHeader(FArchive& Ar, UObject* Owner) override;
-	virtual uint32 GetResourceCount() const override;
-	virtual void GetResources(FQuery& Out) override;
-	uint32 GetCurveCount(float InLODIndex) const;
-	uint32 GetSize() const;
+	UE_API virtual void SerializeHeader(FArchive& Ar, UObject* Owner) override;
+	UE_API virtual uint32 GetResourceCount() const override;
+	UE_API virtual void GetResources(FQuery& Out) override;
+	UE_API uint32 GetCurveCount(float InLODIndex) const;
+	UE_API uint32 GetSize() const;
 #if !WITH_EDITORONLY_DATA
-	bool IsCookedBulkDataLoadedForCPUAccess() const;
+	UE_API bool IsCookedBulkDataLoadedForCPUAccess() const;
 #endif
 
 	bool IsValid() const { return Header.ClusterCount > 0 && Header.PointCount > 0; }
@@ -972,7 +974,7 @@ struct FHairStrandsRootBulkData : FHairStrandsBulkCommon
 // Hair debug data
 
 /** Hair strands debug data */
-struct HAIRSTRANDSCORE_API FHairStrandsDebugDatas
+struct  FHairStrandsDebugDatas
 {
 	bool IsValid() const { return VoxelData.Num() > 0;  }
 
@@ -1003,9 +1005,11 @@ struct HAIRSTRANDSCORE_API FHairStrandsDebugDatas
 	TArray<FVoxel> VoxelData;
 };
 
-struct HAIRSTRANDSCORE_API FHairStrandsDebugResources
+struct FHairStrandsDebugResources
 {
 	FHairStrandsDebugDatas::FDesc VoxelDescription;
 	TRefCountPtr<FRDGPooledBuffer> VoxelOffsetAndCount;
 	TRefCountPtr<FRDGPooledBuffer> VoxelData;
 };
+
+#undef UE_API
\ No newline at end of file
diff --git a/Engine/Plugins/Runtime/HairStrands/Source/HairStrandsSolver/Public/GroomSolverComponent.h b/Engine/Plugins/Runtime/HairStrands/Source/HairStrandsSolver/Public/GroomSolverComponent.h
--- a/Engine/Plugins/Runtime/HairStrands/Source/HairStrandsSolver/Public/GroomSolverComponent.h
+++ b/Engine/Plugins/Runtime/HairStrands/Source/HairStrandsSolver/Public/GroomSolverComponent.h
@@ -8,6 +8,8 @@
 
 #include "GroomSolverComponent.generated.h"
 
+#define UE_API HAIRSTRANDSSOLVER_API
+
 class FRHIGPUBufferReadback;
 class UMeshDeformer;
 
@@ -98,14 +100,14 @@ struct TStructOpsTypeTraits<FDataflowGroomSolverProxy> : public TStructOpsTypeTr
 };
 
 /** Groom solver component in which groom component could be added to be solver together */
-UCLASS(HideCategories = (Object, Physics, Collision, Activation, Mobility, "Components|Activation"), meta = (BlueprintSpawnableComponent), ClassGroup = Physics)
-class HAIRSTRANDSSOLVER_API UGroomSolverComponent : public UMeshComponent, public IDataflowPhysicsSolverInterface
+UCLASS(HideCategories = (Object, Physics, Collision, Activation, Mobility, "Components|Activation"), meta = (BlueprintSpawnableComponent), ClassGroup = Physics, MinimalAPI)
+class UGroomSolverComponent : public UMeshComponent, public IDataflowPhysicsSolverInterface
 {
 	GENERATED_BODY()
 
 public:
 
-	UGroomSolverComponent(const FObjectInitializer& ObjectInitializer);
+	UE_API UGroomSolverComponent(const FObjectInitializer& ObjectInitializer);
 
 	// Begin IDataflowPhysicsSolverInterface overrides
 	virtual FString GetSimulationName() const override {return GetName();};
@@ -121,37 +123,37 @@ public:
 
 	//~ Begin UPrimitiveComponent Interface.
 #if WITH_EDITOR
-	virtual void PostEditChangeProperty(FPropertyChangedEvent& PropertyChangedEvent) override;
+	UE_API virtual void PostEditChangeProperty(FPropertyChangedEvent& PropertyChangedEvent) override;
 #endif
 	//~ End UPrimitiveComponent Interface.
 	
 	//~ Begin USceneComponent interface 
-	virtual void TickComponent(float DeltaTime, enum ELevelTick TickType, FActorComponentTickFunction* ThisTickFunction) override;
-	virtual void OnRegister() override;
-	virtual void OnUnregister() override;
-	virtual void CreateRenderState_Concurrent(FRegisterComponentContext* Context) override;
-	virtual void SendRenderDynamicData_Concurrent() override;
-	virtual void DestroyRenderState_Concurrent() override;
+	UE_API virtual void TickComponent(float DeltaTime, enum ELevelTick TickType, FActorComponentTickFunction* ThisTickFunction) override;
+	UE_API virtual void OnRegister() override;
+	UE_API virtual void OnUnregister() override;
+	UE_API virtual void CreateRenderState_Concurrent(FRegisterComponentContext* Context) override;
+	UE_API virtual void SendRenderDynamicData_Concurrent() override;
+	UE_API virtual void DestroyRenderState_Concurrent() override;
 	//~ End USceneComponent interface
 
 	/** Add a groom component to the solver */
 	UFUNCTION(BlueprintCallable, Category = "Groom")
-	void AddGroomComponent(UGroomComponent* GroomComponent);
+	UE_API void AddGroomComponent(UGroomComponent* GroomComponent);
 	
 	/** Add a groom component to the solver */
 	UFUNCTION(BlueprintCallable, Category = "Groom")
-	void RemoveGroomComponent(UGroomComponent* GroomComponent);
+	UE_API void RemoveGroomComponent(UGroomComponent* GroomComponent);
 
 	/** Reset the groom physics components */
     UFUNCTION(BlueprintCallable, Category = "Groom")
-    void ResetGroomComponents();
+    UE_API void ResetGroomComponents();
 
 	/** Reset the groom physics components */
 	const TSet< TObjectPtr<UGroomComponent> >& GetGroomComponents() const {return GroomComponents;}
 
 	/* Change the MeshDeformer solver that is used for this Component. */
 	UFUNCTION(BlueprintCallable, Category = "Groom")
-	void SetDeformerSolver(UMeshDeformer* DeformerSolver);
+	UE_API void SetDeformerSolver(UMeshDeformer* DeformerSolver);
 
 	/** Get the groom solver settings */
 	const FGroomSolverSettings& GetSolverSettings() const {return SolverSettings;};
@@ -192,3 +194,4 @@ private :
 	FDataflowGroomSolverProxy GroomSolverProxy;
 };
 
+#undef UE_API 
