[Hair] Add source curve/point index mapping so that attribute computed on source data can be retrieved with this index mapping.

This is useful for dataflow attribute computed for strands data and stored into attribute collection. Dataflow attribute are computed on source data, not on build data.

diff --git a/Engine/Plugins/Runtime/HairStrands/Shaders/Private/Deformers/DeformerGroomStrandsRead.ush b/Engine/Plugins/Runtime/HairStrands/Shaders/Private/Deformers/DeformerGroomStrandsRead.ush
--- a/Engine/Plugins/Runtime/HairStrands/Shaders/Private/Deformers/DeformerGroomStrandsRead.ush
+++ b/Engine/Plugins/Runtime/HairStrands/Shaders/Private/Deformers/DeformerGroomStrandsRead.ush
@@ -222,3 +222,15 @@ float3x4 ReadCurveDeformedTransform_{DataInterfaceName}(uint CurveIndex)
 		return float3x4(1, 0, 0, 0, 0, 1, 0, 0, 0, 0, 1, 0);
 #endif
 }
+
+uint ReadSourceCurveIndex_{DataInterfaceName}(uint CurveIndex)
+{
+	CurveIndex -= {DataInterfaceName}_BaseCurveIndex;
+	return ReadHairSourceCurveIndex({DataInterfaceName}_CurveRemappingBuffer, CurveIndex, {DataInterfaceName}_bCurveRemapping);
+}
+
+uint ReadSourcePointIndex_{DataInterfaceName}(uint PointIndex)
+{
+	PointIndex -= {DataInterfaceName}_BasePointIndex;
+	return ReadHairSourcePointIndex({DataInterfaceName}_PointRemappingBuffer, PointIndex, {DataInterfaceName}_bPointRemapping);
+}
\ No newline at end of file
diff --git a/Engine/Plugins/Runtime/HairStrands/Source/HairStrandsCore/Private/GroomAsset.cpp b/Engine/Plugins/Runtime/HairStrands/Source/HairStrandsCore/Private/GroomAsset.cpp
--- a/Engine/Plugins/Runtime/HairStrands/Source/HairStrandsCore/Private/GroomAsset.cpp
+++ b/Engine/Plugins/Runtime/HairStrands/Source/HairStrandsCore/Private/GroomAsset.cpp
@@ -530,7 +530,8 @@ static bool BuildHairGroup_Strands(
 	const TArray<FHairGroupsInterpolation>& InHairGroupsInterpolation,
 	const TArray<FHairGroupsLOD>& InHairGroupsLOD,
 	TArray<FHairGroupInfoWithVisibility>& OutHairGroupsInfo,
-	TArray<FHairGroupPlatformData>& OutHairGroupsData)
+	TArray<FHairGroupPlatformData>& OutHairGroupsData, 
+	const bool bBuildSourceMapping)
 {
 	const bool bIsValid = HairDescriptionGroups.IsValid();
 	if (bIsValid)
@@ -541,7 +542,7 @@ static bool BuildHairGroup_Strands(
 
 		FHairStrandsDatas StrandsData;
 		FHairStrandsDatas GuidesData;
-		FGroomBuilder::BuildData(HairGroup, InHairGroupsInterpolation[GroupIndex], OutHairGroupsInfo[GroupIndex], StrandsData, GuidesData);
+		FGroomBuilder::BuildData(HairGroup, InHairGroupsInterpolation[GroupIndex], OutHairGroupsInfo[GroupIndex], StrandsData, GuidesData, true, true, bBuildSourceMapping);
 
 		FGroomBuilder::BuildBulkData(HairGroup.Info, GuidesData, OutHairGroupsData[GroupIndex].Guides.BulkData, false /*bAllowCompression*/);
 		FGroomBuilder::BuildBulkData(HairGroup.Info, StrandsData, OutHairGroupsData[GroupIndex].Strands.BulkData, true /*bAllowCompression*/);
@@ -1518,10 +1519,11 @@ static bool IsStrandsInterpolationAttributes(const FName PropertyName)
 		|| PropertyName == GET_MEMBER_NAME_CHECKED(FHairDecimationSettings, VertexDecimation)
 
 
-		// Add dependency on simulation and per LOD-simulation/global-interpolation to strip-out interpolation data if there are not needed
+		// Add dependency on simulation, on deformer, and on per LOD-simulation/global-interpolation to strip-out interpolation data if there are not needed
 		|| PropertyName == UGroomAsset::GetEnableGlobalInterpolationMemberName()
 		|| PropertyName == UGroomAsset::GetEnableSimulationCacheMemberName()
 		|| PropertyName == GET_MEMBER_NAME_CHECKED(FHairSolverSettings, EnableSimulation)
+		|| PropertyName == GET_MEMBER_NAME_CHECKED(FHairSolverSettings, bEnableDeformation)
 
 		// Decimation control needs to force a rebuild
 		|| PropertyName == GET_MEMBER_NAME_CHECKED(FHairLODSettings, CurveDecimation)
@@ -2181,7 +2183,7 @@ namespace GroomDerivedDataCacheUtils
 		return FDerivedDataCacheInterface::BuildCacheKey(*(TEXT("GROOM_V") + FGroomBuilder::GetVersion() + TEXT("_")), *GetGroomDerivedDataVersion(), *KeySuffix);
 	}
 
-	void SerializeHairInterpolationSettingsForDDC(FArchive& Ar, uint32 GroupIndex, FHairGroupsInterpolation& InterpolationSettings, FHairGroupsLOD& LODSettings, bool bRequireInterpolationData)
+	void SerializeHairSettingsForDDC(FArchive& Ar, uint32 GroupIndex, FHairGroupsInterpolation& InterpolationSettings, FHairGroupsLOD& LODSettings, FHairGroupsPhysics& PhysicsSettings, bool bRequireInterpolationData)
 	{
 		bool bSupportCompressedPosition = DoesHairStrandsSupportCompressedPosition();
 		bool bSupportTriangleStrip = GetHairStrandsUsesTriangleStrips();
@@ -2194,6 +2196,7 @@ namespace GroomDerivedDataCacheUtils
 
 		InterpolationSettings.BuildDDCKey(Ar);
 		LODSettings.BuildDDCKey(Ar);
+		PhysicsSettings.BuildDDCKey(Ar);
 	}
 
 	// Geneate DDC key for an entire hair group (i.e. with all its LODs) for cards data
@@ -2278,7 +2281,7 @@ namespace GroomDerivedDataCacheUtils
 		return MeshPrefixString + Hash.ToString();
 	}
 
-	FString BuildStrandsDerivedDataKeySuffix(uint32 GroupIndex, const FHairGroupsInterpolation& InterpolationSettings, const FHairGroupsLOD& LODSettings, bool bNeedInterpolationData, const FHairDescriptionBulkData* HairBulkDescription)
+	FString BuildStrandsDerivedDataKeySuffix(uint32 GroupIndex, const FHairGroupsInterpolation& InterpolationSettings, const FHairGroupsLOD& LODSettings, const FHairGroupsPhysics& PhysicsSettings, bool bNeedInterpolationData, const FHairDescriptionBulkData* HairBulkDescription)
 	{
 		// Serialize the build settings into a temporary array
 		// The archive is flagged as persistent so that machines of different endianness produce identical binary results.
@@ -2286,7 +2289,7 @@ namespace GroomDerivedDataCacheUtils
 		TempBytes.Reserve(64);
 		FMemoryWriter Ar(TempBytes, /*bIsPersistent=*/ true);
 
-		GroomDerivedDataCacheUtils::SerializeHairInterpolationSettingsForDDC(Ar, GroupIndex, const_cast<FHairGroupsInterpolation&>(InterpolationSettings), const_cast<FHairGroupsLOD&>(LODSettings), bNeedInterpolationData);
+		GroomDerivedDataCacheUtils::SerializeHairSettingsForDDC(Ar, GroupIndex, const_cast<FHairGroupsInterpolation&>(InterpolationSettings), const_cast<FHairGroupsLOD&>(LODSettings), const_cast<FHairGroupsPhysics&>(PhysicsSettings), bNeedInterpolationData);
 
 		FString KeySuffix;
 		if (HairBulkDescription)
@@ -2396,10 +2399,11 @@ FString UGroomAsset::GetDerivedDataKeyForStrands(uint32 GroupIndex)
 {
 	const FHairGroupsInterpolation& InterpolationSettings = GetHairGroupsInterpolation()[GroupIndex];
 	const FHairGroupsLOD& LODSettings = GetHairGroupsLOD()[GroupIndex];
+	const FHairGroupsPhysics& PhysicsSettings = GetHairGroupsPhysics()[GroupIndex];
 
 	// If simulation or global interpolation is enabled, then interpolation data are required. Otherwise they can be skipped
 	const bool bNeedInterpolationData = NeedsInterpolationData(GroupIndex);
-	const FString KeySuffix = GroomDerivedDataCacheUtils::BuildStrandsDerivedDataKeySuffix(GroupIndex, InterpolationSettings, LODSettings, bNeedInterpolationData, HairDescriptionBulkData[HairDescriptionType].Get());
+	const FString KeySuffix = GroomDerivedDataCacheUtils::BuildStrandsDerivedDataKeySuffix(GroupIndex, InterpolationSettings, LODSettings, PhysicsSettings, bNeedInterpolationData, HairDescriptionBulkData[HairDescriptionType].Get());
 	const FString DerivedDataKey = GroomDerivedDataCacheUtils::BuildGroomDerivedDataKey(KeySuffix);
 
 	return DerivedDataKey;
@@ -2450,7 +2454,7 @@ const FHairDescriptionGroups& UGroomAsset::GetHairDescriptionGroups()
 		if (!CachedHairDescriptionGroups[HairDescriptionType])
 		{
 			TUniquePtr<FHairDescriptionGroups> HairDescriptionGroups = MakeUnique<FHairDescriptionGroups>();
-			FGroomBuilder::BuildHairDescriptionGroups(*CachedHairDescription[HairDescriptionType], *HairDescriptionGroups);
+			FGroomBuilder::BuildHairDescriptionGroups(*CachedHairDescription[HairDescriptionType], *HairDescriptionGroups, true /*bAllowAddEndControlPoint*/);
 			
 			// Since we don't take read-lock, we need to make sure all memory is visible inside newly created object before publishing it.
 			FPlatformMisc::MemoryBarrier();
@@ -2617,7 +2621,8 @@ bool UGroomAsset::CacheStrandsData(uint32 GroupIndex, FString& OutDerivedDataKey
 			GetHairGroupsInterpolation(),
 			GetHairGroupsLOD(),
 			GetHairGroupsInfo(),
-			GetHairGroupsPlatformData());
+			GetHairGroupsPlatformData(), 
+			IsDeformationEnable(GroupIndex));
 		
 		if (bSuccess)
 		{
diff --git a/Engine/Plugins/Runtime/HairStrands/Source/HairStrandsCore/Private/GroomAssetPhysics.cpp b/Engine/Plugins/Runtime/HairStrands/Source/HairStrandsCore/Private/GroomAssetPhysics.cpp
--- a/Engine/Plugins/Runtime/HairStrands/Source/HairStrandsCore/Private/GroomAssetPhysics.cpp
+++ b/Engine/Plugins/Runtime/HairStrands/Source/HairStrandsCore/Private/GroomAssetPhysics.cpp
@@ -85,6 +85,12 @@ FHairGroupsPhysics::FHairGroupsPhysics()
 	StrandsParameters = FHairStrandsParameters();
 }
 
+void FHairGroupsPhysics::BuildDDCKey(FArchive& Ar)
+{
+	// This is used to know if curve/point mapping needs to be built for a particular groom/group
+	Ar << SolverSettings.bEnableDeformation;
+}
+
 bool FHairSolverSettings::operator==(const FHairSolverSettings& A) const
 {
 	return 
diff --git a/Engine/Plugins/Runtime/HairStrands/Source/HairStrandsCore/Private/GroomBuilder.cpp b/Engine/Plugins/Runtime/HairStrands/Source/HairStrandsCore/Private/GroomBuilder.cpp
--- a/Engine/Plugins/Runtime/HairStrands/Source/HairStrandsCore/Private/GroomBuilder.cpp
+++ b/Engine/Plugins/Runtime/HairStrands/Source/HairStrandsCore/Private/GroomBuilder.cpp
@@ -54,7 +54,7 @@ bool DoesHairStrandsSupportCompressedPosition();
 
 FString FGroomBuilder::GetVersion()
 {
-	FString Out = TEXT("v18b_Debug0");
+	FString Out = TEXT("v19_Debug4");
 	if (IsRenderCurveEnabled())
 	{
 		Out += TEXT("_Curve3e");
@@ -94,13 +94,15 @@ namespace FHairStrandsDecimation
 		float CurveDecimationPercentage,
 		float VertexDecimationPercentage,
 		bool bContinuousDecimationReordering,
-		FHairStrandsDatas& OutData);
+		FHairStrandsDatas& OutData,
+		const bool bBuildSourceMapping);
 
 	void Decimate(
 		const FHairStrandsDatas& InData, 
 		const uint32 NumCurves,
 		const int32 NumVertices,
-		FHairStrandsDatas& OutData);
+		FHairStrandsDatas& OutData,
+		const bool bBuildSourceMapping);
 }
 
 namespace HairStrandsBuilder
@@ -259,6 +261,18 @@ namespace HairStrandsBuilder
 		const FHairStrandsCurves& Curves = HairStrands.StrandsCurves;
 		const FHairStrandsPoints& Points = HairStrands.StrandsPoints;
 
+		// Allocate remapping data if needed
+		TArray<FHairStrandsIndexFormat::Type> OutPointsMapping;
+		if (Points.PointsMapping.Num() > 0)
+		{
+			OutPointsMapping.SetNum(NumPoints * FHairStrandsIndexFormat::ComponentCount);
+		}
+		TArray<FHairStrandsIndexFormat::Type> OutCurvesMapping;
+		if (Curves.CurvesMapping.Num() > 0)
+		{
+			OutCurvesMapping.SetNum(NumCurves * FHairStrandsIndexFormat::ComponentCount);
+		}
+
 		const float MaxLength = GetHairStrandsMaxLength(HairStrands);
 		const float MaxRadius = GetHairStrandsMaxRadius(HairStrands);
 
@@ -365,6 +379,14 @@ namespace HairStrandsBuilder
 					uint8* Data = (uint8*)AttributeAO.GetData();
 					Data[PointIndex + IndexOffset] = FMath::Clamp(uint32(Points.PointsAO[PointIndex + IndexOffset] * 0xFF), 0u, 0xFFu);
 				}
+
+				// Per-Vertex Source Mapping
+				{
+					if(!Points.PointsMapping.IsEmpty())
+					{
+						OutPointsMapping[PointIndex + IndexOffset] = Points.PointsMapping[PointIndex + IndexOffset];
+					}
+				}
 			}
 
 			// Per-Curve Length
@@ -379,6 +401,14 @@ namespace HairStrandsBuilder
 				Packed[CurveIndex] = bSeedValid ? RandomSeeds[CurveIndex] : 0;
 			}
 
+			// Per-Curve Source Mapping
+			{
+				if(!Curves.CurvesMapping.IsEmpty())
+				{
+					OutCurvesMapping[CurveIndex] = Curves.CurvesMapping[CurveIndex];
+				}
+			}
+
 			// Per-Curve Root UVs
 			if (HasHairAttribute(Attributes, EHairAttribute::RootUV))
 			{
@@ -783,6 +813,19 @@ namespace HairStrandsBuilder
 		CopyToBulkData<FHairStrandsPointToCurveFormat>(OutBulkData.Data.PointToCurve, OutPointToCurve);
 		ReportSize(TEXT("PointToCurve"), OutPointToCurve);
 
+		if (!OutPointsMapping.IsEmpty())
+		{
+			OutBulkData.Header.Flags |= FHairStrandsBulkData::DataFlags_HasPointMapping;
+			CopyToBulkData<FHairStrandsIndexFormat>(OutBulkData.Data.PointMapping, OutPointsMapping);
+			ReportSize(TEXT("PointsMapping"), OutPointsMapping);
+		}
+		if (!OutCurvesMapping.IsEmpty())
+		{
+			OutBulkData.Header.Flags |= FHairStrandsBulkData::DataFlags_HasCurveMapping;
+			CopyToBulkData<FHairStrandsIndexFormat>(OutBulkData.Data.CurveMapping, OutCurvesMapping);
+			ReportSize(TEXT("CurvesMapping"), OutCurvesMapping);
+		}
+
 		// Build curve to point count mapping for runtime CLOD
 		OutBulkData.Header.CurveToPointCount.Reserve(OutPackedCurves.Num());
 		for (const FHairStrandsCurveFormat::Type& Curve : OutPackedCurves)
@@ -811,6 +854,9 @@ namespace HairStrandsBuilder
 
 		OutBulkData.Header.Strides.TranscodedPositionChunkElementCount = TranscodedPositionChunkElementCount;
 		OutBulkData.Header.Strides.TranscodedPositionChunkStride = sizeof(FTranscodedHairPositions);
+
+		OutBulkData.Header.Strides.CurveMappingStride = FHairStrandsIndexFormat::SizeInByte;
+		OutBulkData.Header.Strides.PointMappingStride = FHairStrandsIndexFormat::SizeInByte;
 	}
 } // namespace HairStrandsBuilder
 
@@ -1915,7 +1961,7 @@ bool NeedsToAddEndingControlPoint(uint32 CurveNumVertices)
 	return GetHairStrandsUsesTriangleStrips() && uint32(CurveNumVertices + 1) <= HAIR_MAX_NUM_POINT_PER_CURVE;
 }
 
-bool FGroomBuilder::BuildHairDescriptionGroups(const FHairDescription& HairDescription, FHairDescriptionGroups& Out, bool bAllowAddEndControlPoint)
+bool FGroomBuilder::BuildHairDescriptionGroups(const FHairDescription& HairDescription, FHairDescriptionGroups& Out, const bool bAllowAddEndControlPoint)
 {
 	TRACE_CPUPROFILER_EVENT_SCOPE(FGroomBuilder::BuildHairDescriptionGroups);
 
@@ -1999,7 +2045,7 @@ bool FGroomBuilder::BuildHairDescriptionGroups(const FHairDescription& HairDescr
 	const bool bCanUseClosestGuidesAndWeights = bImportGuides && StrandIDs.IsValid() && (bPrecomputedWeight1 || bPrecomputedWeight3);
 	check(bCanUseClosestGuidesAndWeights == HairDescription.HasAttribute(EHairAttribute::PrecomputedGuideWeights)); // Sanity check
 
-	auto FindOrAdd = [&Out, &GroupIDs, &GroupNames](FStrandID StrandID) -> FHairDescriptionGroup&
+	auto FindOrAdd = [&GroupIDs, &GroupNames](FStrandID StrandID, FHairDescriptionGroups& OutGroups) -> FHairDescriptionGroup&
 	{
 		const bool bValidGroupId = GroupIDs.IsValid();
 		const bool bValidGroupName = GroupNames.IsValid();
@@ -2008,14 +2054,14 @@ bool FGroomBuilder::BuildHairDescriptionGroups(const FHairDescription& HairDescr
 		if (bValidGroupId)
 		{
 			const int32 GroupID = GroupIDs[StrandID];
-			for (FHairDescriptionGroup& Group : Out.HairGroups)
+			for (FHairDescriptionGroup& Group : OutGroups.HairGroups)
 			{
 				if (Group.Info.GroupID == GroupID)
 				{
 					return Group;
 				}
 			}
-			FHairDescriptionGroup& Group = Out.HairGroups.AddDefaulted_GetRef();
+			FHairDescriptionGroup& Group = OutGroups.HairGroups.AddDefaulted_GetRef();
 			Group.Info.GroupID = GroupID;
 			Group.Info.GroupName = ToGroupName(StrandID, GroupID, GroupNames);
 			return Group;
@@ -2023,16 +2069,16 @@ bool FGroomBuilder::BuildHairDescriptionGroups(const FHairDescription& HairDescr
 		// 2. If no GroupID are valid, use unique GroupName to create GroupID
 		else if (bValidGroupName)
 		{
-			const uint32 GroupID = Out.HairGroups.Num();
+			const uint32 GroupID = OutGroups.HairGroups.Num();
 			const FName GroupName = ToGroupName(StrandID, GroupID, GroupNames);
-			for (FHairDescriptionGroup& Group : Out.HairGroups)
+			for (FHairDescriptionGroup& Group : OutGroups.HairGroups)
 			{
 				if (Group.Info.GroupName == GroupName)
 				{
 					return Group;
 				}
 			}
-			FHairDescriptionGroup& Group = Out.HairGroups.AddDefaulted_GetRef();
+			FHairDescriptionGroup& Group = OutGroups.HairGroups.AddDefaulted_GetRef();
 			Group.Info.GroupID = GroupID;
 			Group.Info.GroupName = GroupName;
 			return Group;
@@ -2040,14 +2086,14 @@ bool FGroomBuilder::BuildHairDescriptionGroups(const FHairDescription& HairDescr
 		// 3. If there are no GroupID and no GroupName, add a default group are valid, use unique group name to create GroupID
 		else
 		{
-			if (Out.HairGroups.IsEmpty())
+			if (OutGroups.HairGroups.IsEmpty())
 			{
 				const uint32 GroupID = 0;
-				FHairDescriptionGroup& Group = Out.HairGroups.AddDefaulted_GetRef();
+				FHairDescriptionGroup& Group = OutGroups.HairGroups.AddDefaulted_GetRef();
 				Group.Info.GroupID = GroupID;
 				Group.Info.GroupName = ToGroupName(StrandID, GroupID, GroupNames);
 			}
-			return Out.HairGroups[0];
+			return OutGroups.HairGroups[0];
 		}
 	};
 
@@ -2057,6 +2103,36 @@ bool FGroomBuilder::BuildHairDescriptionGroups(const FHairDescription& HairDescr
 	bool bCurveCountWarningIssued = false;
 	bool bPointCountWarningIssued = false;
 
+	// Initialize Curve/Point index remapping 
+	// * Per group source point offsets
+	// * Per group points source mapping
+	// * Per group curves source mapping
+	TArray<uint32> GroupSourceOffsets;
+	TArray<TArray<uint32>> GroupPointsMapping;
+	TArray<TArray<uint32>> GroupCurvesMapping;
+	{
+		FHairDescriptionGroups LocalGroups;
+		for (int32 CurveIndex = 0; CurveIndex < NumCurves; ++CurveIndex)
+		{
+			FStrandID StrandID(CurveIndex);
+			if (StrandNumVertices[StrandID] > 0)
+			{
+				FindOrAdd(StrandID, LocalGroups);
+			}			
+		}
+
+		int32 MaxGroupID = 0;
+		for(int32 GroupIndex = 0, GroupCount=LocalGroups.HairGroups.Num(); GroupIndex < GroupCount; ++GroupIndex)
+		{
+			MaxGroupID = FMath::Max(LocalGroups.HairGroups[GroupIndex].Info.GroupID, MaxGroupID);
+		}
+		
+		const uint32 NumGroupID = MaxGroupID + 1;
+		GroupSourceOffsets.Init(0, NumGroupID);
+		GroupPointsMapping.Init(TArray<uint32>(), NumGroupID);	
+		GroupCurvesMapping.Init(TArray<uint32>(), NumGroupID);
+	}
+
 	int32 GlobalVertexIndex = 0;
 	for (int32 CurveIndex = 0; CurveIndex < NumCurves; ++CurveIndex)
 	{
@@ -2092,7 +2168,7 @@ bool FGroomBuilder::BuildHairDescriptionGroups(const FHairDescription& HairDescr
 		}
 
 		FHairStrandsRawDatas* CurrentHairStrandsDatas = nullptr;
-		FHairDescriptionGroup& Group = FindOrAdd(StrandID);
+		FHairDescriptionGroup& Group = FindOrAdd(StrandID, Out);
 		bool bNumCurveValid = false;
 		if (bHasInvalidPosition)
 		{
@@ -2223,9 +2299,18 @@ bool FGroomBuilder::BuildHairDescriptionGroups(const FHairDescription& HairDescr
 			StrandWidth = StrandWidths[StrandID];
 		}
 
+		if (!bIsGuide)
+		{
+			GroupCurvesMapping[Group.Info.GroupID].Add(CurveIndex);
+		}
+
 		for (int32 VertexIndex = 0; VertexIndex < CurveNumVertices; ++VertexIndex, ++GlobalVertexIndex)
 		{
 			FVertexID VertexID(GlobalVertexIndex);
+			if(!bIsGuide)
+			{
+				GroupPointsMapping[Group.Info.GroupID].Add(GroupSourceOffsets[Group.Info.GroupID]);
+			}
 
 			// Skip vertex beyond the supported limit. Continue to loop for incrementing GlobalVertexIndex
 			if (VertexIndex >= HAIR_MAX_NUM_POINT_PER_CURVE)
@@ -2279,6 +2364,11 @@ bool FGroomBuilder::BuildHairDescriptionGroups(const FHairDescription& HairDescr
 				check(P0Index > 0);
 				check(P1Index >= 0);
 
+				if(!bIsGuide)
+				{
+					GroupPointsMapping[Group.Info.GroupID].Add(GroupSourceOffsets[Group.Info.GroupID]);
+				}
+
 				// Add a small position offset to avoid superposed positions of the extra position
 				const FVector3f P0 = CurrentHairStrandsDatas->StrandsPoints.PointsPosition[P0Index];
 				const FVector3f P1 = CurrentHairStrandsDatas->StrandsPoints.PointsPosition[P1Index];
@@ -2299,6 +2389,10 @@ bool FGroomBuilder::BuildHairDescriptionGroups(const FHairDescription& HairDescr
 				}
 				CurrentHairStrandsDatas->StrandsPoints.PointsRadius.Add(VertexWidth * 0.5f);
 			}
+			if(!bIsGuide)
+			{
+				GroupSourceOffsets[Group.Info.GroupID]++;
+			}
 		}
 	}
 
@@ -2344,12 +2438,16 @@ bool FGroomBuilder::BuildHairDescriptionGroups(const FHairDescription& HairDescr
 		Group.Info.NumGuideVertices = Group.Guides.GetNumPoints();
 
 		// Prepare/Resize all attributes based on position/curves counts
-		Group.Strands.StrandsPoints.SetNum(Group.Strands.GetNumPoints(), 0u);
-		Group.Strands.StrandsCurves.SetNum(Group.Strands.GetNumCurves(), 0u);
+		Group.Strands.StrandsPoints.SetNum(Group.Strands.GetNumPoints(), 0u, true /*bHasMapping*/);
+		Group.Strands.StrandsCurves.SetNum(Group.Strands.GetNumCurves(), 0u, true /*bHasMapping*/);
 
 		// Prepare/Resize all attributes based on position/curves counts
 		Group.Guides.StrandsPoints.SetNum(Group.Guides.GetNumPoints(), 0u);
 		Group.Guides.StrandsCurves.SetNum(Group.Guides.GetNumCurves(), 0u);
+
+		// Set the strands points mapping
+		Group.Strands.StrandsPoints.PointsMapping = GroupPointsMapping[Group.Info.GroupID];
+		Group.Strands.StrandsCurves.CurvesMapping = GroupCurvesMapping[Group.Info.GroupID];
 	}
 
 	return true;
@@ -2366,7 +2464,8 @@ void FGroomBuilder::BuildData(
 	FHairStrandsDatas& OutRen,
 	FHairStrandsDatas& OutSim,
 	const bool bAllowCurveReordering,
-	const bool bApplyDecimation)
+	const bool bApplyDecimation,
+	const bool bBuildSourceMapping)
 {
 	FHairGroupInfo DummyGroupInfo;
 	BuildData(
@@ -2376,7 +2475,8 @@ void FGroomBuilder::BuildData(
 		OutRen,
 		OutSim,
 		bAllowCurveReordering,
-		bApplyDecimation);
+		bApplyDecimation,
+		bBuildSourceMapping);
 }
 
 void FGroomBuilder::BuildData(
@@ -2386,7 +2486,8 @@ void FGroomBuilder::BuildData(
 	FHairStrandsDatas& OutRen,
 	FHairStrandsDatas& OutSim,
 	const bool bAllowCurveReordering,
-	const bool bApplyDecimation)
+	const bool bApplyDecimation,
+	const bool bBuildSourceMapping)
 {
 	TRACE_CPUPROFILER_EVENT_SCOPE(FGroomBuilder::BuildData);
 
@@ -2413,7 +2514,7 @@ void FGroomBuilder::BuildData(
 			{
 				FHairStrandsDatas FullData = OutRen;
 				OutRen.Reset();
-				FHairStrandsDecimation::Decimate(FullData, CurveDecimation, VertexDecimation, bCurveReordering, OutRen);
+				FHairStrandsDecimation::Decimate(FullData, CurveDecimation, VertexDecimation, bCurveReordering, OutRen, bBuildSourceMapping);
 			}
 			
 			OutGroupInfo.NumCurves			= OutRen.GetNumCurves();
@@ -2441,19 +2542,19 @@ void FGroomBuilder::BuildData(
 					{
 						// Pick the new guides among the imported ones
 						FHairStrandsDatas TempSim = HairStrandsBuilder::BuildInternalData(InHairDescriptionGroup.Guides);
-						FHairStrandsDecimation::Decimate(TempSim, InSettings.InterpolationSettings.RiggedGuideNumCurves, InSettings.InterpolationSettings.RiggedGuideNumPoints, OutSim);
+						FHairStrandsDecimation::Decimate(TempSim, InSettings.InterpolationSettings.RiggedGuideNumCurves, InSettings.InterpolationSettings.RiggedGuideNumPoints, OutSim, bBuildSourceMapping);
 					}
 					else
 					{
 						// Otherwise let s pick the guides among the rendered strands
 						check(OutRen.GetNumPoints() > 0);
-						FHairStrandsDecimation::Decimate(OutRen, InSettings.InterpolationSettings.RiggedGuideNumCurves, InSettings.InterpolationSettings.RiggedGuideNumPoints, OutSim);
+						FHairStrandsDecimation::Decimate(OutRen, InSettings.InterpolationSettings.RiggedGuideNumCurves, InSettings.InterpolationSettings.RiggedGuideNumPoints, OutSim, bBuildSourceMapping);
 					}
 				}
 				else if (!bHasImportedGuides || InSettings.InterpolationSettings.GuideType == EGroomGuideType::Generated)
 				{
 					check(OutRen.GetNumPoints() > 0);
-					FHairStrandsDecimation::Decimate(OutRen, HairToGuideDensity, 1, false, OutSim);
+					FHairStrandsDecimation::Decimate(OutRen, HairToGuideDensity, 1, false, OutSim, bBuildSourceMapping);
 				}
 				else
 				{
@@ -2611,7 +2712,8 @@ void Decimate(
 		const FHairStrandsDatas& InData, 
 		const uint32 NumCurves,
 		const int32 NumVertices,
-		FHairStrandsDatas& OutData)
+		FHairStrandsDatas& OutData,
+		const bool bBuildSourceMapping)
 {
 	const uint32 InCurveCount = InData.StrandsCurves.Num();
 	const uint32 OutCurveCount = FMath::Clamp(NumCurves, 1u, InCurveCount);
@@ -2635,8 +2737,8 @@ void Decimate(
 	const uint32 CurveCount = PointsSampler.SampleIndices.Num();
 	uint32 PointCount = CurveCount * NumVertices;
 
-	OutData.StrandsCurves.SetNum(CurveCount, InAttributes);
-	OutData.StrandsPoints.SetNum(PointCount, InAttributes);
+	OutData.StrandsCurves.SetNum(CurveCount, InAttributes, bBuildSourceMapping);
+	OutData.StrandsPoints.SetNum(PointCount, InAttributes, bBuildSourceMapping);
 	OutData.HairDensity = InData.HairDensity;
 
 	PointCount = 0;
@@ -2692,7 +2794,8 @@ void Decimate(
 	float CurveDecimationPercentage, 
 	float VertexDecimationPercentage, 
 	bool bContinuousDecimationReordering,
-	FHairStrandsDatas& OutData)
+	FHairStrandsDatas& OutData,
+	const bool bBuildSourceMapping)
 {
 	CurveDecimationPercentage = FMath::Clamp(CurveDecimationPercentage, 0.f, 1.f);
 	VertexDecimationPercentage = FMath::Clamp(VertexDecimationPercentage, 0.f, 1.f);
@@ -2741,8 +2844,8 @@ void Decimate(
 		}
 	}
 
-	OutData.StrandsCurves.SetNum(OutCurveCount, InAttributes);
-	OutData.StrandsPoints.SetNum(OutTotalPointCount, InAttributes);
+	OutData.StrandsCurves.SetNum(OutCurveCount, InAttributes, bBuildSourceMapping);
+	OutData.StrandsPoints.SetNum(OutTotalPointCount, InAttributes, bBuildSourceMapping);
 	OutData.HairDensity = InData.HairDensity;
 
 	uint32 OutPointOffset = 0; 
diff --git a/Engine/Plugins/Runtime/HairStrands/Source/HairStrandsCore/Private/GroomDebug.cpp b/Engine/Plugins/Runtime/HairStrands/Source/HairStrandsCore/Private/GroomDebug.cpp
--- a/Engine/Plugins/Runtime/HairStrands/Source/HairStrandsCore/Private/GroomDebug.cpp
+++ b/Engine/Plugins/Runtime/HairStrands/Source/HairStrandsCore/Private/GroomDebug.cpp
@@ -885,6 +885,8 @@ static void AddHairDebugPrintInstancePass(
 				D.Data3.Z |= FFloat16(Instance->HairGroupPublicData->VFInput.Strands.Common.LengthScale).Encoded << 16u;
 				D.Data3.W |= FFloat16(Instance->HairGroupPublicData->GetActiveStrandsCoverageScale()).Encoded;
 				D.Data3.W |= (Instance->HairGroupPublicData->bAutoLOD || UseHairStrandsForceAutoLOD()) ? (1u << 16u) : 0;
+				D.Data3.W |= Instance->HairGroupPublicData->VFInput.Strands.Common.bCurveRemapping ? (1u << 17u) : 0;
+				D.Data3.W |= Instance->HairGroupPublicData->VFInput.Strands.Common.bPointRemapping ? (1u << 18u) : 0;
 
 				#if RHI_RAYTRACING
 				bHasRaytracing = Instance->Strands.RenRaytracingResource != nullptr;
diff --git a/Engine/Plugins/Runtime/HairStrands/Source/HairStrandsCore/Private/GroomEdit.cpp b/Engine/Plugins/Runtime/HairStrands/Source/HairStrandsCore/Private/GroomEdit.cpp
--- a/Engine/Plugins/Runtime/HairStrands/Source/HairStrandsCore/Private/GroomEdit.cpp
+++ b/Engine/Plugins/Runtime/HairStrands/Source/HairStrandsCore/Private/GroomEdit.cpp
@@ -166,7 +166,7 @@ static void EditableGroomGroupToHairDescription(FHairDescription& OutHairDescrip
 }
 #endif // WITH_EDITORONLY_DATA
 
-void ConvertFromGroomAsset(UGroomAsset* InGroom, FEditableGroom* OutGroom, const bool bAllowCurveReordering, const bool bApplyDecimation)
+void ConvertFromGroomAsset(UGroomAsset* InGroom, FEditableGroom* OutGroom, const bool bAllowCurveReordering, const bool bApplyDecimation, const bool bAllowAddEndControlPoint)
 {
 #if WITH_EDITORONLY_DATA
 	check(InGroom);
@@ -174,7 +174,8 @@ void ConvertFromGroomAsset(UGroomAsset* InGroom, FEditableGroom* OutGroom, const
 	UGroomAsset& In = *InGroom;
 	FEditableGroom& Out = *OutGroom;
 
-	const FHairDescriptionGroups& HairDescriptionGroups = In.GetHairDescriptionGroups();
+	FHairDescriptionGroups HairDescriptionGroups;
+	FGroomBuilder::BuildHairDescriptionGroups(In.GetHairDescription(), HairDescriptionGroups, bAllowAddEndControlPoint);
 
 	const uint32 GroupCount = In.GetNumHairGroups();
 	Out.Groups.SetNum(GroupCount);
diff --git a/Engine/Plugins/Runtime/HairStrands/Source/HairStrandsCore/Private/GroomResources.cpp b/Engine/Plugins/Runtime/HairStrands/Source/HairStrandsCore/Private/GroomResources.cpp
--- a/Engine/Plugins/Runtime/HairStrands/Source/HairStrandsCore/Private/GroomResources.cpp
+++ b/Engine/Plugins/Runtime/HairStrands/Source/HairStrandsCore/Private/GroomResources.cpp
@@ -1203,7 +1203,7 @@ void FHairMeshesDeformedResource::InternalRelease()
 
 FHairStrandsRestResource::FHairStrandsRestResource(FHairStrandsBulkData& InBulkData, EHairStrandsResourcesType InCurveType, const FHairResourceName& InResourceName, const FName& InOwnerName) :
 	FHairCommonResource(EHairStrandsAllocationType::Deferred, InResourceName, InOwnerName),
-	PositionBuffer(), PointAttributeBuffer(), CurveAttributeBuffer(), PointToCurveBuffer(), BulkData(InBulkData), CurveType(InCurveType)
+	PositionBuffer(), PointAttributeBuffer(), CurveAttributeBuffer(), CurveMappingBuffer(), PointMappingBuffer(), PointToCurveBuffer(), BulkData(InBulkData), CurveType(InCurveType)
 {
 	MaxAvailableCurveCount = 0;
 
@@ -1285,6 +1285,30 @@ void FHairStrandsRestResource::InternalAllocate(FRDGBuilder& GraphBuilder)
 		InternalCreateByteAddressBufferRDG(GraphBuilder, DummyAttribute, EPixelFormat::PF_R32_UINT, PointAttributeBuffer, ToHairResourceDebugName(HAIRSTRANDS_RESOUCE_NAME(CurveType, Hair.StrandsRest_PointAttributeBuffer), ResourceName), OwnerName, EHairResourceUsageType::Static);
 		GraphBuilder.UseExternalAccessMode(Register(GraphBuilder, PointAttributeBuffer, ERDGImportedBufferFlags::CreateSRV).Buffer, ERHIAccess::SRVMask);
 	}
+	if (!!(BulkData.Header.Flags & FHairStrandsBulkData::DataFlags_HasPointMapping))
+	{
+		InternalCreateByteAddressBufferRDG_FromHairBulkData(GraphBuilder, BulkData.Data.PointMapping, PointMappingBuffer, ToHairResourceDebugName(HAIRSTRANDS_RESOUCE_NAME(CurveType, Hair.StrandsRest_PointMappingBuffer), ResourceName), OwnerName, EHairResourceUsageType::Static);
+	}
+	else
+	{
+		TArray<uint32> DummyAttribute;
+		DummyAttribute.Add(0u);
+		InternalCreateByteAddressBufferRDG(GraphBuilder, DummyAttribute, EPixelFormat::PF_R32_UINT, PointMappingBuffer, ToHairResourceDebugName(HAIRSTRANDS_RESOUCE_NAME(CurveType, Hair.StrandsRest_PointMappingBuffer), ResourceName), OwnerName, EHairResourceUsageType::Static);
+		GraphBuilder.UseExternalAccessMode(Register(GraphBuilder, PointMappingBuffer, ERDGImportedBufferFlags::CreateSRV).Buffer, ERHIAccess::SRVMask);
+	}
+
+	if (!!(BulkData.Header.Flags & FHairStrandsBulkData::DataFlags_HasCurveMapping))
+	{
+		InternalCreateByteAddressBufferRDG_FromHairBulkData(GraphBuilder, BulkData.Data.CurveMapping, CurveMappingBuffer, ToHairResourceDebugName(HAIRSTRANDS_RESOUCE_NAME(CurveType, Hair.StrandsRest_CurveMappingBuffer), ResourceName), OwnerName, EHairResourceUsageType::Static);
+	}
+	else if(CurveMappingBuffer.Buffer == nullptr)
+	{
+		TArray<uint32> DummyAttribute;
+		DummyAttribute.Add(0u);
+		InternalCreateByteAddressBufferRDG(GraphBuilder, DummyAttribute, EPixelFormat::PF_R32_UINT, CurveMappingBuffer, ToHairResourceDebugName(HAIRSTRANDS_RESOUCE_NAME(CurveType, Hair.StrandsRest_CurveMappingBuffer), ResourceName), OwnerName, EHairResourceUsageType::Static);
+		GraphBuilder.UseExternalAccessMode(Register(GraphBuilder, CurveMappingBuffer, ERDGImportedBufferFlags::CreateSRV).Buffer, ERHIAccess::SRVMask);
+	}
+	
 	InternalCreateByteAddressBufferRDG_FromHairBulkData(GraphBuilder, BulkData.Data.Curves, CurveBuffer, ToHairResourceDebugName(HAIRSTRANDS_RESOUCE_NAME(CurveType, Hair.StrandsRest_CurveBuffer), ResourceName), OwnerName, EHairResourceUsageType::Static);
 
 	InternalCreateByteAddressBufferRDG_FromHairBulkData(GraphBuilder, BulkData.Data.PointToCurve, PointToCurveBuffer, ToHairResourceDebugName(HAIRSTRANDS_RESOUCE_NAME(CurveType, Hair.StrandsRest_PointToCurveBuffer), ResourceName), OwnerName, EHairResourceUsageType::Static);
@@ -1353,6 +1377,8 @@ void FHairStrandsRestResource::InternalRelease()
 	PointToCurveBuffer.Release();
 	TangentBuffer.Release();
 	CurveBuffer.Release();
+	CurveMappingBuffer.Release();
+	PointMappingBuffer.Release();
 	MaxAvailableCurveCount = 0;
 	CachedTangentPointCount = 0;
 }
diff --git a/Engine/Plugins/Runtime/HairStrands/Source/HairStrandsCore/Private/HairAttributes.cpp b/Engine/Plugins/Runtime/HairStrands/Source/HairStrandsCore/Private/HairAttributes.cpp
--- a/Engine/Plugins/Runtime/HairStrands/Source/HairStrandsCore/Private/HairAttributes.cpp
+++ b/Engine/Plugins/Runtime/HairStrands/Source/HairStrandsCore/Private/HairAttributes.cpp
@@ -130,6 +130,11 @@ void FHairStrandsDatas::CopyPoint(const FHairStrandsDatas& In, FHairStrandsDatas
 	if (HasHairAttribute(InAttributes, EHairAttribute::Color))		{ Out.StrandsPoints.PointsBaseColor[OutIndex] = In.StrandsPoints.PointsBaseColor[InIndex]; }
 	if (HasHairAttribute(InAttributes, EHairAttribute::Roughness))	{ Out.StrandsPoints.PointsRoughness[OutIndex] = In.StrandsPoints.PointsRoughness[InIndex]; }
 	if (HasHairAttribute(InAttributes, EHairAttribute::AO))			{ Out.StrandsPoints.PointsAO[OutIndex] = In.StrandsPoints.PointsAO[InIndex]; }
+
+	if(!Out.StrandsPoints.PointsMapping.IsEmpty())
+	{ 
+		Out.StrandsPoints.PointsMapping[OutIndex] = In.StrandsPoints.PointsMapping.IsValidIndex(InIndex) ? In.StrandsPoints.PointsMapping[InIndex] : InIndex;
+	}
 }
 
 
@@ -143,6 +148,12 @@ void FHairStrandsDatas::CopyPointLerp(const FHairStrandsDatas& In, FHairStrandsD
 	if (HasHairAttribute(InAttributes, EHairAttribute::Color))		{ Out.StrandsPoints.PointsBaseColor[OutIndex]	= FMath::Lerp(In.StrandsPoints.PointsBaseColor[InIndex0], In.StrandsPoints.PointsBaseColor[InIndex1], InAlpha); }
 	if (HasHairAttribute(InAttributes, EHairAttribute::Roughness))	{ Out.StrandsPoints.PointsRoughness[OutIndex]	= FMath::Lerp(In.StrandsPoints.PointsRoughness[InIndex0], In.StrandsPoints.PointsRoughness[InIndex1], InAlpha); }
 	if (HasHairAttribute(InAttributes, EHairAttribute::AO))			{ Out.StrandsPoints.PointsAO[OutIndex]			= FMath::Lerp(In.StrandsPoints.PointsAO[InIndex0], In.StrandsPoints.PointsAO[InIndex1], InAlpha); }
+
+	if (!Out.StrandsPoints.PointsMapping.IsEmpty())
+	{
+		const int32 PointIndex = (InAlpha < 0.5) ? InIndex0 : InIndex1;
+		Out.StrandsPoints.PointsMapping[OutIndex] = In.StrandsPoints.PointsMapping.IsValidIndex(PointIndex) ? In.StrandsPoints.PointsMapping[PointIndex] : PointIndex;
+	}
 }
 
 void FHairStrandsDatas::CopyCurve(const FHairStrandsDatas& In, FHairStrandsDatas& Out, uint32 InAttributes, uint32 InIndex, uint32 OutIndex)
@@ -165,4 +176,9 @@ void FHairStrandsDatas::CopyCurve(const FHairStrandsDatas& In, FHairStrandsDatas
 	{
 		Out.StrandsCurves.ClumpIDs[OutIndex] = In.StrandsCurves.ClumpIDs[InIndex];
 	}
+
+	if (!Out.StrandsCurves.CurvesMapping.IsEmpty())
+	{
+		Out.StrandsCurves.CurvesMapping[OutIndex] = In.StrandsCurves.CurvesMapping.IsValidIndex(InIndex) ? In.StrandsCurves.CurvesMapping[InIndex] : InIndex;
+	}
 }
\ No newline at end of file
diff --git a/Engine/Plugins/Runtime/HairStrands/Source/HairStrandsCore/Private/HairStrandsDatas.cpp b/Engine/Plugins/Runtime/HairStrands/Source/HairStrandsCore/Private/HairStrandsDatas.cpp
--- a/Engine/Plugins/Runtime/HairStrands/Source/HairStrandsCore/Private/HairStrandsDatas.cpp
+++ b/Engine/Plugins/Runtime/HairStrands/Source/HairStrandsCore/Private/HairStrandsDatas.cpp
@@ -53,9 +53,10 @@ void FHairStrandsCurves::Reset()
 	ClumpIDs.Reset();
 	CurvesClosestGuideIDs.Reset();
 	CurvesClosestGuideWeights.Reset();
+	CurvesMapping.Reset();
 }
 
-void FHairStrandsCurves::SetNum(const uint32 NumCurves, uint32 InAttributes)
+void FHairStrandsCurves::SetNum(const uint32 NumCurves, uint32 InAttributes, bool bHasMapping)
 {
 	CurvesOffset.SetNum(NumCurves + 1);
 	CurvesCount.SetNum(NumCurves);
@@ -70,6 +71,11 @@ void FHairStrandsCurves::SetNum(const uint32 NumCurves, uint32 InAttributes)
 		CurvesClosestGuideIDs.SetNum(NumCurves);
 		CurvesClosestGuideWeights.SetNum(NumCurves);
 	}
+
+	if(bHasMapping)
+	{
+		CurvesMapping.SetNum(NumCurves);
+	}
 }
 
 void FHairStrandsPoints::Reset()
@@ -80,9 +86,10 @@ void FHairStrandsPoints::Reset()
 	PointsBaseColor.Reset();
 	PointsRoughness.Reset();
 	PointsAO.Reset();
+	PointsMapping.Reset();
 }
 
-void FHairStrandsPoints::SetNum(const uint32 NumPoints, uint32 InAttributes)
+void FHairStrandsPoints::SetNum(const uint32 NumPoints, uint32 InAttributes, bool bHasMapping)
 {
 	PointsPosition.SetNum(NumPoints);
 	PointsRadius.SetNum(NumPoints);
@@ -92,6 +99,11 @@ void FHairStrandsPoints::SetNum(const uint32 NumPoints, uint32 InAttributes)
 	if (HasHairAttribute(InAttributes, EHairAttribute::Color))		{ PointsBaseColor.SetNum(NumPoints); }
 	if (HasHairAttribute(InAttributes, EHairAttribute::Roughness))	{ PointsRoughness.SetNum(NumPoints); }
 	if (HasHairAttribute(InAttributes, EHairAttribute::AO))			{ PointsAO.SetNum(NumPoints); }
+
+	if(bHasMapping)
+	{
+		PointsMapping.SetNum(NumPoints);
+	}
 }
 
 void FHairStrandsRawDatas::Reset()
@@ -604,6 +616,8 @@ void FHairStrandsBulkData::SerializeHeader(FArchive& Ar, UObject* Owner)
 	Ar << Header.Strides.CurveAttributeChunkStride;
 	Ar << Header.Strides.PointAttributeChunkStride;
 	Ar << Header.Strides.TranscodedPositionChunkStride;
+	Ar << Header.Strides.CurveMappingStride;
+	Ar << Header.Strides.PointMappingStride;
 
 	Ar << Header.Strides.PointToCurveChunkElementCount;
 	Ar << Header.Strides.CurveAttributeChunkElementCount;
@@ -622,7 +636,7 @@ void FHairStrandsBulkData::GetResourceVersion(FArchive& Ar) const
 
 uint32 FHairStrandsBulkData::GetResourceCount() const
 {
-	return 5;
+	return 7;
 }
 
 void FHairStrandsBulkData::GetResources(FHairStrandsBulkCommon::FQuery& Out)
@@ -663,6 +677,15 @@ void FHairStrandsBulkData::GetResources(FHairStrandsBulkCommon::FQuery& Out)
 		}
 		Out.Add(Data.PointToCurve, 			TEXT("_PointToCurve"), 			Data.PointToCurve.LoadedSize,		PointToCurveIndexSize);
 		Out.Add(Data.Curves, 				TEXT("_Curves"), 				Data.Curves.LoadedSize, 			CurveCount * Header.Strides.CurveStride);
+
+		if(Header.Flags & DataFlags_HasPointMapping)
+		{
+			Out.Add(Data.PointMapping, TEXT("_PointMapping"), Data.PointMapping.LoadedSize, PointCount * Header.Strides.PointMappingStride);
+		}
+		if(Header.Flags & DataFlags_HasCurveMapping)
+		{
+			Out.Add(Data.CurveMapping, TEXT("_CurveMapping"), Data.CurveMapping.LoadedSize, CurveCount * Header.Strides.CurveMappingStride);
+		}
 	}
 }
 
@@ -675,6 +698,8 @@ uint32 FHairStrandsBulkData::GetSize() const
 	Out += Data.PointAttributes.GetBulkDataSize();
 	Out += Data.PointToCurve.GetBulkDataSize();
 	Out += Data.Curves.GetBulkDataSize();
+	Out += Data.CurveMapping.GetBulkDataSize();
+	Out += Data.PointMapping.GetBulkDataSize();
 	return Out;
 }
 
@@ -696,6 +721,8 @@ void FHairStrandsBulkData::Reset()
 	Data.PointAttributes.RemoveBulkData();
 	Data.PointToCurve.RemoveBulkData();
 	Data.Curves.RemoveBulkData();
+	Data.CurveMapping.RemoveBulkData();
+	Data.PointMapping.RemoveBulkData();
 
 	// Reset the bulk byte buffer to ensure the (serialize) data size is reset to 0
 	Data.Positions 		= FHairBulkContainer();
@@ -704,6 +731,8 @@ void FHairStrandsBulkData::Reset()
 	Data.PointAttributes= FHairBulkContainer();
 	Data.PointToCurve	= FHairBulkContainer();
 	Data.Curves			= FHairBulkContainer();
+	Data.CurveMapping	= FHairBulkContainer();
+	Data.PointMapping	= FHairBulkContainer();
 }
 
 void FHairStrandsBulkData::ResetLoadedSize()
@@ -714,6 +743,8 @@ void FHairStrandsBulkData::ResetLoadedSize()
 	Data.PointAttributes.LoadedSize	= 0;
 	Data.PointToCurve.LoadedSize	= 0;
 	Data.Curves.LoadedSize			= 0;
+	Data.CurveMapping.LoadedSize	= 0;
+	Data.PointMapping.LoadedSize	= 0;
 }
 
 float FHairStrandsBulkData::GetCoverageScale(float InCurvePercentage /*[0..1]*/) const
@@ -743,6 +774,8 @@ bool FHairStrandsBulkData::IsCookedBulkDataLoadedForCPUAccess() const
 	bIsAllBulkDataLoaded &= Data.PointAttributes.IsBulkDataEmptyOrLoaded();
 	bIsAllBulkDataLoaded &= Data.PointToCurve.IsBulkDataEmptyOrLoaded();
 	bIsAllBulkDataLoaded &= Data.Curves.IsBulkDataEmptyOrLoaded();
+	bIsAllBulkDataLoaded &= Data.CurveMapping.IsBulkDataEmptyOrLoaded();
+	bIsAllBulkDataLoaded &= Data.PointMapping.IsBulkDataEmptyOrLoaded();
 
 	return bIsAllBulkDataLoaded;
 }
diff --git a/Engine/Plugins/Runtime/HairStrands/Source/HairStrandsCore/Private/HairStrandsInterpolation.cpp b/Engine/Plugins/Runtime/HairStrands/Source/HairStrandsCore/Private/HairStrandsInterpolation.cpp
--- a/Engine/Plugins/Runtime/HairStrands/Source/HairStrandsCore/Private/HairStrandsInterpolation.cpp
+++ b/Engine/Plugins/Runtime/HairStrands/Source/HairStrandsCore/Private/HairStrandsInterpolation.cpp
@@ -1807,13 +1807,15 @@ FHairGroupPublicData::FVertexFactoryInput InternalComputeHairStrandsVertexInputD
 	OutVFInput.Strands.Common.Flags |= Instance->Strands.Modifier.bScatterSceneLighting ? HAIR_FLAGS_SCATTER_SCENE_LIGHT : 0u;
 	OutVFInput.Strands.Common.Flags |= Instance->Strands.Modifier.bUseStableRasterization ? HAIR_FLAGS_STABLE_RASTER : 0u;
 	OutVFInput.Strands.Common.Flags |= bRaytracingGeometry ? HAIR_FLAGS_RAYTRACING_GEOMETRY : 0u;
-	OutVFInput.Strands.Common.Flags |= Instance->Debug.Proxy->Holdout() ? HAIR_FLAGS_HOLDOUT : 0u;
-	OutVFInput.Strands.Common.Flags |= Instance->Debug.Proxy->CastsContactShadow() ? HAIR_FLAGS_CAST_CONTACT_SHADOW : 0u;
+	OutVFInput.Strands.Common.Flags |= (Instance->Debug.Proxy && Instance->Debug.Proxy->Holdout()) ? HAIR_FLAGS_HOLDOUT : 0u;
+	OutVFInput.Strands.Common.Flags |= (Instance->Debug.Proxy && Instance->Debug.Proxy->CastsContactShadow()) ? HAIR_FLAGS_CAST_CONTACT_SHADOW : 0u;
 	OutVFInput.Strands.Common.RaytracingProceduralSplits = GetHairRaytracingProceduralSplits();
 	OutVFInput.Strands.Common.GroupIndex = Instance->Debug.GroupIndex;
 	OutVFInput.Strands.Common.GroupCount = Instance->Debug.GroupCount;
 	OutVFInput.Strands.Common.bSimulation = (Instance->Guides.bIsSimulationEnable || Instance->Guides.bIsSimulationCacheEnable) ? 1 : 0;
-	OutVFInput.Strands.Common.bSingleGuide = (Instance->Strands.InterpolationResource && (Instance->Strands.InterpolationResource->BulkData.Header.Flags & FHairStrandsInterpolationBulkData::DataFlags_HasSingleGuideData)) ? 1 : 0;
+	OutVFInput.Strands.Common.bSingleGuide = (Instance->Strands.InterpolationResource && (Instance->Strands.InterpolationResource->BulkData.Header.Flags & FHairStrandsInterpolationBulkData::DataFlags_HasSingleGuideData) != 0) ? 1 : 0;
+	OutVFInput.Strands.Common.bCurveRemapping = (Instance->Strands.RestResource && (Instance->Strands.RestResource->BulkData.Header.Flags & FHairStrandsBulkData::DataFlags_HasCurveMapping) != 0) ? 1 : 0;
+	OutVFInput.Strands.Common.bPointRemapping = (Instance->Strands.RestResource && (Instance->Strands.RestResource->BulkData.Header.Flags & FHairStrandsBulkData::DataFlags_HasPointMapping) != 0) ? 1 : 0;
 
 	// Add local transform assignement here ...
 
diff --git a/Engine/Plugins/Runtime/HairStrands/Source/HairStrandsCore/Public/GroomAsset.h b/Engine/Plugins/Runtime/HairStrands/Source/HairStrandsCore/Public/GroomAsset.h
--- a/Engine/Plugins/Runtime/HairStrands/Source/HairStrandsCore/Public/GroomAsset.h
+++ b/Engine/Plugins/Runtime/HairStrands/Source/HairStrandsCore/Public/GroomAsset.h
@@ -245,6 +245,8 @@ FORCEINLINE uint32 GetDataSize(const FHairStrandsBulkData& BulkData)
 	Total += BulkData.Data.PointAttributes.IsBulkDataLoaded()	? BulkData.Data.PointAttributes.GetBulkDataSize() : 0;
 	Total += BulkData.Data.Curves.IsBulkDataLoaded() 			? BulkData.Data.Curves.GetBulkDataSize(): 0;
 	Total += BulkData.Data.PointToCurve.IsBulkDataLoaded()		? BulkData.Data.PointToCurve.GetBulkDataSize() : 0;
+	Total += BulkData.Data.CurveMapping.IsBulkDataLoaded()		? BulkData.Data.CurveMapping.GetBulkDataSize() : 0;
+	Total += BulkData.Data.PointMapping.IsBulkDataLoaded()		? BulkData.Data.PointMapping.GetBulkDataSize() : 0;
 	return Total;
 }
 
diff --git a/Engine/Plugins/Runtime/HairStrands/Source/HairStrandsCore/Public/GroomAssetPhysics.h b/Engine/Plugins/Runtime/HairStrands/Source/HairStrandsCore/Public/GroomAssetPhysics.h
--- a/Engine/Plugins/Runtime/HairStrands/Source/HairStrandsCore/Public/GroomAssetPhysics.h
+++ b/Engine/Plugins/Runtime/HairStrands/Source/HairStrandsCore/Public/GroomAssetPhysics.h
@@ -287,6 +287,9 @@ struct FHairGroupsPhysics
 	FHairStrandsParameters StrandsParameters;
 
 	UE_API bool operator==(const FHairGroupsPhysics& A) const;
+
+	/** Build DDC key */
+	void BuildDDCKey(FArchive& Ar);
 };
 
 USTRUCT(BlueprintType)
diff --git a/Engine/Plugins/Runtime/HairStrands/Source/HairStrandsCore/Public/GroomBuilder.h b/Engine/Plugins/Runtime/HairStrands/Source/HairStrandsCore/Public/GroomBuilder.h
--- a/Engine/Plugins/Runtime/HairStrands/Source/HairStrandsCore/Public/GroomBuilder.h
+++ b/Engine/Plugins/Runtime/HairStrands/Source/HairStrandsCore/Public/GroomBuilder.h
@@ -64,7 +64,7 @@ struct FGroomBuilder
 	static UE_API bool BuildHairDescriptionGroups(
 		const FHairDescription& HairDescription, 
 		FHairDescriptionGroups& Out,
-		bool bAllowAddEndControlPoint=true);
+		const bool bAllowAddEndControlPoint = true);
 
 	// 2.a Build FHairStrandsDatas from HairDescriptionGroups(HairStrandsRawData) and DecimationSettings (Strands / Guides)
 	static UE_API void BuildData(
@@ -74,7 +74,8 @@ struct FGroomBuilder
 		FHairStrandsDatas& OutStrands,
 		FHairStrandsDatas& OutGuides,
 		const bool bAllowCurveReordering = true,
-		const bool bApplyDecimation = true);
+		const bool bApplyDecimation = true,
+		const bool bBuildSourceMapping = false);
 
 	static UE_API void BuildData(
 		const FHairDescriptionGroup& InHairDescriptionGroup,
@@ -82,7 +83,8 @@ struct FGroomBuilder
 		FHairStrandsDatas& OutStrands,
 		FHairStrandsDatas& OutGuides,
 		const bool bAllowCurveReordering = true,
-		const bool bApplyDecimation = true);
+		const bool bApplyDecimation = true,
+		const bool bBuildSourceMapping = false);
 
 	// 2.b Build FHairStrandsDatas from a FHairStrandsRawDatas (Strands / Guides)
 	UE_DEPRECATED(5.6, "Please, no longer use this build function, as it does not contains certain build steps (decimation/shuffling)")
diff --git a/Engine/Plugins/Runtime/HairStrands/Source/HairStrandsCore/Public/GroomEdit.h b/Engine/Plugins/Runtime/HairStrands/Source/HairStrandsCore/Public/GroomEdit.h
--- a/Engine/Plugins/Runtime/HairStrands/Source/HairStrandsCore/Public/GroomEdit.h
+++ b/Engine/Plugins/Runtime/HairStrands/Source/HairStrandsCore/Public/GroomEdit.h
@@ -80,7 +80,7 @@ enum EEditableGroomOperations
 };
 
 // Convert a groom asset into an editable groom asset
-HAIRSTRANDSCORE_API void ConvertFromGroomAsset(UGroomAsset* In, FEditableGroom* Out, const bool bAllowCurveReordering = true, const bool bApplyDecimation = true);
+HAIRSTRANDSCORE_API void ConvertFromGroomAsset(UGroomAsset* In, FEditableGroom* Out, const bool bAllowCurveReordering = true, const bool bApplyDecimation = true, const bool bAllowAddEndControlPoint = true);
 
 // Convert an editable groom asset into a groom asset
 // The 'operations' flag indicates what type of modifications have been done onto the editable groom. 
diff --git a/Engine/Plugins/Runtime/HairStrands/Source/HairStrandsCore/Public/GroomResources.h b/Engine/Plugins/Runtime/HairStrands/Source/HairStrandsCore/Public/GroomResources.h
--- a/Engine/Plugins/Runtime/HairStrands/Source/HairStrandsCore/Public/GroomResources.h
+++ b/Engine/Plugins/Runtime/HairStrands/Source/HairStrandsCore/Public/GroomResources.h
@@ -164,6 +164,8 @@ struct FHairStrandsRestResource : public FHairCommonResource
 		Total += GetBufferTotalNumBytes(CurveAttributeBuffer);
 		Total += GetBufferTotalNumBytes(PointToCurveBuffer);
 		Total += GetBufferTotalNumBytes(CurveBuffer);
+		Total += GetBufferTotalNumBytes(CurveMappingBuffer);
+		Total += GetBufferTotalNumBytes(PointMappingBuffer);
 		return Total;
 	}
 
@@ -181,6 +183,12 @@ struct FHairStrandsRestResource : public FHairCommonResource
 
 	/* Strand hair per-curve attribute buffer */
 	FRDGExternalBuffer CurveAttributeBuffer;
+	
+	/* Strand hair per-curve source mapping buffer */
+	FRDGExternalBuffer CurveMappingBuffer;
+
+	/* Strand hair per-point source mapping buffer */
+	FRDGExternalBuffer PointMappingBuffer;
 
 	/* Strand hair vertex to curve index mapping */
 	FRDGExternalBuffer PointToCurveBuffer;
diff --git a/Engine/Plugins/Runtime/HairStrands/Source/HairStrandsCore/Public/HairStrandsDatas.h b/Engine/Plugins/Runtime/HairStrands/Source/HairStrandsCore/Public/HairStrandsDatas.h
--- a/Engine/Plugins/Runtime/HairStrands/Source/HairStrandsCore/Public/HairStrandsDatas.h
+++ b/Engine/Plugins/Runtime/HairStrands/Source/HairStrandsCore/Public/HairStrandsDatas.h
@@ -498,7 +498,7 @@ struct FHairStrandsBulkCommon
 struct FHairStrandsPoints
 {
 	/** Set the number of points */
-	UE_API void SetNum(const uint32 NumPoints, uint32 InAttributes);
+	UE_API void SetNum(const uint32 NumPoints, uint32 InAttributes, bool bHasMapping = false);
 
 	/** Reset the points to 0 */
 	UE_API void Reset();
@@ -524,14 +524,17 @@ struct FHairStrandsPoints
 	TArray<float> PointsRoughness; // [0..1]
 
 	/** Material per-vertex 'baked' AO (optional) */
-	TArray<float> PointsAO; // [0..1]	
+	TArray<float> PointsAO; // [0..1]
+
+	/** Points source mapping */
+	TArray<uint32> PointsMapping;
 };
 
 /** Hair strands Curves attribute */
 struct FHairStrandsCurves
 {
 	/** Set the number of Curves */
-	UE_API void SetNum(const uint32 NumPoints, uint32 InAttributes);
+	UE_API void SetNum(const uint32 NumPoints, uint32 InAttributes, bool bHasMapping = false);
 
 	/** Reset the curves to 0 */
 	UE_API void Reset();
@@ -572,6 +575,9 @@ struct FHairStrandsCurves
 
 	/** Flags for attributes */
 	uint32 AttributeFlags = 0;
+	
+	/** Curves source mapping */
+	TArray<uint32> CurvesMapping;
 };
 
 /** Hair strands datas that are stored on CPU */
@@ -630,6 +636,8 @@ struct FHairStrandsBulkData : FHairStrandsBulkCommon
 		DataFlags_HasTranscodedPosition = 0x4u,	// Contains transcoded position
 		DataFlags_HasTrimmedCurve = 0x8u,		// Source data had more curves per group than the limit and some curves were trimmed
 		DataFlags_HasTrimmedPoint = 0x10u,		// Source data had more points per curve than the limit and some points were trimmed
+		DataFlags_HasPointMapping = 0x20u,		// Contains platform to source points mapping data
+		DataFlags_HasCurveMapping = 0x40u		// Contains platform to source curves mapping data
 	};
 
 	UE_API virtual void SerializeHeader(FArchive& Ar, UObject* Owner) override;
@@ -698,6 +706,8 @@ struct FHairStrandsBulkData : FHairStrandsBulkCommon
 			uint32 CurveAttributeChunkStride = 0;
 			uint32 PointAttributeChunkStride = 0;
 			uint32 TranscodedPositionChunkStride = 0;
+			uint32 CurveMappingStride = 0;
+			uint32 PointMappingStride = 0;
 
 			// Number of element per chunk block
 			uint32 PointToCurveChunkElementCount = 0;
@@ -715,6 +725,8 @@ struct FHairStrandsBulkData : FHairStrandsBulkCommon
 		FHairBulkContainer PointAttributes;		// Size = x*PointCount (depends on the per-point stored attributes)
 		FHairBulkContainer PointToCurve; 		// Size = PointCount
 		FHairBulkContainer Curves;				// Size = CurveCount
+		FHairBulkContainer CurveMapping;		// Size = CurveCount
+		FHairBulkContainer PointMapping;		// Size = PointCount
 	} Data;
 };
 
diff --git a/Engine/Plugins/Runtime/HairStrands/Source/HairStrandsDataflow/Private/GetGroomAssetNode.cpp b/Engine/Plugins/Runtime/HairStrands/Source/HairStrandsDataflow/Private/GetGroomAssetNode.cpp
--- a/Engine/Plugins/Runtime/HairStrands/Source/HairStrandsDataflow/Private/GetGroomAssetNode.cpp
+++ b/Engine/Plugins/Runtime/HairStrands/Source/HairStrandsDataflow/Private/GetGroomAssetNode.cpp
@@ -25,7 +25,7 @@ namespace UE::Groom::Private
 
 		if (GroomAsset)
 		{
-			ConvertFromGroomAsset(const_cast<UGroomAsset*>(GroomAsset), &EditGroom, false, false);
+			ConvertFromGroomAsset(const_cast<UGroomAsset*>(GroomAsset), &EditGroom, false, false, false);
 
 			TArray<FVector3f> PointRestPositions;
 			TArray<int32> ObjectCurveOffsets;
diff --git a/Engine/Plugins/Runtime/HairStrands/Source/HairStrandsDataflow/Private/GroomAssetTerminalNode.cpp b/Engine/Plugins/Runtime/HairStrands/Source/HairStrandsDataflow/Private/GroomAssetTerminalNode.cpp
--- a/Engine/Plugins/Runtime/HairStrands/Source/HairStrandsDataflow/Private/GroomAssetTerminalNode.cpp
+++ b/Engine/Plugins/Runtime/HairStrands/Source/HairStrandsDataflow/Private/GroomAssetTerminalNode.cpp
@@ -132,7 +132,7 @@ void FGroomAssetTerminalDataflowNode::SetAssetValue(TObjectPtr<UObject> Asset, U
 
 			// Build an editable groom asset for the strands
 			FEditableGroom EditGroom;
-			ConvertFromGroomAsset(const_cast<UGroomAsset*>(GroomAsset), &EditGroom, false, false);
+			ConvertFromGroomAsset(const_cast<UGroomAsset*>(GroomAsset), &EditGroom, false, false, false);
 
 			if(GuidesFacade.GetNumObjects() == EditGroom.Groups.Num())
 			{
diff --git a/Engine/Plugins/Runtime/HairStrands/Source/HairStrandsDeformer/Private/DeformerGroomGuidesRead.cpp b/Engine/Plugins/Runtime/HairStrands/Source/HairStrandsDeformer/Private/DeformerGroomGuidesRead.cpp
--- a/Engine/Plugins/Runtime/HairStrands/Source/HairStrandsDeformer/Private/DeformerGroomGuidesRead.cpp
+++ b/Engine/Plugins/Runtime/HairStrands/Source/HairStrandsDeformer/Private/DeformerGroomGuidesRead.cpp
@@ -122,6 +122,8 @@ BEGIN_SHADER_PARAMETER_STRUCT(FOptimusGroomGuidesReadParameters, )
 	SHADER_PARAMETER_RDG_BUFFER_SRV(ByteAddressBuffer, PointRestPositions)
 	SHADER_PARAMETER_RDG_BUFFER_SRV(ByteAddressBuffer, CurvePointOffsets)
 	SHADER_PARAMETER_RDG_BUFFER_SRV(ByteAddressBuffer, PointCurveIndices)
+	SHADER_PARAMETER_RDG_BUFFER_SRV(ByteAddressBuffer, CurveMapping)
+	SHADER_PARAMETER_RDG_BUFFER_SRV(ByteAddressBuffer, PointMapping)
 	SHADER_PARAMETER(FMatrix44f, ObjectRestTransform)
 	SHADER_PARAMETER(FMatrix44f, ObjectDeformedTransform)
 	SHADER_PARAMETER_RDG_BUFFER_SRV(Buffer<float4>, TriangleRestPositions)
@@ -239,6 +241,8 @@ void FOptimusGroomGuidesReadDataProviderProxy::AllocateResources(FRDGBuilder& Gr
 				Resource.PointRestPositions	= RegisterAsSRV(GraphBuilder, GroupInstance->Guides.RestResource->PositionBuffer);
 				Resource.CurvePointOffsets = RegisterAsSRV(GraphBuilder, GroupInstance->Guides.RestResource->CurveBuffer);
 				Resource.PointCurveIndices = RegisterAsSRV(GraphBuilder, GroupInstance->Guides.RestResource->PointToCurveBuffer);
+				Resource.CurveMapping = RegisterAsSRV(GraphBuilder, GroupInstance->Guides.RestResource->CurveMappingBuffer);
+				Resource.PointMapping = RegisterAsSRV(GraphBuilder, GroupInstance->Guides.RestResource->PointMappingBuffer);
 			}
 			{
 				FBindingResources& Resource = BindingResources.AddDefaulted_GetRef();
@@ -292,6 +296,8 @@ void FOptimusGroomGuidesReadDataProviderProxy::GatherDispatchData(FDispatchData
 				Parameters.PointRestPositions = Resource.PointRestPositions;
 				Parameters.CurvePointOffsets = Resource.CurvePointOffsets;
 				Parameters.PointCurveIndices = Resource.PointCurveIndices;
+				Parameters.CurveMapping = Resource.CurveMapping;
+				Parameters.PointMapping = Resource.PointMapping;
 			}
 
 			{
diff --git a/Engine/Plugins/Runtime/HairStrands/Source/HairStrandsDeformer/Private/DeformerGroomGuidesRead.h b/Engine/Plugins/Runtime/HairStrands/Source/HairStrandsDeformer/Private/DeformerGroomGuidesRead.h
--- a/Engine/Plugins/Runtime/HairStrands/Source/HairStrandsDeformer/Private/DeformerGroomGuidesRead.h
+++ b/Engine/Plugins/Runtime/HairStrands/Source/HairStrandsDeformer/Private/DeformerGroomGuidesRead.h
@@ -97,6 +97,8 @@ private:
 		FRDGBufferSRVRef CurvePointOffsets = nullptr;
 		FRDGBufferSRVRef PointRestPositions = nullptr;
 		FRDGBufferSRVRef PointCurveIndices = nullptr;
+		FRDGBufferSRVRef CurveMapping = nullptr;
+		FRDGBufferSRVRef PointMapping = nullptr;
 	};
 	TArray<FInstanceResources> InstanceResources;
 };
diff --git a/Engine/Plugins/Runtime/HairStrands/Source/HairStrandsDeformer/Private/DeformerGroomStrandsRead.cpp b/Engine/Plugins/Runtime/HairStrands/Source/HairStrandsDeformer/Private/DeformerGroomStrandsRead.cpp
--- a/Engine/Plugins/Runtime/HairStrands/Source/HairStrandsDeformer/Private/DeformerGroomStrandsRead.cpp
+++ b/Engine/Plugins/Runtime/HairStrands/Source/HairStrandsDeformer/Private/DeformerGroomStrandsRead.cpp
@@ -36,6 +36,8 @@ TArray<FOptimusCDIPinDefinition> UOptimusGroomStrandsReadDataInterface::GetPinDe
 	Defs.Add({ "PointCurveIndex", "ReadPointCurveIndex", StrandsPoints,          "ReadNumPoints", false, "PointCurveIndex"  });
 	Defs.Add({ "CurveRestTransform", "ReadCurveRestTransform", StrandsCurves,   "ReadNumCurves", false, "CurveRestTransform"  });
 	Defs.Add({ "CurveDeformedTransform", "ReadCurveDeformedTransform", StrandsCurves,          "ReadNumCurves", true, "CurveDeformedTransform"  });
+	Defs.Add({ "CurveSourceIndex", "ReadCurveSourceIndex", StrandsCurves,          "ReadNumCurves", true, "CurveSourceIndex"  });
+	Defs.Add({ "PointSourceIndex", "ReadPointSourceIndex", StrandsPoints,          "ReadNumPoints", true, "PointSourceIndex"  });
 
 	// Interpolation buffers
 	Defs.Add({ "GuideIndex", "ReadPointGuideIndices", StrandsPoints,   "ReadNumPoints", false, "PointGuideIndices"  });
@@ -156,6 +158,16 @@ void UOptimusGroomStrandsReadDataInterface::GetSupportedInputs(TArray<FShaderFun
 		.SetName(TEXT("ReadCurveDeformedTransform"))
 		.AddReturnType(EShaderFundamentalType::Float, 3, 4)
 		.AddParam(EShaderFundamentalType::Uint);
+
+	OutFunctions.AddDefaulted_GetRef()
+		.SetName(TEXT("ReadCurveSourceIndex"))
+		.AddReturnType(EShaderFundamentalType::Uint)
+		.AddParam(EShaderFundamentalType::Uint);
+
+	OutFunctions.AddDefaulted_GetRef()
+		.SetName(TEXT("ReadPointSourceIndex"))
+		.AddReturnType(EShaderFundamentalType::Uint)
+		.AddParam(EShaderFundamentalType::Uint);
 	
 }
 
@@ -292,6 +304,8 @@ void FOptimusGroomStrandsReadProviderProxy::AllocateResources(FRDGBuilder& Graph
 				Resource.PointToCurveBuffer	= RegisterAsSRV(GraphBuilder, GroupInstance->Strands.RestResource->PointToCurveBuffer);
 				Resource.CurveAttributeBuffer = RegisterAsSRV(GraphBuilder, GroupInstance->Strands.RestResource->CurveAttributeBuffer);
 				Resource.PointAttributeBuffer = RegisterAsSRV(GraphBuilder, GroupInstance->Strands.RestResource->PointAttributeBuffer);
+				Resource.CurveMappingBuffer = RegisterAsSRV(GraphBuilder, GroupInstance->Strands.RestResource->CurveMappingBuffer);
+				Resource.PointMappingBuffer = RegisterAsSRV(GraphBuilder, GroupInstance->Strands.RestResource->PointMappingBuffer);
 			}
 
 			{
diff --git a/Engine/Shaders/Private/HairStrands/HairStrandsDebug.usf b/Engine/Shaders/Private/HairStrands/HairStrandsDebug.usf
--- a/Engine/Shaders/Private/HairStrands/HairStrandsDebug.usf
+++ b/Engine/Shaders/Private/HairStrands/HairStrandsDebug.usf
@@ -1474,6 +1474,8 @@ struct FInstanceInfo
 	bool CacheSim;
 	bool bHasLODSwitched;
 	bool bAutoLOD;
+	bool bCurveRemapping;
+	bool bPointRemapping;
 	uint Flags;
 
 	// Binding
@@ -1551,7 +1553,9 @@ FInstanceInfo ReadInstanceInfo(uint Index)
 	Out.ActivePointCount		= Data2.y;
 	Out.ActiveCurveCount		= Data2.z;
 	Out.Attributes				= Data2.w;
-	Out.bAutoLOD				= Data3.w >> 16u;
+	Out.bAutoLOD				= (Data3.w >> 16u) & 0x1;
+	Out.bCurveRemapping			= (Data3.w >> 17u) & 0x1;
+	Out.bPointRemapping			= (Data3.w >> 18u) & 0x1;
 
 	Out.Radius					= f16tof32(Data3.x & 0xFFFF);
 	Out.Density					= f16tof32(Data3.x >> 16u);
@@ -1756,6 +1760,7 @@ void MainCS(uint GroupIndex : SV_GroupIndex, uint3 DispatchThreadId : SV_Dispatc
 		Print(Context, TEXT("Radius "), FontWhite);
 		Print(Context, TEXT("Root  "), FontWhite);
 		Print(Context, TEXT("Tip   "), FontWhite);
+		Print(Context, TEXT("CvId PtId "), FontWhite);
 		const uint AttributeNameInfoCount = NumAttributeNames();
 		LOOP
 		for (uint AttributeIt=0; AttributeIt < AttributeNameInfoCount; ++AttributeIt)
@@ -1790,6 +1795,8 @@ void MainCS(uint GroupIndex : SV_GroupIndex, uint3 DispatchThreadId : SV_Dispatc
 			Print(Context, Info.Radius, FontYellow, 5, 3); Print(Context, TEXT("  "));
 			Print(Context, Info.RootScale, FontYellow, 5, 3); Print(Context, TEXT(" "));
 			Print(Context, Info.TipScale, FontYellow, 5, 3); Print(Context, TEXT(" "));
+			PrintBool(Context, Info.bCurveRemapping); Print(Context, TEXT(" "));
+			PrintBool(Context, Info.bPointRemapping); Print(Context, TEXT(" "));
 
 			for (uint AttributeIt = 0; AttributeIt < AttributeNameInfoCount; ++AttributeIt)
 			{
diff --git a/Engine/Shaders/Private/HairStrands/HairStrandsPack.ush b/Engine/Shaders/Private/HairStrands/HairStrandsPack.ush
--- a/Engine/Shaders/Private/HairStrands/HairStrandsPack.ush
+++ b/Engine/Shaders/Private/HairStrands/HairStrandsPack.ush
@@ -507,3 +507,27 @@ FHairTriangleIndex UnpackTriangleIndex(uint Encoded)
 	Out.TriangleIndex = Encoded & 0xFFFFFF;
 	return Out;
 }
+
+////////////////////////////////////////////////////////////////////////////////////////////////////////////////////////////////
+// Source point/curve remapping
+// Translate a curve/point index into its source data index
+
+uint ReadHairSourceCurveIndex(ByteAddressBuffer InBuffer, uint InIndex, bool bIsBufferValid)
+{
+	uint Out = ~0;
+	if (bIsBufferValid)
+	{
+		Out = InBuffer.Load(InIndex * 4u);
+	}
+	return Out;
+}
+
+uint ReadHairSourcePointIndex(ByteAddressBuffer InBuffer, uint InIndex, bool bIsBufferValid)
+{
+	uint Out = ~0;
+	if (bIsBufferValid)
+	{
+		Out = InBuffer.Load(InIndex * 4u);
+	}
+	return Out;
+}
\ No newline at end of file
diff --git a/Engine/Shaders/Private/HairStrands/HairStrandsVertexFactoryCommon.ush b/Engine/Shaders/Private/HairStrands/HairStrandsVertexFactoryCommon.ush
--- a/Engine/Shaders/Private/HairStrands/HairStrandsVertexFactoryCommon.ush
+++ b/Engine/Shaders/Private/HairStrands/HairStrandsVertexFactoryCommon.ush
@@ -84,6 +84,8 @@ float3 ReadCardPositionOffset(StructuredBuffer<float4> In, uint InIndex, uint In
 	ByteAddressBuffer			Prefix##_PointToCurveBuffer; \
 	ByteAddressBuffer 			Prefix##_CurveAttributeBuffer; \
 	ByteAddressBuffer 			Prefix##_PointAttributeBuffer; \
+	ByteAddressBuffer 			Prefix##_CurveMappingBuffer; \
+	ByteAddressBuffer 			Prefix##_PointMappingBuffer; \
 	Buffer<uint>				Prefix##_CullingIndirectBuffer; \
 	Buffer<uint>				Prefix##_CullingIndexBuffer; \
 	Buffer<uint> 				Prefix##_CullingIndirectBufferArgs; \
diff --git a/Engine/Source/Runtime/Renderer/Private/HairStrands/HairStrandsVisibility.cpp b/Engine/Source/Runtime/Renderer/Private/HairStrands/HairStrandsVisibility.cpp
--- a/Engine/Source/Runtime/Renderer/Private/HairStrands/HairStrandsVisibility.cpp
+++ b/Engine/Source/Runtime/Renderer/Private/HairStrands/HairStrandsVisibility.cpp
@@ -3209,6 +3209,8 @@ void GetHairStrandsInstanceResources(FRDGBuilder& GraphBuilder, const FViewInfo&
 		OutResources.PointToCurveBuffer		= Register(GraphBuilder, VFInput.Strands.PointToCurveBufferExternal, ERDGImportedBufferFlags::CreateSRV).SRV;
 		OutResources.CurveAttributeBuffer	= Register(GraphBuilder, VFInput.Strands.CurveAttributeBufferExternal, ERDGImportedBufferFlags::CreateSRV).SRV;
 		OutResources.PointAttributeBuffer	= Register(GraphBuilder, VFInput.Strands.PointAttributeBufferExternal, ERDGImportedBufferFlags::CreateSRV).SRV;
+		OutResources.CurveMappingBuffer		= Register(GraphBuilder, VFInput.Strands.CurveMappingBufferExternal, ERDGImportedBufferFlags::CreateSRV).SRV;
+		OutResources.PointMappingBuffer		= Register(GraphBuilder, VFInput.Strands.PointMappingBufferExternal, ERDGImportedBufferFlags::CreateSRV).SRV;
 	}
 	else
 	{
@@ -3218,6 +3220,8 @@ void GetHairStrandsInstanceResources(FRDGBuilder& GraphBuilder, const FViewInfo&
 		OutResources.PointToCurveBuffer		= VFInput.Strands.PointToCurveBuffer.SRV;
 		OutResources.CurveAttributeBuffer	= VFInput.Strands.CurveAttributeBuffer.SRV;
 		OutResources.PointAttributeBuffer	= VFInput.Strands.PointAttributeBuffer.SRV;
+		OutResources.CurveMappingBuffer 	= VFInput.Strands.CurveMappingBuffer.SRV;
+		OutResources.PointMappingBuffer 	= VFInput.Strands.PointMappingBuffer.SRV;
 	}
 }
 
diff --git a/Engine/Source/Runtime/Renderer/Public/HairStrandsInterface.h b/Engine/Source/Runtime/Renderer/Public/HairStrandsInterface.h
--- a/Engine/Source/Runtime/Renderer/Public/HairStrandsInterface.h
+++ b/Engine/Source/Runtime/Renderer/Public/HairStrandsInterface.h
@@ -57,6 +57,8 @@ BEGIN_SHADER_PARAMETER_STRUCT(FHairStrandsInstanceCommonParameters, RENDERER_API
 	SHADER_PARAMETER(uint32, Flags)
 	SHADER_PARAMETER(uint32, bSimulation)
 	SHADER_PARAMETER(uint32, bSingleGuide)
+	SHADER_PARAMETER(uint32, bPointRemapping)
+	SHADER_PARAMETER(uint32, bCurveRemapping)
 	SHADER_PARAMETER(FVector3f, PositionOffset)
 	SHADER_PARAMETER(FVector3f, PrevPositionOffset)
 	SHADER_PARAMETER(FMatrix44f, LocalToWorldPrimitiveTransform)
@@ -72,6 +74,8 @@ BEGIN_SHADER_PARAMETER_STRUCT(FHairStrandsInstanceResourceParameters, RENDERER_A
 	SHADER_PARAMETER_RDG_BUFFER_SRV(ByteAddressBuffer, PointToCurveBuffer)
 	SHADER_PARAMETER_RDG_BUFFER_SRV(ByteAddressBuffer, CurveAttributeBuffer)
 	SHADER_PARAMETER_RDG_BUFFER_SRV(ByteAddressBuffer, PointAttributeBuffer)
+	SHADER_PARAMETER_RDG_BUFFER_SRV(ByteAddressBuffer, CurveMappingBuffer)
+	SHADER_PARAMETER_RDG_BUFFER_SRV(ByteAddressBuffer, PointMappingBuffer)
 END_SHADER_PARAMETER_STRUCT()
 
 // Instance prev. resources (RDG)
@@ -306,6 +310,8 @@ public:
 			FRDGImportedBuffer PositionOffsetBuffer;
 			FRDGImportedBuffer PrevPositionOffsetBuffer;
 			FRDGImportedBuffer CurveBuffer;
+			FRDGImportedBuffer CurveMappingBuffer;
+			FRDGImportedBuffer PointMappingBuffer;
 
 			FRDGExternalBuffer PositionBufferExternal;
 			FRDGExternalBuffer PrevPositionBufferExternal;
@@ -316,6 +322,8 @@ public:
 			FRDGExternalBuffer PositionOffsetBufferExternal;
 			FRDGExternalBuffer PrevPositionOffsetBufferExternal;
 			FRDGExternalBuffer CurveBufferExternal;
+			FRDGExternalBuffer CurveMappingBufferExternal;
+			FRDGExternalBuffer PointMappingBufferExternal;
 
 			FShaderResourceViewRHIRef PositionBufferRHISRV				= nullptr;
 			FShaderResourceViewRHIRef PrevPositionBufferRHISRV			= nullptr;
@@ -326,6 +334,8 @@ public:
 			FShaderResourceViewRHIRef PositionOffsetBufferRHISRV		= nullptr;
 			FShaderResourceViewRHIRef PrevPositionOffsetBufferRHISRV	= nullptr;
 			FShaderResourceViewRHIRef CurveBufferRHISRV					= nullptr;
+			FShaderResourceViewRHIRef CurveMappingBufferRHISRV			= nullptr;
+			FShaderResourceViewRHIRef PointMappingBufferRHISRV			= nullptr;
 
 			FHairStrandsInstanceCommonParameters Common; 
 		} Strands;
