Fix typo causing hair width attribute to be flagged as not present while there.

diff --git a/Engine/Plugins/Runtime/HairStrands/Source/HairStrandsCore/Private/HairDescription.cpp b/Engine/Plugins/Runtime/HairStrands/Source/HairStrandsCore/Private/HairDescription.cpp
--- a/Engine/Plugins/Runtime/HairStrands/Source/HairStrandsCore/Private/HairDescription.cpp
+++ b/Engine/Plugins/Runtime/HairStrands/Source/HairStrandsCore/Private/HairDescription.cpp
@@ -73,7 +73,7 @@ bool FHairDescription::HasAttribute(EHairAttribute InAttribute) const
 
 	switch (InAttribute)
 	{
-	case EHairAttribute::Width:			return StrandAttributes().GetAttributesRef<float>(HairAttribute::Vertex::Width).IsValid() || StrandAttributes().GetAttributesRef<float>(HairAttribute::Strand::Width).IsValid() || StrandAttributes().GetAttributesRef<float>(HairAttribute::Groom::Width).IsValid();
+	case EHairAttribute::Width:			return VertexAttributes().GetAttributesRef<float>(HairAttribute::Vertex::Width).IsValid() || StrandAttributes().GetAttributesRef<float>(HairAttribute::Strand::Width).IsValid() || GroomAttributes().GetAttributesRef<float>(HairAttribute::Groom::Width).IsValid();
 	case EHairAttribute::RootUV:		return StrandAttributes().GetAttributesRef<FVector2f>(HairAttribute::Strand::RootUV).IsValid();
 	case EHairAttribute::ClumpID:		return StrandAttributes().GetAttributesRef<int>(HairAttribute::Strand::ClumpID).IsValid() || StrandAttributes().GetAttributesRef<FVector3f>(HairAttribute::Strand::ClumpID).IsValid();
 	case EHairAttribute::StrandID:		return StrandAttributes().GetAttributesRef<int>(HairAttribute::Strand::ID).IsValid();
