Fix crashes when building groom with Dataflow.

diff --git a/Engine/Plugins/Runtime/HairStrands/Source/HairStrandsCore/Private/GroomResources.cpp b/Engine/Plugins/Runtime/HairStrands/Source/HairStrandsCore/Private/GroomResources.cpp
--- a/Engine/Plugins/Runtime/HairStrands/Source/HairStrandsCore/Private/GroomResources.cpp
+++ b/Engine/Plugins/Runtime/HairStrands/Source/HairStrandsCore/Private/GroomResources.cpp
@@ -1064,10 +1064,7 @@ void FHairCardsRestResource::ReleaseResource()
 	NormalsBuffer.ReleaseResource();
 	UVsBuffer.ReleaseResource();
 	MaterialsBuffer.ReleaseResource();
-	if (BulkData.Header.bVertexColor)
-	{
-		VertexColorsBuffer.ReleaseResource();
-	}
+	VertexColorsBuffer.ReleaseResource();
 }
 
 /////////////////////////////////////////////////////////////////////////////////////////
@@ -1166,10 +1163,7 @@ void FHairMeshesRestResource::ReleaseResource()
 	RestPositionBuffer.ReleaseResource();
 	NormalsBuffer.ReleaseResource();
 	UVsBuffer.ReleaseResource();
-	if (BulkData.Header.bVertexColor)
-	{
-		VertexColorsBuffer.ReleaseResource();
-	}
+	VertexColorsBuffer.ReleaseResource();
 }
 
 /////////////////////////////////////////////////////////////////////////////////////////
