Fix incorrect groom binding transfer when source skeletal mesh has different section count based on project settings.

This CL changes uses the matching section as a material index in order to recover the set of section mapped onto this material index. This was the original intent of the matching section index.

diff --git a/Engine/Plugins/Runtime/HairStrands/Source/HairStrandsCore/Private/GroomBindingBuilder.cpp b/Engine/Plugins/Runtime/HairStrands/Source/HairStrandsCore/Private/GroomBindingBuilder.cpp
--- a/Engine/Plugins/Runtime/HairStrands/Source/HairStrandsCore/Private/GroomBindingBuilder.cpp
+++ b/Engine/Plugins/Runtime/HairStrands/Source/HairStrandsCore/Private/GroomBindingBuilder.cpp
@@ -78,7 +78,7 @@ FGroomBindingBuilder::FInput::FInput(UGroomBindingAsset* BindingAsset, const ITa
 FString FGroomBindingBuilder::GetVersion()
 {
 	// Important to update the version when groom building changes
-	return TEXT("5");
+	return TEXT("6a");
 }
 
 ///////////////////////////////////////////////////////////////////////////////////////////////////
@@ -206,6 +206,7 @@ public:
 	virtual uint32 GetNumTriangles() const = 0;
 	virtual uint32 GetBaseIndex() const = 0;
 	virtual uint32 GetBaseVertexIndex() const = 0;
+	virtual uint32 GetMaterialIndex() const = 0;
 	virtual ~IMeshSectionData() {}
 };
 
@@ -227,6 +228,33 @@ public:
 	virtual void SetVertexBinding(uint32 VertexIt, float InterpolatedAttribute) = 0;
 	virtual float GetVertexBinding(uint32 VertexIt) const = 0;
 	virtual ~IMeshLODData() {}
+
+	bool IsValidMaterialIndex(int32 InMaterialIndex) const
+	{
+		for (int32 SectionIt = 0, NumSections = GetNumSections(); SectionIt < NumSections; ++SectionIt)
+		{
+			if (uint32(InMaterialIndex) == GetSection(SectionIt).GetMaterialIndex())
+			{
+				return true;
+			}
+		}
+		return false;
+	}
+
+	TArray<int32> GetSectionsMatchingMaterialIndex(int32 InMaterialIndex) const
+	{
+		// Find all the sections which are mapped onto the target material index
+		TArray<int32> Out;
+		Out.Reserve(GetNumSections());
+		for (int32 SectionIndex=0, NumSections = GetNumSections(); SectionIndex<NumSections; ++SectionIndex)
+		{
+			if (GetSection(SectionIndex).GetMaterialIndex() == InMaterialIndex)
+			{
+				Out.Add(SectionIndex);
+			}
+		}
+		return Out;
+	}
 };
 
 // Interface to wrap the mesh source and query its LOD data
@@ -282,6 +310,13 @@ public:
 		return MeshData->LODRenderData[LODIndex].RenderSections[SectionIndex].BaseVertexIndex;
 	}
 
+	virtual uint32 GetMaterialIndex() const override
+	{
+		check(MeshData->LODRenderData.IsValidIndex(LODIndex));
+		check(MeshData->LODRenderData[LODIndex].RenderSections.IsValidIndex(SectionIndex));
+		return MeshData->LODRenderData[LODIndex].RenderSections[SectionIndex].MaterialIndex;
+	}
+
 private:
 	int32 LODIndex = -1;
 	int32 SectionIndex = -1;
@@ -531,10 +566,11 @@ private:
 class FGeometryCacheSection : public IMeshSectionData
 {
 public:
-	FGeometryCacheSection(FGeometryCacheMeshBatchInfo& InSection, uint32 InNumVertices, uint32 InBaseVertexIdex)
+	FGeometryCacheSection(FGeometryCacheMeshBatchInfo& InSection, uint32 InNumVertices, uint32 InBaseVertexIdex, uint32 InMaterialIndex)
 		: Section(InSection)
 		, NumVertices(InNumVertices)
 		, BaseVertexIndex(InBaseVertexIdex)
+		, MaterialIndex(InMaterialIndex)
 	{
 	}
 
@@ -558,10 +594,16 @@ public:
 		return BaseVertexIndex;
 	}
 
+	virtual uint32 GetMaterialIndex() const override
+	{
+		return MaterialIndex;
+	}
+
 private:
 	FGeometryCacheMeshBatchInfo& Section;
 	uint32 NumVertices;
 	uint32 BaseVertexIndex;
+	uint32 MaterialIndex;
 };
 
 // GeometryCache have only one LOD so FGeometryCacheData provides both mesh source and mesh LOD data
@@ -597,7 +639,7 @@ public:
 						}
 						SectionRanges.Add(SectionRange);
 
-						Sections.Add(FGeometryCacheSection(BatchInfo, SectionRange.Num(), SectionRange.Min));
+						Sections.Add(FGeometryCacheSection(BatchInfo, SectionRange.Num(), SectionRange.Min, BatchInfo.MaterialIndex));
 					}
 				}
 				else
@@ -920,7 +962,7 @@ namespace GroomBinding_RBFWeighting
 
 	void FillLocalValidPoints(
 		const GroomBinding_Mesh::IMeshLODData& MeshLODData, 
-		const int32 TargetSection,
+		const int32 TargetMaterialIndex,
 		const FHairStrandsRootData& ProjectionLOD, 
 		TArray<bool>& OutValidPoints)
 	{
@@ -928,14 +970,15 @@ namespace GroomBinding_RBFWeighting
 
 		OutValidPoints.Init(false, MeshLODData.GetNumVertices());
 
-		const bool ValidSection = (TargetSection >= 0 && TargetSection < MeshLODData.GetNumSections());
-
+		const bool ValidMaterialIndex = (TargetMaterialIndex >= 0 && MeshLODData.IsValidMaterialIndex(TargetMaterialIndex));
+		const TArray<int32> MatchingSections = MeshLODData.GetSectionsMatchingMaterialIndex(TargetMaterialIndex);
+		
 		for (uint32 EncodedTriangleId : ProjectionLOD.UniqueTriangleIndexBuffer)
 		{
 			uint32 SectionIndex = 0;
 			uint32 TriangleIndex = 0;
 			FHairStrandsRootUtils::UnpackTriangleIndex(EncodedTriangleId, TriangleIndex, SectionIndex);
-			if (!ValidSection || (ValidSection && (SectionIndex == TargetSection)))
+			if (!ValidMaterialIndex || (ValidMaterialIndex && (MatchingSections.Find(SectionIndex) != INDEX_NONE)))
 			{
 				const GroomBinding_Mesh::IMeshSectionData& Section = MeshLODData.GetSection(SectionIndex);
 				for (uint32 VertexIt = 0; VertexIt < 3; ++VertexIt)
@@ -947,20 +990,25 @@ namespace GroomBinding_RBFWeighting
 		}
 	}
 
-	void FillGlobalValidPoints(const GroomBinding_Mesh::IMeshLODData& MeshLODData, const int32 TargetSection, TArray<bool>& OutValidPoints)
+	void FillGlobalValidPoints(const GroomBinding_Mesh::IMeshLODData& MeshLODData, const int32 TargetMaterialIndex, TArray<bool>& OutValidPoints)
 	{
 		const TArray<uint32>& TriangleIndices = MeshLODData.GetIndexBuffer(); 
-		if (TargetSection >= 0 && TargetSection < MeshLODData.GetNumSections())
+		if (TargetMaterialIndex >= 0 && MeshLODData.IsValidMaterialIndex(TargetMaterialIndex))
 		{
+			const TArray<int32> MatchingSections = MeshLODData.GetSectionsMatchingMaterialIndex(TargetMaterialIndex);
+
 			OutValidPoints.Init(false, MeshLODData.GetNumVertices());
 
-			const GroomBinding_Mesh::IMeshSectionData& Section = MeshLODData.GetSection(TargetSection);
-			for (uint32 TriangleIt = 0; TriangleIt < Section.GetNumTriangles(); ++TriangleIt)
+			for (int32 MatchingSectionIt : MatchingSections)
 			{
-				for (uint32 VertexIt = 0; VertexIt < 3; ++VertexIt)
+				const GroomBinding_Mesh::IMeshSectionData& Section = MeshLODData.GetSection(MatchingSectionIt);
+				for (uint32 TriangleIt = 0; TriangleIt < Section.GetNumTriangles(); ++TriangleIt)
 				{
-					const uint32 VertexIndex = TriangleIndices[Section.GetBaseIndex() + 3 * TriangleIt + VertexIt];
-					OutValidPoints[VertexIndex] = true;
+					for (uint32 VertexIt = 0; VertexIt < 3; ++VertexIt)
+					{
+						const uint32 VertexIndex = TriangleIndices[Section.GetBaseIndex() + 3 * TriangleIt + VertexIt];
+						OutValidPoints[VertexIndex] = true;
+					}
 				}
 			}
 		}
@@ -982,7 +1030,7 @@ namespace GroomBinding_RBFWeighting
 		FHairRootGroupData& Out, 
 		const bool bNeedStrandsRoot,
 		const uint32 NumInterpolationPoints, 
-		const int32 MatchingSection, 
+		const int32 MatchingMaterialIndex,
 		const GroomBinding_Mesh::IMeshData* MeshData,
 		const int32 TargetMeshMinLOD)
 	{
@@ -1000,14 +1048,14 @@ namespace GroomBinding_RBFWeighting
 
 			const GroomBinding_Mesh::IMeshLODData& MeshLODData = MeshData->GetMeshLODData(MeshLODIndex);
 
-			int32 TargetSection = -1;
+			int32 TargetMaterialIndex = -1;
 			bool GlobalSamples = false;
 			const FVector3f* PositionsPointer = nullptr;
 			if (Out.MeshPositions_Transferred.Num() == MeshLODCount)
 			{
 				PositionsPointer = Out.MeshPositions_Transferred[MeshLODIndex].GetData();
 				GlobalSamples = true;
-				TargetSection = MatchingSection;
+				TargetMaterialIndex = MatchingMaterialIndex;				
 			}
 			else
 			{
@@ -1018,7 +1066,7 @@ namespace GroomBinding_RBFWeighting
 			{
 				TArray<bool> ValidPoints;
 				{
-					FillLocalValidPoints(MeshLODData, TargetSection, Out.SimRootDatas[MeshLODIndex], ValidPoints);
+					FillLocalValidPoints(MeshLODData, TargetMaterialIndex, Out.SimRootDatas[MeshLODIndex], ValidPoints);
 
 					FPointsSampler PointsSampler(ValidPoints, PositionsPointer, MaxSamples);
 					const uint32 SampleCount = PointsSampler.SamplePositions.Num();
@@ -1040,7 +1088,7 @@ namespace GroomBinding_RBFWeighting
 			{
 				TArray<bool> ValidPoints;
 
-				FillGlobalValidPoints(MeshLODData, TargetSection, ValidPoints);
+				FillGlobalValidPoints(MeshLODData, TargetMaterialIndex, ValidPoints);
 
 				FPointsSampler PointsSampler(ValidPoints, PositionsPointer, MaxSamples);
 				const uint32 SampleCount = PointsSampler.SamplePositions.Num();
@@ -1683,8 +1731,7 @@ namespace GroomBinding_RootProjection
 		const GroomBinding_Mesh::IMeshData* InMeshData,
 		int32 TargetMeshMinLOD,
 		const TArray<TArray<FVector3f>>& InTransferredPositions,
-		TArray<FHairStrandsRootData>& OutRootData,
-		uint32 InMatchingSection)
+		TArray<FHairStrandsRootData>& OutRootData)
 	{
 		// 2. Project root for each mesh LOD
 		const uint32 CurveCount = InRootPositions.GetNumCurves();
@@ -2177,54 +2224,59 @@ namespace GroomBinding_Transfer
 		int32 SourceMeshLODIndex,
 		TArray<TArray<FVector3f>>& OutPositions_Original, 
 		TArray<TArray<FVector3f>>& OutPositions_Transferred,
-		const int32 MatchingSection)
+		const int32 InMatchingMaterialIndex)
 	{
-
 		// 1. Insert triangles into a 2D UV grid
 		auto BuildGrid = [InSourceMeshData](
 			int32 InSourceLODIndex,
-			int32 InSourceSectionId,
-			int32 InTargetSectionId,
+			int32 InSourceMaterialIndex,
 			int32 InChannelIndex,
 			FTriangleGrid2D& OutGrid)
 		{
-			// Notes:
-			// LODs are transfered using the LOD0 of the source mesh, as the LOD count can mismatch between source and target meshes.
-			// Assume that the section 0 contains the head section, which is where the hair/facial hair should be projected on
-			const GroomBinding_Mesh::IMeshLODData& MeshLODData = InSourceMeshData->GetMeshLODData(InSourceLODIndex);
-			const uint32 SourceTriangleCount = MeshLODData.GetSection(InSourceSectionId).GetNumTriangles();
-			const uint32 SourceSectionBaseIndex = MeshLODData.GetSection(InSourceSectionId).GetBaseIndex();
+			OutGrid.Reset();
 
-			const TArray<uint32>& SourceIndexBuffer = MeshLODData.GetIndexBuffer();
+			const GroomBinding_Mesh::IMeshLODData& MeshLODData = InSourceMeshData->GetMeshLODData(InSourceLODIndex);
 
-			OutGrid.Reset();
+			// Find all the sections which are mapped onto the target material index
+			const TArray<int32> SourceTargetSections = MeshLODData.GetSectionsMatchingMaterialIndex(InSourceMaterialIndex);
 
 			bool bIsGridPopulated = false;
-			for (uint32 SourceTriangleIt = 0; SourceTriangleIt < SourceTriangleCount; ++SourceTriangleIt)
+			for (int32 SectionIndex : SourceTargetSections)
 			{
-				FTriangleGrid2D::FTriangle T;
-				T.SectionIndex = InSourceSectionId;
-				T.SectionBaseIndex = SourceSectionBaseIndex;
-				T.TriangleIndex = SourceTriangleIt;
+				// Notes:
+				// LODs are transfered using the LOD0 of the source mesh, as the LOD count can mismatch between source and target meshes.
+				// Assume that the selected sections contains the head part, which is where the hair/facial hair should be projected on
+				const uint32 SourceTriangleCount = MeshLODData.GetSection(SectionIndex).GetNumTriangles();
+				const uint32 SourceSectionBaseIndex = MeshLODData.GetSection(SectionIndex).GetBaseIndex();
+
+				const TArray<uint32>& SourceIndexBuffer = MeshLODData.GetIndexBuffer();
+
+				for (uint32 SourceTriangleIt = 0; SourceTriangleIt < SourceTriangleCount; ++SourceTriangleIt)
+				{
+					FTriangleGrid2D::FTriangle T;
+					T.SectionIndex = SectionIndex;
+					T.SectionBaseIndex = SourceSectionBaseIndex;
+					T.TriangleIndex = SourceTriangleIt;
 
-				T.I0 = SourceIndexBuffer[T.SectionBaseIndex + SourceTriangleIt * 3 + 0];
-				T.I1 = SourceIndexBuffer[T.SectionBaseIndex + SourceTriangleIt * 3 + 1];
-				T.I2 = SourceIndexBuffer[T.SectionBaseIndex + SourceTriangleIt * 3 + 2];
+					T.I0 = SourceIndexBuffer[T.SectionBaseIndex + SourceTriangleIt * 3 + 0];
+					T.I1 = SourceIndexBuffer[T.SectionBaseIndex + SourceTriangleIt * 3 + 1];
+					T.I2 = SourceIndexBuffer[T.SectionBaseIndex + SourceTriangleIt * 3 + 2];
 
-				T.P0 = MeshLODData.GetVertexPosition(T.I0);
-				T.P1 = MeshLODData.GetVertexPosition(T.I1);
-				T.P2 = MeshLODData.GetVertexPosition(T.I2);
+					T.P0 = MeshLODData.GetVertexPosition(T.I0);
+					T.P1 = MeshLODData.GetVertexPosition(T.I1);
+					T.P2 = MeshLODData.GetVertexPosition(T.I2);
 
-				T.UV0 = MeshLODData.GetVertexUV(T.I0, InChannelIndex);
-				T.UV1 = MeshLODData.GetVertexUV(T.I1, InChannelIndex);
-				T.UV2 = MeshLODData.GetVertexUV(T.I2, InChannelIndex);
+					T.UV0 = MeshLODData.GetVertexUV(T.I0, InChannelIndex);
+					T.UV1 = MeshLODData.GetVertexUV(T.I1, InChannelIndex);
+					T.UV2 = MeshLODData.GetVertexUV(T.I2, InChannelIndex);
 
-				// Needed??
-				//T.N0 = MeshLODData.GetVertexNormal(T.I0);
-				//T.N1 = MeshLODData.GetVertexNormal(T.I1);
-				//T.N2 = MeshLODData.GetVertexNormal(T.I2);
+					// Needed??
+					//T.N0 = MeshLODData.GetVertexNormal(T.I0);
+					//T.N1 = MeshLODData.GetVertexNormal(T.I1);
+					//T.N2 = MeshLODData.GetVertexNormal(T.I2);
 
-				bIsGridPopulated = OutGrid.Insert(T) || bIsGridPopulated;
+					bIsGridPopulated = OutGrid.Insert(T) || bIsGridPopulated;
+				}
 			}
 
 			return bIsGridPopulated;
@@ -2233,16 +2285,16 @@ namespace GroomBinding_Transfer
 		// 1. Insert triangles into a 2D UV grid
 		const uint32 ChannelIndex = 0;
 		const GroomBinding_Mesh::IMeshLODData& SourceMeshLODData = InSourceMeshData->GetMeshLODData(SourceMeshLODIndex);
-		const bool bIsMatchingSectionValid = MatchingSection < SourceMeshLODData.GetNumSections();
-		const int32 SourceSectionId = bIsMatchingSectionValid ? MatchingSection : 0;
-		if (!bIsMatchingSectionValid && GHairStrandsBindingBuilderWarningEnable > 0)
+		const bool bIsMatchingMaterialValid = SourceMeshLODData.IsValidMaterialIndex(InMatchingMaterialIndex);
+		const int32 SourceMaterialIndex = bIsMatchingMaterialValid ? InMatchingMaterialIndex : 0;
+		const int32 TargetMaterialIndex = SourceMaterialIndex;
+		if (!bIsMatchingMaterialValid && GHairStrandsBindingBuilderWarningEnable > 0)
 		{
-			UE_LOG(LogHairStrands, Warning, TEXT("[Groom] Binding asset will not respect the requested 'Matching section' %d. The source skeletal mesh does not have such a section. Instead 'Matching Section' 0 will be used."), MatchingSection);
+			UE_LOG(LogHairStrands, Warning, TEXT("[Groom] Binding asset will not respect the requested 'Matching material' %d. The source skeletal mesh does not have such a material. Instead 'Matching material' 0 will be used."), InMatchingMaterialIndex);
 		}
-		const int32 TargetSectionId = SourceSectionId;
 		FTriangleGrid2D Grid(256);
 		{
-			const bool bIsGridPopulated = BuildGrid(SourceMeshLODIndex, SourceSectionId, TargetSectionId, ChannelIndex, Grid);
+			const bool bIsGridPopulated = BuildGrid(SourceMeshLODIndex, SourceMaterialIndex, ChannelIndex, Grid);
 			if (!bIsGridPopulated)
 			{
 				UE_LOG(LogHairStrands, Error, TEXT("[Groom] Binding asset could not be built. The source skeletal mesh is missing or has invalid UVs."));
@@ -2266,11 +2318,11 @@ namespace GroomBinding_Transfer
 
 			// Check that the target SectionId is valid for the current LOD. 
 			// If this is not the case, then fall back to section 0 and rebuild the source triangle grid to match the same section ID (1.)
-			int32 LocalSourceSectionId = SourceSectionId;
-			int32 LocalTargetSectionId = TargetSectionId;
+			int32 LocalSourceMaterialIndex = SourceMaterialIndex;
+			int32 LocalTargetMaterialIndex = TargetMaterialIndex;
 			const GroomBinding_Mesh::IMeshLODData& TargetMeshLODData = InTargetMeshData->GetMeshLODData(TargetLODIndex);
 
-			if (LocalTargetSectionId >= TargetMeshLODData.GetNumSections())
+			if (!TargetMeshLODData.IsValidMaterialIndex(LocalTargetMaterialIndex))
 			{
 				if ( TargetMeshLODData.GetNumSections() == 0 )
 				{
@@ -2280,12 +2332,12 @@ namespace GroomBinding_Transfer
 
 				if (GHairStrandsBindingBuilderWarningEnable > 0)
 				{
-					UE_LOG(LogHairStrands, Warning, TEXT("[Groom] Binding asset will not respect the requested 'Matching section' %d for LOD %d. The target skeletal mesh does not have such a section for this LOD. Instead section 0 will be used for this given LOD."), TargetSectionId, TargetLODIndex);
+					UE_LOG(LogHairStrands, Warning, TEXT("[Groom] Binding asset will not respect the requested 'Matching material' %d for LOD %d. The target skeletal mesh does not have such a material for this LOD. Instead material 0 will be used for this given LOD."), TargetMaterialIndex, TargetLODIndex);
 				}
 
-				LocalTargetSectionId = 0;
-				LocalSourceSectionId = 0;
-				const bool bIsGridPopulated = BuildGrid(SourceMeshLODIndex, LocalSourceSectionId, LocalTargetSectionId, ChannelIndex, Grid);
+				LocalTargetMaterialIndex = 0;
+				LocalSourceMaterialIndex = 0;
+				const bool bIsGridPopulated = BuildGrid(SourceMeshLODIndex, LocalSourceMaterialIndex, ChannelIndex, Grid);
 				if (!bIsGridPopulated)
 				{
 					UE_LOG(LogHairStrands, Error, TEXT("[Groom] Binding asset could not be built for LOD %d. The source skeletal mesh is missing or has invalid UVs."), TargetLODIndex);
@@ -2293,7 +2345,6 @@ namespace GroomBinding_Transfer
 				}
 			}
 
-			const uint32 TargetTriangleCount = TargetMeshLODData.GetSection(LocalTargetSectionId).GetNumTriangles();
 			const uint32 TargetVertexCount = TargetMeshLODData.GetNumVertices();
 
 			TSet<FVector2f> UVs;
@@ -2316,7 +2367,7 @@ namespace GroomBinding_Transfer
 #if BINDING_PARALLEL_BUILDING
 			ParallelFor(TargetVertexCount,
 				[
-					LocalTargetSectionId,
+					LocalTargetMaterialIndex,
 					ChannelIndex,
 					&TargetMeshLODData,
 					TargetLODIndex,
@@ -2329,7 +2380,8 @@ namespace GroomBinding_Transfer
 #endif
 			{
 				const int32 SectionIt = TargetMeshLODData.GetSectionFromVertexIndex(TargetVertexIt);
-				if (SectionIt != LocalTargetSectionId)
+				const int32 MaterialIndex = TargetMeshLODData.GetSection(SectionIt).GetMaterialIndex();
+				if (MaterialIndex != LocalTargetMaterialIndex)
 				{
 					
 					OutPositions_Original[TargetLODIndex][TargetVertexIt] = FVector3f(0,0,0);
@@ -2924,8 +2976,7 @@ bool BuildHairRootGroupData(const FGroomBindingBuilder::FInput& In, uint32 InGro
 					TargetMeshData.Get(),
 					In.TargetMeshMinLOD,
 					OutData.MeshPositions_Transferred,
-					OutData.SimRootDatas,
-					In.MatchingSection))
+					OutData.SimRootDatas))
 				{
 					UE_LOG(LogHairStrands, Error, TEXT("[Groom] Binding asset could not be built. Some guide roots are not close enough to the target mesh to be projected onto it."));
 					return false; 
@@ -2943,8 +2994,7 @@ bool BuildHairRootGroupData(const FGroomBindingBuilder::FInput& In, uint32 InGro
 					TargetMeshData.Get(),
 					In.TargetMeshMinLOD,
 					OutData.MeshPositions_Transferred,
-					OutData.RenRootDatas,
-					In.MatchingSection))
+					OutData.RenRootDatas))
 				{
 					UE_LOG(LogHairStrands, Error, TEXT("[Groom] Binding asset could not be built. Some strand roots are not close enough to the target mesh to be projected onto it."));
 					return false;
@@ -2974,8 +3024,7 @@ bool BuildHairRootGroupData(const FGroomBindingBuilder::FInput& In, uint32 InGro
 							TargetMeshData.Get(),
 							In.TargetMeshMinLOD,
 							OutData.MeshPositions_Transferred,
-							OutData.CardsRootDatas[CardsLODIt],
-							0/*In.MatchingSection*/))
+							OutData.CardsRootDatas[CardsLODIt]))
 						{
 							UE_LOG(LogHairStrands, Error, TEXT("[Groom] Binding asset could not be built. Some cards guide roots are not close enough to the target mesh to be projected onto it."));
 							return false; 
