[Hair] Protect root LOD bulk data access to validate bulk data validity.

Add condition test for resetting LoadedSize as the bulkdata might be invalid at this point.

diff --git a/Engine/Plugins/Runtime/HairStrands/Source/HairStrandsCore/Private/GroomDebug.cpp b/Engine/Plugins/Runtime/HairStrands/Source/HairStrandsCore/Private/GroomDebug.cpp
--- a/Engine/Plugins/Runtime/HairStrands/Source/HairStrandsCore/Private/GroomDebug.cpp
+++ b/Engine/Plugins/Runtime/HairStrands/Source/HairStrandsCore/Private/GroomDebug.cpp
@@ -852,10 +852,10 @@ static void AddHairDebugPrintInstancePass(
 				const int32 MeshLODIndex = Instance->HairGroupPublicData->MeshLODIndex;
 				if (MeshLODIndex>=0 && Instance->Strands.RestRootResource && Instance->Strands.RestRootResource->IsValid(MeshLODIndex))
 				{
-					D.Data1.X = Instance->Strands.RestRootResource->GetLOD(MeshLODIndex)->BulkData.Header.UniqueSectionIndices.Num();
-					D.Data1.Y = Instance->Strands.RestRootResource->GetLOD(MeshLODIndex)->BulkData.Header.UniqueTriangleCount;
-					D.Data1.Z = Instance->Strands.RestRootResource->GetLOD(MeshLODIndex)->BulkData.Header.RootCount;
-					D.Data1.W = Instance->Strands.RestRootResource->GetLOD(MeshLODIndex)->BulkData.Header.PointCount;
+					D.Data1.X = Instance->Strands.RestRootResource->GetLOD(MeshLODIndex)->GetBulkData().Header.UniqueSectionIndices.Num();
+					D.Data1.Y = Instance->Strands.RestRootResource->GetLOD(MeshLODIndex)->GetBulkData().Header.UniqueTriangleCount;
+					D.Data1.Z = Instance->Strands.RestRootResource->GetLOD(MeshLODIndex)->GetBulkData().Header.RootCount;
+					D.Data1.W = Instance->Strands.RestRootResource->GetLOD(MeshLODIndex)->GetBulkData().Header.PointCount;
 				}
 
 				{
diff --git a/Engine/Plugins/Runtime/HairStrands/Source/HairStrandsCore/Private/GroomManager.cpp b/Engine/Plugins/Runtime/HairStrands/Source/HairStrandsCore/Private/GroomManager.cpp
--- a/Engine/Plugins/Runtime/HairStrands/Source/HairStrandsCore/Private/GroomManager.cpp
+++ b/Engine/Plugins/Runtime/HairStrands/Source/HairStrandsCore/Private/GroomManager.cpp
@@ -550,7 +550,7 @@ static void GetOrAllocateCachedGeometry(
 
 	check(RootResource);
 	check(RootResource->IsDataValid(LODIndex));
-	const TArray<uint32>& UniqueSections = RootResource->GetLOD(LODIndex)->BulkData.Header.UniqueSectionIndices;
+	const TArray<uint32>& UniqueSections = RootResource->GetLOD(LODIndex)->GetBulkData().Header.UniqueSectionIndices;
 
 	// Create deformed position buffer (output)
 	FRDGBufferSRVRef DeformedPositionSRV = nullptr;
@@ -1111,6 +1111,7 @@ static void RunHairStrandsDeformation_Guide(
 					InstanceData.Instance->Guides.RestResource->GetCurveCount(),
 					InterpolationFactor,
 					1.f /*MaxHairRadius*/,
+					false /*bAllowAddedControlPoint - Guide can't have an extra control point (see FGroomBuilder::BuildHairDescriptionGroups())*/,
 					CacheResources0,
 					CacheResources1,
 					RegisterAsSRV(GraphBuilder, InstanceData.Instance->Guides.RestResource->PositionBuffer),
@@ -1596,6 +1597,7 @@ static void RunHairStrandsInterpolation_Strands(
 					InstanceData.ActiveCurveCount,
 					InterpolationFactor,
 					InstanceData.Instance->Strands.Modifier.HairWidth * 0.5f /*InMaxHairRadius*/,
+					true /*bAllowAddedControlPoint - Render strands can have an extra control point*/,
 					CacheResources0,
 					CacheResources1,
 					RegisterAsSRV(GraphBuilder, InstanceData.Instance->Strands.RestResource->PositionBuffer),
diff --git a/Engine/Plugins/Runtime/HairStrands/Source/HairStrandsCore/Private/GroomResources.cpp b/Engine/Plugins/Runtime/HairStrands/Source/HairStrandsCore/Private/GroomResources.cpp
--- a/Engine/Plugins/Runtime/HairStrands/Source/HairStrandsCore/Private/GroomResources.cpp
+++ b/Engine/Plugins/Runtime/HairStrands/Source/HairStrandsCore/Private/GroomResources.cpp
@@ -1600,9 +1600,9 @@ void FHairStrandsCullingResource::InternalRelease()
 
 /////////////////////////////////////////////////////////////////////////////////////////
 
-FHairStrandsLODRestRootResource::FHairStrandsLODRestRootResource(FHairStrandsRootBulkData& InBulkData, EHairStrandsResourcesType InCurveType, uint32 InMeshLODIndex, const FHairResourceName& InResourceName, const FName& InOwnerName) :
+FHairStrandsLODRestRootResource::FHairStrandsLODRestRootResource(TArray<FHairStrandsRootBulkData>& InBulkDataArray, int32 InBulkDataIndex, EHairStrandsResourcesType InCurveType, uint32 InMeshLODIndex, const FHairResourceName& InResourceName, const FName& InOwnerName) :
 	FHairCommonResource(EHairStrandsAllocationType::Deferred, InResourceName, InOwnerName),
-	BulkData(InBulkData), CurveType(InCurveType), MeshLODIndex(InMeshLODIndex)
+	BulkDataArray(InBulkDataArray), BulkDataIndex(InBulkDataIndex), CurveType(InCurveType), MeshLODIndex(InMeshLODIndex)
 {
 	// MaxAvailableCurveCount is not set to 0, as we need to track the loaded curve data *per LOD*. 
 	// This is done through callback to InternalIsLODDataLoaded()
@@ -1613,15 +1613,16 @@ FHairStrandsLODRestRootResource::FHairStrandsLODRestRootResource(FHairStrandsRoo
 
 void FHairStrandsLODRestRootResource::Reset()
 {
-	check(MeshLODIndex == BulkData.Header.LODIndex);
+	check(MeshLODIndex == GetBulkData().Header.LODIndex);
 	Status 		= EStatus::Invalid;
-	SampleCount = BulkData.Header.SampleCount;
+	SampleCount = GetBulkData().Header.SampleCount;
 }
 
 bool FHairStrandsLODRestRootResource::InternalGetOrRequestData(uint32 InRequestedCurveCount, uint32 InRequestedPointCount)
 {
 	if (StreamingRequest.IsNone())
 	{
+		FHairStrandsRootBulkData& BulkData = GetBulkData();
 #if !WITH_EDITORONLY_DATA
 		if (BulkData.IsCookedBulkDataLoadedForCPUAccess())
 		{
@@ -1647,7 +1648,7 @@ void FHairStrandsLODRestRootResource::InternalAllocate(FRDGBuilder& GraphBuilder
 	{
 		Reset();
 	}
-
+	FHairStrandsRootBulkData& BulkData = GetBulkData();
 	check(BulkData.Header.PointCount > 0);
 	{
 		// Clamp to the max supported
@@ -1700,7 +1701,13 @@ void FHairStrandsLODRestRootResource::InternalRelease()
 
 void FHairStrandsLODRestRootResource::InternalResetLoadedSize()
 {
-	BulkData.ResetLoadedSize();
+	// Conditionnaly reset loaded size at under certain circumpstances, the bulk data array might having reset 
+	// (due to pending build, while some resource are still in flight/being destroyed). This shouldn't be needed, 
+	// but can't figure out why we end in such a case at the moment.
+	if (IsBulkDataValid())
+	{
+		GetBulkData().ResetLoadedSize();
+	}
 }
 
 /////////////////////////////////////////////////////////////////////////////////////////
@@ -1708,9 +1715,11 @@ void FHairStrandsLODRestRootResource::InternalResetLoadedSize()
 FHairStrandsRestRootResource::FHairStrandsRestRootResource(TArray<FHairStrandsRootBulkData>& InBulkDatas, EHairStrandsResourcesType InCurveType, const FHairResourceName& InResourceName, const FName& InOwnerName)
 : CurveType(InCurveType)
 {
+	int32 Index = 0;
 	for (FHairStrandsRootBulkData& LODBulkData : InBulkDatas)
 	{
-		LODs.Add(new FHairStrandsLODRestRootResource(LODBulkData, InCurveType, LODBulkData.Header.LODIndex, InResourceName, InOwnerName));
+		LODs.Add(new FHairStrandsLODRestRootResource(InBulkDatas, Index, InCurveType, LODBulkData.Header.LODIndex, InResourceName, InOwnerName));
+		Index++;
 	}
 }
 
@@ -1784,8 +1793,8 @@ FHairStrandsLODDeformedRootResource::FHairStrandsLODDeformedRootResource(FHairSt
 	if (InRestResources->IsDataValid(InMeshLODIndex))
 	{
 		FHairStrandsLODRestRootResource* LODData = InRestResources->GetLOD(InMeshLODIndex);
-		RootCount = LODData->BulkData.Header.RootCount;
-		SampleCount = LODData->BulkData.Header.SampleCount;
+		RootCount = LODData->GetBulkData().Header.RootCount;
+		SampleCount = LODData->GetBulkData().Header.SampleCount;
 	}
 
 	Status = EStatus::Invalid;
diff --git a/Engine/Plugins/Runtime/HairStrands/Source/HairStrandsCore/Private/HairStrandsMeshProjection.cpp b/Engine/Plugins/Runtime/HairStrands/Source/HairStrandsCore/Private/HairStrandsMeshProjection.cpp
--- a/Engine/Plugins/Runtime/HairStrands/Source/HairStrandsCore/Private/HairStrandsMeshProjection.cpp
+++ b/Engine/Plugins/Runtime/HairStrands/Source/HairStrandsCore/Private/HairStrandsMeshProjection.cpp
@@ -352,9 +352,9 @@ void AddHairStrandUpdateMeshTrianglesPass(
 	}
 
 	FHairStrandsLODRestRootResource& RestLODData = *RestResources->GetLOD(MeshLODIndex);
-	check(RestLODData.BulkData.Header.LODIndex == MeshLODIndex);
-	const uint32 UniqueTriangleCount = RestLODData.BulkData.Header.UniqueTriangleCount;
-	const uint32 TotalMeshSectionCount = RestLODData.BulkData.Header.MeshSectionCount;
+	check(RestLODData.GetBulkData().Header.LODIndex == MeshLODIndex);
+	const uint32 UniqueTriangleCount = RestLODData.GetBulkData().Header.UniqueTriangleCount;
+	const uint32 TotalMeshSectionCount = RestLODData.GetBulkData().Header.MeshSectionCount;
 	if (UniqueTriangleCount == 0)
 	{
 		return;
@@ -565,7 +565,7 @@ void AddHairStrandInitMeshSamplesPass(
 	
 	FHairStrandsLODRestRootResource& RestLODData = *RestResources->GetLOD(LODIndex);
 
-	const uint32 TotalSectionCount = RestLODData.BulkData.Header.MeshSectionCount;
+	const uint32 TotalSectionCount = RestLODData.GetBulkData().Header.MeshSectionCount;
 	const uint32 EffectiveSectionCount = MeshData.Sections.Num();
 	if (EffectiveSectionCount == 0 || RestLODData.SampleCount == 0)
 	{
@@ -797,7 +797,7 @@ static void AddFollicleMaskPass(
 	Parameters->RootUVsBuffer = nullptr;
 	Parameters->OutputResolution = OutputResolution;
 	Parameters->MaxRootCount = RootCount;
-	Parameters->MaxUniqueTriangleIndex = LODData.BulkData.Header.UniqueTriangleCount;
+	Parameters->MaxUniqueTriangleIndex = LODData.GetBulkData().Header.UniqueTriangleCount;
 	Parameters->Channel = FMath::Min(Channel, 3u);
 	Parameters->KernelSizeInPixels = FMath::Clamp(KernelSizeInPixels, 2u, 200u);
 	Parameters->RenderTargets[0] = FRenderTargetBinding(OutTexture, bNeedClear ? ERenderTargetLoadAction::EClear : ERenderTargetLoadAction::ELoad, 0);
diff --git a/Engine/Plugins/Runtime/HairStrands/Source/HairStrandsCore/Public/GroomResources.h b/Engine/Plugins/Runtime/HairStrands/Source/HairStrandsCore/Public/GroomResources.h
--- a/Engine/Plugins/Runtime/HairStrands/Source/HairStrandsCore/Public/GroomResources.h
+++ b/Engine/Plugins/Runtime/HairStrands/Source/HairStrandsCore/Public/GroomResources.h
@@ -393,14 +393,14 @@ struct FHairStrandsCullingResource : public FHairCommonResource
 struct FHairStrandsLODRestRootResource : public FHairCommonResource
 {
 	/** Build the hair strands resource */
-	FHairStrandsLODRestRootResource(FHairStrandsRootBulkData& BulkData, EHairStrandsResourcesType CurveType, uint32 InMeshLODIndex, const FHairResourceName& ResourceName, const FName& OwnerName);
+	FHairStrandsLODRestRootResource(TArray<FHairStrandsRootBulkData>& InBulkDataArray, int32 InBulkDataIndex, EHairStrandsResourcesType CurveType, uint32 InMeshLODIndex, const FHairResourceName& ResourceName, const FName& OwnerName);
 
 	/* Init/Release buffers */
 	virtual void InternalAllocate(FRDGBuilder& GraphBuilder, uint32 InCurveCount, uint32 InPointCount) override;
 	virtual void InternalRelease() override;
 	virtual bool InternalGetOrRequestData(uint32 InRequestedCurveCount, uint32 InRequestedPointCount) override;
 	virtual void InternalResetLoadedSize() override;
-	virtual FHairStrandsBulkCommon* InternalGetBulkData() override { return &BulkData; }
+	virtual FHairStrandsBulkCommon* InternalGetBulkData() override { return &GetBulkData(); }
 
 	/* Get the resource name */
 	virtual FString GetFriendlyName() const override { return TEXT("FHairStrandsLODRestRootResource"); }
@@ -409,7 +409,22 @@ struct FHairStrandsLODRestRootResource : public FHairCommonResource
 	void Reset();
 
 	// Accessors
-	uint32 GetRootCount()const { return BulkData.GetRootCount(); }
+	uint32 GetRootCount()const { return GetBulkData().GetRootCount(); }
+
+	/* Accessor to deference bulk data*/
+	FORCEINLINE bool IsBulkDataValid() const { return BulkDataArray.IsValidIndex(BulkDataIndex); }
+
+	FORCEINLINE const FHairStrandsRootBulkData& GetBulkData() const
+	{
+		check(BulkDataArray.IsValidIndex(BulkDataIndex));
+		return BulkDataArray[BulkDataIndex];
+	}
+
+	FORCEINLINE FHairStrandsRootBulkData& GetBulkData()
+	{
+		check(BulkDataArray.IsValidIndex(BulkDataIndex));
+		return BulkDataArray[BulkDataIndex];
+	}
 
 	/* Return the memory size for GPU resources */
 	uint32 GetResourcesSize() const 
@@ -427,7 +442,7 @@ struct FHairStrandsLODRestRootResource : public FHairCommonResource
 
 	enum class EStatus { Invalid, Initialized, Completed };
 
-	const bool IsDataValid() const { return BulkData.Header.RootCount > 0 && BulkData.Header.LODIndex >= 0; }
+	const bool IsDataValid() const { return GetBulkData().Header.RootCount > 0 && GetBulkData().Header.LODIndex >= 0; }
 	const bool IsValid() const { return Status == EStatus::Completed; }
 	EStatus Status = EStatus::Invalid;
 	uint32 AvailableCurveCount = 0;
@@ -450,7 +465,8 @@ struct FHairStrandsLODRestRootResource : public FHairCommonResource
 	FRDGExternalBuffer RestSamplePositionsBuffer;
 
 	/* Store CPU data for root info & root binding */
-	FHairStrandsRootBulkData& BulkData;
+	TArray<FHairStrandsRootBulkData>& BulkDataArray;
+	int32 BulkDataIndex = -1;
 
 	/* Type of curves */
 	const EHairStrandsResourcesType CurveType;
