Fix Groom Sim cache.

diff --git a/Engine/Plugins/Runtime/HairStrands/Source/HairStrandsCore/Private/HairStrandsInterpolation.cpp b/Engine/Plugins/Runtime/HairStrands/Source/HairStrandsCore/Private/HairStrandsInterpolation.cpp
--- a/Engine/Plugins/Runtime/HairStrands/Source/HairStrandsCore/Private/HairStrandsInterpolation.cpp
+++ b/Engine/Plugins/Runtime/HairStrands/Source/HairStrandsCore/Private/HairStrandsInterpolation.cpp
@@ -279,6 +279,7 @@ void AddGroomCacheUpdatePass(
 	uint32 CurveCount,
 	float InterpolationFactor,
 	float InMaxHairRadius,
+	bool bAllowAddedControlPoint,
 	FGroomCacheResources CacheResources0,
 	FGroomCacheResources CacheResources1,
 	FRDGBufferSRVRef InRestPositionBuffer,
@@ -311,7 +312,7 @@ void AddGroomCacheUpdatePass(
 	Parameters->InterpolationFactor = InterpolationFactor;
 	Parameters->MaxHairRadius = InMaxHairRadius;
 	Parameters->bHasRadiusData = CacheResources0.bHasRadiusData ? 1u : 0u;
-	Parameters->bHasAddedControlPoint = GetHairStrandsUsesTriangleStrips() ? 1u : 0u;
+	Parameters->bHasAddedControlPoint = (bAllowAddedControlPoint && GetHairStrandsUsesTriangleStrips()) ? 1u : 0u;
 
 	const FIntVector DispatchCount = FIntVector(FMath::DivideAndRoundUp(PointCount, FGroomCacheUpdatePassCS::GetGroupSize()), 1, 1);
 	TShaderMapRef<FGroomCacheUpdatePassCS> ComputeShader(ShaderMap);
diff --git a/Engine/Plugins/Runtime/HairStrands/Source/HairStrandsCore/Public/HairStrandsInterpolation.h b/Engine/Plugins/Runtime/HairStrands/Source/HairStrandsCore/Public/HairStrandsInterpolation.h
--- a/Engine/Plugins/Runtime/HairStrands/Source/HairStrandsCore/Public/HairStrandsInterpolation.h
+++ b/Engine/Plugins/Runtime/HairStrands/Source/HairStrandsCore/Public/HairStrandsInterpolation.h
@@ -77,6 +77,7 @@ void AddGroomCacheUpdatePass(
 	uint32 CurveCount,
 	float InterpolationFactor,
 	float InMaxHairRadius,
+	bool bAllowAddedControlPoint,
 	FGroomCacheResources CacheResources0,
 	FGroomCacheResources CacheResources1,
 	FRDGBufferSRVRef InRestPositionBuffer,
