Fix incorrect groom binding validation when initializing groom compnent resources when a groom asset has strands data but none of the LOD are using strands.

diff --git a/Engine/Plugins/Runtime/HairStrands/Source/HairStrandsCore/Private/GroomComponent.cpp b/Engine/Plugins/Runtime/HairStrands/Source/HairStrandsCore/Private/GroomComponent.cpp
--- a/Engine/Plugins/Runtime/HairStrands/Source/HairStrandsCore/Private/GroomComponent.cpp
+++ b/Engine/Plugins/Runtime/HairStrands/Source/HairStrandsCore/Private/GroomComponent.cpp
@@ -2809,19 +2809,23 @@ void UGroomComponent::InitResources(bool bIsBindingReloading)
 					const UGroomBindingAsset::FHairGroupResource& BindingResource = LocalBindingAsset->GetHairGroupResources()[GroupIt];
 					const FHairGroupResources& GroomResources = GroomAsset->GetHairGroupsResources()[GroupIt];
 
+					const bool bHasGuides = GroomAsset->NeedsInterpolationData(GroupIt);
+					const bool bHasStrands = GroomAsset->HasGeometryType(GroupIt, EGroomGeometryType::Strands);
+					const bool bHasCards = GroomAsset->HasGeometryType(GroupIt, EGroomGeometryType::Cards);
+
 					bool bValid = true;
 					// Check guides
-					if (GroomResources.Guides.IsValid())
+					if (GroomResources.Guides.IsValid() && bHasGuides)
 					{
 						bValid = bValid && BindingResource.SimRootResources && BindingResource.SimRootResources->GetLODCount() > 0;
 					}
 					// Check strands
-					if (bValid && GroomResources.Strands.IsValid())
+					if (bValid && GroomResources.Strands.IsValid() && bHasStrands)
 					{
 						bValid = bValid && BindingResource.RenRootResources && BindingResource.RenRootResources->GetLODCount() > 0;
 					}
 					// Check cards
-					if (bValid && GroomResources.Cards.LODs.Num() > 0)
+					if (bValid && GroomResources.Cards.LODs.Num() > 0 && bHasCards)
 					{
 						bValid = bValid && GroomResources.Cards.LODs.Num() == BindingResource.CardsRootResources.Num();
 						if (bValid)
