Export more attribute types to groom collection

diff --git a/Engine/Plugins/Runtime/HairStrands/Source/HairStrandsDataflow/Private/GroomAssetTerminalNode.cpp b/Engine/Plugins/Runtime/HairStrands/Source/HairStrandsDataflow/Private/GroomAssetTerminalNode.cpp
--- a/Engine/Plugins/Runtime/HairStrands/Source/HairStrandsDataflow/Private/GroomAssetTerminalNode.cpp
+++ b/Engine/Plugins/Runtime/HairStrands/Source/HairStrandsDataflow/Private/GroomAssetTerminalNode.cpp
@@ -1,4 +1,4 @@
-﻿// Copyright Epic Games, Inc. All Rights Reserved.
+// Copyright Epic Games, Inc. All Rights Reserved.
 
 #include "GroomAssetTerminalNode.h"
 
@@ -71,6 +71,26 @@ namespace UE::Groom::Private
 		}
 	}
 
+	template<typename FacadeType, typename AttributeType>
+	FORCEINLINE void BuildVerticesAttribute(const FManagedArrayCollection& InCollection, FManagedArrayCollection* OutCollection, const int32 NumPoints, const FName& AttributeName)
+	{
+		const TManagedArray<AttributeType>& VerticesAttribute = InCollection.GetAttribute<AttributeType>(AttributeName, FacadeType::VerticesGroup);
+		if(OutCollection->NumElements(FacadeType::PointsGroup) != NumPoints)
+		{
+			if(OutCollection->NumElements(FacadeType::PointsGroup) > 0)
+			{
+				OutCollection->EmptyGroup(FacadeType::PointsGroup);
+			}
+			OutCollection->AddElements(NumPoints, FacadeType::PointsGroup);
+		}
+		TManagedArray<AttributeType>& PointsAttribute = OutCollection->AddAttribute<AttributeType>(AttributeName, FacadeType::PointsGroup);
+		
+		for(int32 PointIndex = 0; PointIndex < NumPoints; ++PointIndex)
+		{
+			PointsAttribute[PointIndex] = VerticesAttribute[2*PointIndex];
+		}
+	}
+
 	template<typename FacadeType>
 	FORCEINLINE void TransferVerticesAttributes(const FManagedArrayCollection& InCollection, FManagedArrayCollection* OutCollection, const int32 NumPoints)
 	{
@@ -80,22 +100,69 @@ namespace UE::Groom::Private
 		{
 			if(InCollection.GetAttributeType(AttributeName, FacadeType::VerticesGroup) == EManagedArrayType::FFloatType)
 			{
-				const TManagedArray<float>& VerticesAttribute = InCollection.GetAttribute<float>(AttributeName, FacadeType::VerticesGroup);
-				if(OutCollection->NumElements(FacadeType::PointsGroup) != NumPoints)
-				{
-					if(OutCollection->NumElements(FacadeType::PointsGroup) > 0)
-					{
-						OutCollection->EmptyGroup(FacadeType::PointsGroup);
-					}
-					OutCollection->AddElements(NumPoints, FacadeType::PointsGroup);
-				}
-				TManagedArray<float>& PointsAttribute = OutCollection->AddAttribute<float>(AttributeName, FacadeType::PointsGroup);
-				
-				for(int32 PointIndex = 0; PointIndex < NumPoints; ++PointIndex)
-				{
-					PointsAttribute[PointIndex] = 0.5f * (VerticesAttribute[2*PointIndex] + VerticesAttribute[2*PointIndex+1]);
-				}
+				BuildVerticesAttribute<FacadeType, float>(InCollection, OutCollection, NumPoints, AttributeName);
+			}
+			else if(InCollection.GetAttributeType(AttributeName, FacadeType::VerticesGroup) == EManagedArrayType::FVector4fType)
+			{
+				BuildVerticesAttribute<FacadeType, FVector4f>(InCollection, OutCollection, NumPoints, AttributeName);
+			}
+			else if(InCollection.GetAttributeType(AttributeName, FacadeType::VerticesGroup) == EManagedArrayType::FVectorType)
+			{
+				BuildVerticesAttribute<FacadeType, FVector3f>(InCollection, OutCollection, NumPoints, AttributeName);
+			}
+			else if(InCollection.GetAttributeType(AttributeName, FacadeType::VerticesGroup) == EManagedArrayType::FVector2DType)
+			{
+				BuildVerticesAttribute<FacadeType, FVector2f>(InCollection, OutCollection, NumPoints, AttributeName);
+			}
+			if(InCollection.GetAttributeType(AttributeName, FacadeType::VerticesGroup) == EManagedArrayType::FInt32Type)
+			{
+				BuildVerticesAttribute<FacadeType, int32>(InCollection, OutCollection, NumPoints, AttributeName);
+			}
+			else if(InCollection.GetAttributeType(AttributeName, FacadeType::VerticesGroup) == EManagedArrayType::FIntVector4Type)
+			{
+				BuildVerticesAttribute<FacadeType, FIntVector4>(InCollection, OutCollection, NumPoints, AttributeName);
+			}
+			else if(InCollection.GetAttributeType(AttributeName, FacadeType::VerticesGroup) == EManagedArrayType::FIntVectorType)
+			{
+				BuildVerticesAttribute<FacadeType, FIntVector3>(InCollection, OutCollection, NumPoints, AttributeName);
 			}
+			else if(InCollection.GetAttributeType(AttributeName, FacadeType::VerticesGroup) == EManagedArrayType::FIntVector2Type)
+			{
+				BuildVerticesAttribute<FacadeType, FIntVector2>(InCollection, OutCollection, NumPoints, AttributeName);
+			}
+			else if(InCollection.GetAttributeType(AttributeName, FacadeType::VerticesGroup) == EManagedArrayType::FBoolType)
+			{
+				BuildVerticesAttribute<FacadeType, bool>(InCollection, OutCollection, NumPoints, AttributeName);
+			}
+			else if(InCollection.GetAttributeType(AttributeName, FacadeType::VerticesGroup) == EManagedArrayType::FLinearColorType)
+			{
+				BuildVerticesAttribute<FacadeType, FLinearColor>(InCollection, OutCollection, NumPoints, AttributeName);
+			}
+			else if(InCollection.GetAttributeType(AttributeName, FacadeType::VerticesGroup) == EManagedArrayType::FQuatType)
+			{
+				BuildVerticesAttribute<FacadeType, FQuat4f>(InCollection, OutCollection, NumPoints, AttributeName);
+			}
+			else if(InCollection.GetAttributeType(AttributeName, FacadeType::VerticesGroup) == EManagedArrayType::FTransform3fType)
+			{
+				BuildVerticesAttribute<FacadeType, FTransform3f>(InCollection, OutCollection, NumPoints, AttributeName); 
+
+			}
+		}
+	}
+
+	template<typename FacadeType>
+	FORCEINLINE void RegisterSkeletalMeshes(const FManagedArrayCollection& InCollection, const FacadeType& GroomFacade, UGroomAsset* GroomAsset)
+	{
+		const TManagedArray<TObjectPtr<UObject>>& ObjectSkeletalMeshes =
+						InCollection.GetAttribute<TObjectPtr<UObject>>(UE::Groom::FGroomGuidesFacade::ObjectSkeletalMeshesAttribute, FacadeType::ObjectsGroup);
+
+		const TManagedArray<int32>& ObjectMeshLODs =
+			InCollection.GetAttribute<int32>(UE::Groom::FGroomGuidesFacade::ObjectMeshLODsAttribute, FacadeType::ObjectsGroup);
+				
+		for(int32 GroupIndex = 0; GroupIndex < GroomFacade.GetNumObjects(); ++GroupIndex)
+		{
+			GroomAsset->GetDataflowSettings().SetSkeletalMesh(GroupIndex,
+				Cast<USkeletalMesh>(ObjectSkeletalMeshes[GroupIndex]), ObjectMeshLODs[GroupIndex]);
 		}
 	}
 }
@@ -138,7 +205,6 @@ void FGroomAssetTerminalDataflowNode::SetAssetValue(TObjectPtr<UObject> Asset, U
 			{
 				// Build the editable guides
 				UE::Groom::Private::BuildEditableGuides(GuidesFacade,EditGroom);
-				 
 
 				// Ensure compilation dependent assets is done
 				FAssetCompilingManager::Get().FinishCompilationForObjects({ GroomAsset });
@@ -146,7 +212,7 @@ void FGroomAssetTerminalDataflowNode::SetAssetValue(TObjectPtr<UObject> Asset, U
 				// Convert to groom asset
 				ConvertToGroomAsset(const_cast<UGroomAsset*>(GroomAsset), &EditGroom, EEditableGroomOperations::ControlPoints_Modified);
 			}
-		}
+		}  
 
 		if(OutCollection)
 		{
@@ -155,17 +221,11 @@ void FGroomAssetTerminalDataflowNode::SetAssetValue(TObjectPtr<UObject> Asset, U
 
 			if(InCollection.HasAttribute(UE::Groom::FGroomGuidesFacade::ObjectSkeletalMeshesAttribute, UE::Groom::FGroomGuidesFacade::ObjectsGroup))
 			{
-				const TManagedArray<TObjectPtr<UObject>>& ObjectSkeletalMeshes =
-					InCollection.GetAttribute<TObjectPtr<UObject>>(UE::Groom::FGroomGuidesFacade::ObjectSkeletalMeshesAttribute, UE::Groom::FGroomGuidesFacade::ObjectsGroup);
-
-				const TManagedArray<int32>& ObjectMeshLODs =
-					InCollection.GetAttribute<int32>(UE::Groom::FGroomGuidesFacade::ObjectMeshLODsAttribute, UE::Groom::FGroomGuidesFacade::ObjectsGroup);
-				
-				for(int32 GroupIndex = 0; GroupIndex < GuidesFacade.GetNumObjects(); ++GroupIndex)
-				{
-					GroomAsset->GetDataflowSettings().SetSkeletalMesh(GroupIndex,
-						Cast<USkeletalMesh>(ObjectSkeletalMeshes[GroupIndex]), ObjectMeshLODs[GroupIndex]);
-				}
+				UE::Groom::Private::RegisterSkeletalMeshes(InCollection, GuidesFacade, GroomAsset);
+			}
+			else if(InCollection.HasAttribute(UE::Groom::FGroomGuidesFacade::ObjectSkeletalMeshesAttribute, UE::Groom::FGroomStrandsFacade::ObjectsGroup))
+			{
+				UE::Groom::Private::RegisterSkeletalMeshes(InCollection, StrandsFacade, GroomAsset);
 			}
 		}
 	}
