Extract/report groom attributes in between points and vertices + add support for strands skin weights

diff --git a/Engine/Plugins/Runtime/HairStrands/Source/HairStrandsDataflow/Private/BuildGroomSkinningNodes.cpp b/Engine/Plugins/Runtime/HairStrands/Source/HairStrandsDataflow/Private/BuildGroomSkinningNodes.cpp
new file mode 100644
--- /dev/null
+++ b/Engine/Plugins/Runtime/HairStrands/Source/HairStrandsDataflow/Private/BuildGroomSkinningNodes.cpp
@@ -0,0 +1,399 @@
+// Copyright Epic Games, Inc. All Rights Reserved.
+
+#include "BuildGroomSkinningNodes.h"
+#include "GroomCollectionFacades.h"
+#include "MeshDescriptionToDynamicMesh.h"
+#include "SkeletalMeshAttributes.h"
+#include "SkeletalMeshLODRenderDataToDynamicMesh.h"
+#include "DynamicMesh/DynamicMesh3.h"
+#include "Dataflow/DataflowObjectInterface.h"
+#include "DynamicMesh/MeshTransforms.h"
+#include "Operations/TransferBoneWeights.h"
+#include "Rendering/SkeletalMeshRenderData.h"
+
+#include UE_INLINE_GENERATED_CPP_BY_NAME(BuildGroomSkinningNodes)
+
+namespace UE::Groom::Private
+{
+	FORCEINLINE bool SkeletalMeshToDynamicMesh(const USkeletalMesh* SkeletalMesh, const int32 LodIndex, UE::Geometry::FDynamicMesh3& ToDynamicMesh)
+	{
+		if (SkeletalMesh->HasMeshDescription(LodIndex))
+		{
+			const FMeshDescription* SourceMesh = SkeletalMesh->GetMeshDescription(LodIndex);
+			if (!SourceMesh)
+			{
+				return false;
+			}
+
+			FMeshDescriptionToDynamicMesh Converter;
+			Converter.Convert(SourceMesh, ToDynamicMesh);
+		}
+		else
+		{ 
+			const FSkeletalMeshRenderData* RenderData = SkeletalMesh->GetResourceForRendering();
+			if (!RenderData)
+			{
+				return false;
+			}
+
+			if (!RenderData->LODRenderData.IsValidIndex(LodIndex))
+			{
+				return false;
+			}
+
+			const FSkeletalMeshLODRenderData* SkeletalMeshLODRenderData = &(RenderData->LODRenderData[LodIndex]);
+
+			UE::Geometry::FSkeletalMeshLODRenderDataToDynamicMesh::ConversionOptions ConversionOptions;
+			ConversionOptions.bWantUVs = false;
+			ConversionOptions.bWantVertexColors = false;
+			ConversionOptions.bWantMaterialIDs = false;
+			ConversionOptions.bWantSkinWeights = true;
+
+			UE::Geometry::FSkeletalMeshLODRenderDataToDynamicMesh::Convert(SkeletalMeshLODRenderData, SkeletalMesh->GetRefSkeleton(), ConversionOptions, ToDynamicMesh);
+		}
+
+		return true;
+	}
+
+	template<typename FacadeType>
+	static void BuildSkinningData(FManagedArrayCollection& GroomCollection, const int32 GroupIndex, const TObjectPtr<USkeletalMesh> SkeletalMesh, const int32 LODIndex,
+		const FTransform& RelativeTransform)
+	{
+		FacadeType GroomFacade(GroomCollection);
+
+		if(GroomFacade.IsValid())
+		{
+			TManagedArray<TObjectPtr<UObject>>& ObjectSkeletalMeshes =
+				GroomCollection.AddAttribute<TObjectPtr<UObject>>(UE::Groom::FGroomGuidesFacade::ObjectSkeletalMeshesAttribute, FacadeType::ObjectsGroup);
+
+			TManagedArray<int32>& ObjectMeshLODs =
+				GroomCollection.AddAttribute<int32>(UE::Groom::FGroomGuidesFacade::ObjectMeshLODsAttribute, FacadeType::ObjectsGroup);
+
+			TManagedArray<FIntVector4>& VertexBoneIndices =
+				GroomCollection.AddAttribute<FIntVector4>(UE::Groom::FGroomGuidesFacade::PointBoneIndicesAttribute, FacadeType::VerticesGroup);
+
+			TManagedArray<FVector4f>& VertexBoneWeights =
+				GroomCollection.AddAttribute<FVector4f>(UE::Groom::FGroomGuidesFacade::PointBoneWeightsAttribute, FacadeType::VerticesGroup);
+			
+			if(GroupIndex == INDEX_NONE)
+			{
+				for(int32 LocalIndex = 0; LocalIndex < ObjectSkeletalMeshes.Num(); ++LocalIndex)
+				{
+					ObjectSkeletalMeshes[LocalIndex] = SkeletalMesh;
+					ObjectMeshLODs[LocalIndex] = LODIndex;
+				}
+			}
+			else if(ObjectSkeletalMeshes.IsValidIndex(GroupIndex) && ObjectMeshLODs.IsValidIndex(GroupIndex))
+			{
+				ObjectSkeletalMeshes[GroupIndex] = SkeletalMesh;
+				ObjectMeshLODs[GroupIndex] = LODIndex;
+			}
+			if(SkeletalMesh && SkeletalMesh->IsValidLODIndex(LODIndex))
+			{
+				TMap<FName, FBoneIndexType> TargetBoneToIndex;
+				TargetBoneToIndex.Reserve(SkeletalMesh->GetRefSkeleton().GetRawBoneNum());
+				for (int32 BoneIdx = 0; BoneIdx < SkeletalMesh->GetRefSkeleton().GetRawBoneNum(); ++BoneIdx)
+				{
+					TargetBoneToIndex.Add(SkeletalMesh->GetRefSkeleton().GetRawRefBoneInfo()[BoneIdx].Name, BoneIdx);
+				}
+				UE::Geometry::FDynamicMesh3 DynamicMesh;
+				if (UE::Groom::Private::SkeletalMeshToDynamicMesh(SkeletalMesh, LODIndex, DynamicMesh))
+				{
+					MeshTransforms::ApplyTransform(DynamicMesh, RelativeTransform, true);
+
+					UE::Geometry::FTransferBoneWeights TransferBoneWeights(&DynamicMesh, FSkeletalMeshAttributes::DefaultSkinWeightProfileName);
+					TransferBoneWeights.bUseParallel = true;
+					TransferBoneWeights.MaxNumInfluences = 4;
+					TransferBoneWeights.TransferMethod = UE::Geometry::FTransferBoneWeights::ETransferBoneWeightsMethod::ClosestPointOnSurface;
+
+					if (TransferBoneWeights.Validate() == UE::Geometry::EOperationValidationResult::Ok)
+					{
+						const int32 LocalIndex = GroupIndex;
+						ParallelFor(GroomFacade.GetNumPoints(), [&TransferBoneWeights, &TargetBoneToIndex, &GroomFacade, LocalIndex, &VertexBoneIndices, &VertexBoneWeights](
+							int32 PointIndex)
+						{
+							const int32 CurveIndex = GroomFacade.GetPointCurveIndices()[PointIndex];
+							const int32 ObjectIndex = GroomFacade.GetCurveObjectIndices()[CurveIndex];
+
+							if(LocalIndex == INDEX_NONE || ObjectIndex == LocalIndex)
+							{
+								TArray<int32> BoneIndices;
+								TArray<float> BoneWeights;
+								TransferBoneWeights.TransferWeightsToPoint(BoneIndices, BoneWeights,
+									GroomFacade.GetPointRestPositions()[PointIndex],&TargetBoneToIndex);
+
+								FIntVector4 PointBoneIndices(INDEX_NONE);
+								FVector4f PointBoneWeights(0.0f);
+
+								for(int32 BoneIdx = 0; BoneIdx < BoneIndices.Num(); ++BoneIdx)
+								{
+									PointBoneIndices[BoneIdx] = BoneIndices[BoneIdx];
+									PointBoneWeights[BoneIdx] = BoneWeights[BoneIdx];
+								}
+								VertexBoneIndices[2*PointIndex] = PointBoneIndices;
+								VertexBoneIndices[2*PointIndex+1] = PointBoneIndices;
+								
+								VertexBoneWeights[2*PointIndex] = PointBoneWeights;
+								VertexBoneWeights[2*PointIndex+1] = PointBoneWeights;
+							}
+
+						}, TransferBoneWeights.bUseParallel  ? EParallelForFlags::None : EParallelForFlags::ForceSingleThread);
+					}
+				}
+			}
+		}
+	}
+	
+	template<typename FacadeType>
+	FORCEINLINE void ExtractSkinningData(FManagedArrayCollection& GroomCollection, const FCollectionAttributeKey& BoneIndicesKey, const FCollectionAttributeKey& BoneWeightsKey)
+	{
+		if (!BoneIndicesKey.Attribute.IsEmpty() && !BoneWeightsKey.Attribute.IsEmpty())
+		{
+			FacadeType GroomFacade(GroomCollection);
+
+			if(GroomFacade.IsValid())
+			{
+				const TManagedArray<FIntVector4>* BoneIndices = GroomCollection.FindAttributeTyped<FIntVector4>(FName(UE::Groom::FGroomGuidesFacade::PointBoneIndicesAttribute), FacadeType::VerticesGroup);
+				const TManagedArray<FVector4f>* BoneWeights = GroomCollection.FindAttributeTyped<FVector4f>(FName(UE::Groom::FGroomGuidesFacade::PointBoneWeightsAttribute), FacadeType::VerticesGroup);
+				
+				if(BoneIndices && BoneWeights)
+				{
+					TManagedArray<TArray<int32>>& IndicesArray = GroomCollection.AddAttribute<TArray<int32>>(FName(BoneIndicesKey.Attribute), FName(BoneIndicesKey.Group));
+					TManagedArray<TArray<float>>& WeightsArray = GroomCollection.AddAttribute<TArray<float>>(FName(BoneWeightsKey.Attribute), FName(BoneWeightsKey.Group));
+
+					const int32 NumVertices = GroomFacade.GetNumVertices();
+					for(int32 VertexIndex = 0; VertexIndex < NumVertices; ++VertexIndex)
+					{
+						const FIntVector4& PointIndices = (*BoneIndices)[VertexIndex];
+						const FVector4f& PointWeights = (*BoneWeights)[VertexIndex];
+
+						IndicesArray[VertexIndex].Reset(4);
+						WeightsArray[VertexIndex].Reset(4);
+						for(int32 WeightIndex = 0; WeightIndex < 4; ++WeightIndex)
+						{
+							if(PointIndices[WeightIndex] != INDEX_NONE)
+							{
+								IndicesArray[VertexIndex].Add(PointIndices[WeightIndex]);
+								WeightsArray[VertexIndex].Add(PointWeights[WeightIndex]);
+							}
+						}
+					}
+				}
+			}
+		}
+	}
+
+	template<typename FacadeType>
+	FORCEINLINE void ReportSkinningData(FManagedArrayCollection& GroomCollection, const FCollectionAttributeKey& BoneIndicesKey, const FCollectionAttributeKey& BoneWeightsKey)
+	{
+		if (!BoneIndicesKey.Attribute.IsEmpty() && !BoneWeightsKey.Attribute.IsEmpty())
+		{
+			FacadeType GroomFacade(GroomCollection);
+
+			if(GroomFacade.IsValid())
+			{
+				const TManagedArray<TArray<int32>>* IndicesArray = GroomCollection.FindAttributeTyped<TArray<int32>>(FName(BoneIndicesKey.Attribute), FName(BoneIndicesKey.Group));
+				const TManagedArray<TArray<float>>* WeightsArray = GroomCollection.FindAttributeTyped<TArray<float>>(FName(BoneWeightsKey.Attribute), FName(BoneWeightsKey.Group));
+
+				if(IndicesArray && WeightsArray)
+				{
+					TManagedArray<FIntVector4>& BoneIndices = GroomCollection.AddAttribute<FIntVector4>(FName(UE::Groom::FGroomGuidesFacade::PointBoneIndicesAttribute), FacadeType::VerticesGroup);
+					TManagedArray<FVector4f>& BoneWeights = GroomCollection.AddAttribute<FVector4f>(FName(UE::Groom::FGroomGuidesFacade::PointBoneWeightsAttribute), FacadeType::VerticesGroup);
+				
+					const int32 NumVertices = GroomFacade.GetNumVertices();
+					for(int32 VertexIndex = 0; VertexIndex < NumVertices; ++VertexIndex)
+					{
+						FIntVector4 PointIndices(INDEX_NONE);
+						FVector4f PointWeights(0.0f);
+
+						const int32 NumWeights = FMath::Min(4, (*IndicesArray)[VertexIndex].Num());
+						float TotalWeight = 0.0;
+						for(int32 BoneIdx = 0; BoneIdx < NumWeights; ++BoneIdx)
+						{
+							PointIndices[BoneIdx] = (*IndicesArray)[VertexIndex][BoneIdx];
+							PointWeights[BoneIdx] = (*WeightsArray)[VertexIndex][BoneIdx];
+							TotalWeight += PointWeights[BoneIdx];
+						}
+						if(TotalWeight != 0.0)
+						{
+							for(int32 BoneIdx = 0; BoneIdx < NumWeights; ++BoneIdx)
+							{
+								PointWeights[BoneIdx] /= TotalWeight;
+							}
+						}
+						BoneIndices[VertexIndex] = PointIndices;
+						BoneWeights[VertexIndex] = PointWeights;
+					}
+				}
+			}
+		}
+	}
+}
+
+void FBuildGuidesSkinningDataflowNode::Evaluate(UE::Dataflow::FContext& Context, const FDataflowOutput* Out) const
+{
+	if (Out->IsA<FManagedArrayCollection>(&Collection))
+	{
+		FManagedArrayCollection GroomCollection = GetValue<FManagedArrayCollection>(Context, &Collection);
+
+		if(CurvesType == EGroomCollectionType::Guides)
+		{
+			UE::Groom::Private::BuildSkinningData<UE::Groom::FGroomGuidesFacade>(GroomCollection, GroupIndex, SkeletalMesh, LODIndex, RelativeTransform);
+		}
+		else
+		{
+			UE::Groom::Private::BuildSkinningData<UE::Groom::FGroomStrandsFacade>(GroomCollection, GroupIndex, SkeletalMesh, LODIndex, RelativeTransform);
+		}
+		
+		SetValue(Context, MoveTemp(GroomCollection), &Collection);
+	}
+}
+
+TArray<UE::Dataflow::FRenderingParameter> FBuildGuidesSkinningDataflowNode::GetRenderParametersImpl() const
+{
+	if(CurvesType == EGroomCollectionType::Guides)
+	{
+		return { {TEXT("GuidesRender"), FName("FGroomCollection"), {TEXT("Collection")}}};
+	}
+	else
+	{
+		return { {TEXT("StrandsRender"), FName("FGroomCollection"), {TEXT("Collection")}}};
+	}
+}
+
+void FUnpackSkinWeightsDataflowNode::Evaluate(UE::Dataflow::FContext& Context, const FDataflowOutput* Out) const
+{
+	using namespace UE::Dataflow;
+	
+	// Get the pin value if plugged
+	FCollectionAttributeKey BoneIndicesKeyValue = GetBoneIndicesKey();
+	FCollectionAttributeKey BoneWeightsKeyValue = GetBoneWeightsKey();
+
+	if (Out->IsA<FManagedArrayCollection>(&Collection))
+	{
+		// Evaluate in collection
+		FManagedArrayCollection GroomCollection = GetValue<FManagedArrayCollection>(Context, &Collection);
+
+		if(CurvesType == EGroomCollectionType::Guides)
+		{
+			UE::Groom::Private::ExtractSkinningData<UE::Groom::FGroomGuidesFacade>(GroomCollection, BoneIndicesKeyValue, BoneWeightsKeyValue);
+		}
+		else
+		{
+			UE::Groom::Private::ExtractSkinningData<UE::Groom::FGroomStrandsFacade>(GroomCollection, BoneIndicesKeyValue, BoneWeightsKeyValue);
+		}
+
+		SetValue(Context, MoveTemp(GroomCollection), &Collection);
+	}
+	else if (Out->IsA<FCollectionAttributeKey>(&BoneIndicesKey))
+	{
+		SetValue(Context, MoveTemp(BoneIndicesKeyValue), &BoneIndicesKey);
+	}
+	else if (Out->IsA<FCollectionAttributeKey>(&BoneWeightsKey))
+	{
+		SetValue(Context, MoveTemp(BoneWeightsKeyValue), &BoneWeightsKey);
+	}
+}
+
+TArray<UE::Dataflow::FRenderingParameter> FUnpackSkinWeightsDataflowNode::GetRenderParametersImpl() const
+{
+	if(CurvesType == EGroomCollectionType::Guides)
+	{
+		return { {TEXT("GuidesRender"), FName("FGroomCollection"), {TEXT("Collection")}}};
+	}
+	else
+	{
+		return { {TEXT("StrandsRender"), FName("FGroomCollection"), {TEXT("Collection")}}};
+	}
+}
+
+
+FCollectionAttributeKey FUnpackSkinWeightsDataflowNode::GetBoneIndicesKey() const
+{
+	// Get the pin value if plugged
+	FCollectionAttributeKey Key;
+	Key.Group = (CurvesType == EGroomCollectionType::Guides) ? UE::Groom::FGroomGuidesFacade::VerticesGroup.ToString() :
+		UE::Groom::FGroomStrandsFacade::VerticesGroup.ToString();
+	Key.Attribute = BoneIndicesName;
+	return Key;
+}
+
+FCollectionAttributeKey FUnpackSkinWeightsDataflowNode::GetBoneWeightsKey() const
+{
+	// Get the pin value if plugged
+	FCollectionAttributeKey Key;
+	Key.Group = (CurvesType == EGroomCollectionType::Guides) ? UE::Groom::FGroomGuidesFacade::VerticesGroup.ToString() :
+		UE::Groom::FGroomStrandsFacade::VerticesGroup.ToString();
+	Key.Attribute = BoneWeightsName;
+	return Key;
+}
+
+void FPackSkinWeightsDataflowNode::Evaluate(UE::Dataflow::FContext& Context, const FDataflowOutput* Out) const
+{
+	using namespace UE::Dataflow;
+	
+	// Get the pin value if plugged
+	FCollectionAttributeKey BoneIndicesKeyValue = GetBoneIndicesKey(Context);
+	FCollectionAttributeKey BoneWeightsKeyValue = GetBoneWeightsKey(Context);
+
+	if (Out->IsA<FManagedArrayCollection>(&Collection))
+	{
+		// Evaluate in collection
+		FManagedArrayCollection GroomCollection = GetValue<FManagedArrayCollection>(Context, &Collection);
+
+		if(CurvesType == EGroomCollectionType::Guides)
+		{
+			UE::Groom::Private::ReportSkinningData<UE::Groom::FGroomGuidesFacade>(GroomCollection, BoneIndicesKeyValue, BoneWeightsKeyValue);
+		}
+		else
+		{
+			UE::Groom::Private::ReportSkinningData<UE::Groom::FGroomStrandsFacade>(GroomCollection, BoneIndicesKeyValue, BoneWeightsKeyValue);
+		}
+
+		SetValue(Context, MoveTemp(GroomCollection), &Collection);
+	}
+}
+
+TArray<UE::Dataflow::FRenderingParameter> FPackSkinWeightsDataflowNode::GetRenderParametersImpl() const
+{
+	if(CurvesType == EGroomCollectionType::Guides)
+	{
+		return { {TEXT("GuidesRender"), FName("FGroomCollection"), {TEXT("Collection")}}};
+	}
+	else
+	{
+		return { {TEXT("StrandsRender"), FName("FGroomCollection"), {TEXT("Collection")}}};
+	}
+}
+
+FCollectionAttributeKey FPackSkinWeightsDataflowNode::GetBoneIndicesKey(UE::Dataflow::FContext& Context) const
+{
+	// Get the pin value if plugged
+	FCollectionAttributeKey Key = GetValue(Context, &BoneIndicesKey, BoneIndicesKey);
+
+	// If nothing set used the local value
+	if(Key.Attribute.IsEmpty() && Key.Group.IsEmpty())
+	{
+		Key.Group = (CurvesType == EGroomCollectionType::Guides) ? UE::Groom::FGroomGuidesFacade::VerticesGroup.ToString() :
+																   UE::Groom::FGroomStrandsFacade::VerticesGroup.ToString();
+		Key.Attribute = BoneIndicesName;
+	}
+	return Key;
+}
+
+FCollectionAttributeKey FPackSkinWeightsDataflowNode::GetBoneWeightsKey(UE::Dataflow::FContext& Context) const
+{
+	// Get the pin value if plugged
+	FCollectionAttributeKey Key = GetValue(Context, &BoneWeightsKey, BoneWeightsKey);
+
+	// If nothing set used the local value
+	if(Key.Attribute.IsEmpty() && Key.Group.IsEmpty())
+	{
+		Key.Group = (CurvesType == EGroomCollectionType::Guides) ? UE::Groom::FGroomGuidesFacade::VerticesGroup.ToString() :
+																   UE::Groom::FGroomStrandsFacade::VerticesGroup.ToString();
+		Key.Attribute = BoneWeightsName;
+	}
+	return Key;
+}
+
+
diff --git a/Engine/Plugins/Runtime/HairStrands/Source/HairStrandsDataflow/Private/BuildGuidesSkinningNode.cpp b/Engine/Plugins/Runtime/HairStrands/Source/HairStrandsDataflow/Private/BuildGuidesSkinningNode.cpp
deleted file mode 100644
diff --git a/Engine/Plugins/Runtime/HairStrands/Source/HairStrandsDataflow/Private/HairStrandsDataflowModule.cpp b/Engine/Plugins/Runtime/HairStrands/Source/HairStrandsDataflow/Private/HairStrandsDataflowModule.cpp
--- a/Engine/Plugins/Runtime/HairStrands/Source/HairStrandsDataflow/Private/HairStrandsDataflowModule.cpp
+++ b/Engine/Plugins/Runtime/HairStrands/Source/HairStrandsDataflow/Private/HairStrandsDataflowModule.cpp
@@ -9,7 +9,7 @@
 #include "GroomAssetTerminalNode.h"
 #include "GetGuidesAttributesNodes.h"
 #include "ResampleGuidesPointsNode.h"
-#include "BuildGuidesSkinningNode.h"
+#include "BuildGroomSkinningNodes.h"
 #include "GenerateGuidesCurvesNode.h"
 #include "SmoothGuidesCurvesNode.h"
 #include "GroomDataflowRendering.h"
@@ -58,6 +58,8 @@ void FHairStrandsDataflowModule::StartupModule()
 	DATAFLOW_NODE_REGISTER_CREATION_FACTORY(FSmoothGuidesCurvesDataflowNode);
 	DATAFLOW_NODE_REGISTER_CREATION_FACTORY(FAttachGuidesRootsDataflowNode);
 	DATAFLOW_NODE_REGISTER_CREATION_FACTORY(FBuildGuidesLODsDataflowNode);
+	DATAFLOW_NODE_REGISTER_CREATION_FACTORY(FUnpackSkinWeightsDataflowNode);
+	DATAFLOW_NODE_REGISTER_CREATION_FACTORY(FPackSkinWeightsDataflowNode);
 	
 	UE::Groom::RegisterRenderingCallbacks();
 	
diff --git a/Engine/Plugins/Runtime/HairStrands/Source/HairStrandsDataflow/Public/BuildGroomSkinningNodes.h b/Engine/Plugins/Runtime/HairStrands/Source/HairStrandsDataflow/Public/BuildGroomSkinningNodes.h
new file mode 100644
--- /dev/null
+++ b/Engine/Plugins/Runtime/HairStrands/Source/HairStrandsDataflow/Public/BuildGroomSkinningNodes.h
@@ -0,0 +1,172 @@
+// Copyright Epic Games, Inc. All Rights Reserved.
+
+#pragma once 
+
+#include "CoreMinimal.h"
+#include "GroomAsset.h"
+#include "ChaosLog.h"
+#include "GroomCollectionFacades.h"
+#include "Dataflow/DataflowConnectionTypes.h"
+#include "Dataflow/DataflowCore.h"
+
+#include "BuildGroomSkinningNodes.generated.h"
+
+/** Build the guides skinning by transfering the indices weights from a skelmesh */
+USTRUCT(meta = (Experimental, DataflowGroom))
+struct FBuildGuidesSkinningDataflowNode : public FDataflowNode
+{
+	GENERATED_BODY()
+	
+	DATAFLOW_NODE_DEFINE_INTERNAL(FBuildGuidesSkinningDataflowNode, "TransferSkinWeights", "Groom", "")
+
+public:
+	
+	FBuildGuidesSkinningDataflowNode(const UE::Dataflow::FNodeParameters& InParam, FGuid InGuid = FGuid::NewGuid())
+	: FDataflowNode(InParam, InGuid)
+	{
+		RegisterInputConnection(&Collection);
+		RegisterOutputConnection(&Collection, &Collection);
+	}
+
+	//~ Begin FDataflowNode interface
+	virtual void Evaluate(UE::Dataflow::FContext& Context, const FDataflowOutput* Out) const override;
+	virtual TArray<UE::Dataflow::FRenderingParameter> GetRenderParametersImpl() const override;
+	//~ End FDataflowNode interface
+
+	/** Managed array collection to be used to store datas */
+	UPROPERTY(meta = (DataflowInput, DataflowOutput, DisplayName = "Collection", DataflowPassthrough  = "Collection"))
+	FManagedArrayCollection Collection;
+
+	/** SkeletalMesh used to transfer the skinning weights. Will be stored onto the groom asset */
+	UPROPERTY(EditAnywhere, Category="Skeletal Mesh")
+	TObjectPtr<USkeletalMesh> SkeletalMesh = nullptr;
+
+	/** LOD used to transfer the weights */
+	UPROPERTY(EditAnywhere, Category="Skeletal Mesh", meta = (DisplayName = "LOD Index"))
+	int32 LODIndex = 0;
+	
+	/** Group index on which the dats will be transfered. -1 will transfer on all the groups */
+	UPROPERTY(EditAnywhere, Category="Groom Groups", meta = (DisplayName = "Group Index"))
+	int32 GroupIndex = INDEX_NONE;
+
+	/** The relative transform between the skeletal mesh and the groom asset. */
+	UPROPERTY(EditAnywhere, Category = "Skeletal Mesh")
+	FTransform RelativeTransform;
+
+	/** Type of curves to use to fill the groom collection (guides/strands) */
+	UPROPERTY(EditAnywhere, Category = "Groom Groups", meta = (DisplayName = "Curves Type"))
+	EGroomCollectionType CurvesType = EGroomCollectionType::Guides;
+};
+
+/** Extract the points skinning attributes to fill the vertices ones */
+USTRUCT(meta = (Experimental, DataflowGroom))
+struct FUnpackSkinWeightsDataflowNode : public FDataflowNode
+{
+	GENERATED_BODY()
+	
+	DATAFLOW_NODE_DEFINE_INTERNAL(FUnpackSkinWeightsDataflowNode, "UnpackSkinWeights", "Groom", "")
+
+public:
+	
+	FUnpackSkinWeightsDataflowNode(const UE::Dataflow::FNodeParameters& InParam, FGuid InGuid = FGuid::NewGuid())
+	: FDataflowNode(InParam, InGuid)
+	{
+		RegisterInputConnection(&Collection);
+		RegisterOutputConnection(&Collection, &Collection);
+		RegisterOutputConnection(&BoneIndicesKey);
+		RegisterOutputConnection(&BoneWeightsKey);
+	}
+
+	//~ Begin FDataflowNode interface
+	virtual void Evaluate(UE::Dataflow::FContext& Context, const FDataflowOutput* Out) const override;
+	virtual TArray<UE::Dataflow::FRenderingParameter> GetRenderParametersImpl() const override;
+	//~ End FDataflowNode interface
+
+	/** Managed array collection to be used to store datas */
+	UPROPERTY(meta = (DataflowInput, DataflowOutput, DisplayName = "Collection", DataflowPassthrough  = "Collection"))
+	FManagedArrayCollection Collection;
+
+	/** The name to be set as a bones indices attribute. */
+	UPROPERTY(EditAnywhere, Category = "Vertex Attributes", meta = (DisplayName = "Bone Indices"))
+	FString BoneIndicesName;
+
+	/** The name to be set as a bones weights attribute. */
+	UPROPERTY(EditAnywhere, Category = "Vertex Attributes", meta = (DisplayName = "Bone Weights"))
+	FString BoneWeightsName;
+
+	/** Bone indices key to be used in other nodes if necessary */
+	UPROPERTY(meta = (DisplayName = "Bone Indices Key", DataflowOutput))
+	FCollectionAttributeKey BoneIndicesKey;
+
+	/** Bone weights key to be used in other nodes if necessary */
+	UPROPERTY(meta = (DisplayName = "Bone Weights Key", DataflowOutput))
+	FCollectionAttributeKey BoneWeightsKey;
+
+	/** Type of curves to use to fill the groom collection (guides/strands) */
+	UPROPERTY(EditAnywhere, Category = "Vertex Attributes", meta = (DisplayName = "Curves Type"))
+	EGroomCollectionType CurvesType = EGroomCollectionType::Guides;
+
+private :
+	/** Get the bone indices key */
+	FCollectionAttributeKey GetBoneIndicesKey() const;
+
+	/** Get the bone weights key */
+	FCollectionAttributeKey GetBoneWeightsKey() const;
+};
+
+/** Report the vertices skinning attributes to fill the points ones */
+USTRUCT(meta = (Experimental, DataflowGroom))
+struct FPackSkinWeightsDataflowNode : public FDataflowNode
+{
+	GENERATED_BODY()
+	
+	DATAFLOW_NODE_DEFINE_INTERNAL(FPackSkinWeightsDataflowNode, "PackSkinWeights", "Groom", "")
+
+public:
+	
+	FPackSkinWeightsDataflowNode(const UE::Dataflow::FNodeParameters& InParam, FGuid InGuid = FGuid::NewGuid())
+	: FDataflowNode(InParam, InGuid)
+	{
+		RegisterInputConnection(&Collection);
+		RegisterInputConnection(&BoneIndicesKey);
+		RegisterInputConnection(&BoneWeightsKey);
+		RegisterOutputConnection(&Collection, &Collection);
+	}
+
+	//~ Begin FDataflowNode interface
+	virtual void Evaluate(UE::Dataflow::FContext& Context, const FDataflowOutput* Out) const override;
+	virtual TArray<UE::Dataflow::FRenderingParameter> GetRenderParametersImpl() const override;
+	//~ End FDataflowNode interface
+
+	/** Managed array collection to be used to store datas */
+	UPROPERTY(meta = (DataflowInput, DataflowOutput, DisplayName = "Collection", DataflowPassthrough  = "Collection"))
+	FManagedArrayCollection Collection;
+
+	/** The name to be set as a weight map attribute. */
+	UPROPERTY(EditAnywhere, Category = "Vertex Attributes", meta = (DisplayName = "Bone Indices"))
+	FString BoneIndicesName;
+
+	/** The name to be set as a weight map attribute. */
+	UPROPERTY(EditAnywhere, Category = "Vertex Attributes", meta = (DisplayName = "Bone Weights"))
+	FString BoneWeightsName;
+
+	/** Bone indices key to be used in other nodes if necessary */
+	UPROPERTY(meta = (DisplayName = "Bone Indices Key", DataflowInput))
+	FCollectionAttributeKey BoneIndicesKey;
+
+	/** Bone weights key to be used in other nodes if necessary */
+	UPROPERTY(meta = (DisplayName = "Bone Weights Key", DataflowInput))
+	FCollectionAttributeKey BoneWeightsKey;
+
+	/** Type of curves to use to fill the groom collection (guides/strands) */
+	UPROPERTY(EditAnywhere, Category = "Vertex Attributes", meta = (DisplayName = "Curves Type"))
+	EGroomCollectionType CurvesType = EGroomCollectionType::Guides;
+
+private :
+	/** Get the bone indices key */
+	FCollectionAttributeKey GetBoneIndicesKey(UE::Dataflow::FContext& Context) const;
+
+	/** Get the bone weights key */
+	FCollectionAttributeKey GetBoneWeightsKey(UE::Dataflow::FContext& Context) const;
+};
+
diff --git a/Engine/Plugins/Runtime/HairStrands/Source/HairStrandsDataflow/Public/BuildGuidesSkinningNode.h b/Engine/Plugins/Runtime/HairStrands/Source/HairStrandsDataflow/Public/BuildGuidesSkinningNode.h
deleted file mode 100644
