Add source curve/point count to hair data.

This allows correct dispatch/bound check for collection attributes computed on source data.

diff --git a/Engine/Plugins/Runtime/HairStrands/Source/HairStrandsCore/Private/GroomBuilder.cpp b/Engine/Plugins/Runtime/HairStrands/Source/HairStrandsCore/Private/GroomBuilder.cpp
--- a/Engine/Plugins/Runtime/HairStrands/Source/HairStrandsCore/Private/GroomBuilder.cpp
+++ b/Engine/Plugins/Runtime/HairStrands/Source/HairStrandsCore/Private/GroomBuilder.cpp
@@ -54,7 +54,7 @@ bool DoesHairStrandsSupportCompressedPosition();
 
 FString FGroomBuilder::GetVersion()
 {
-	FString Out = TEXT("v19_Debug4");
+	FString Out = TEXT("v20");
 	if (IsRenderCurveEnabled())
 	{
 		Out += TEXT("_Curve3e");
@@ -604,6 +604,8 @@ namespace HairStrandsBuilder
 		OutBulkData.Header.BoundingBox.IsValid = 1;
 		OutBulkData.Header.CurveCount = HairStrands.GetNumCurves();
 		OutBulkData.Header.PointCount = HairStrands.GetNumPoints();
+		OutBulkData.Header.SourceCurveCount = HairStrands.StrandsCurves.SourceCurveCount;
+		OutBulkData.Header.SourcePointCount = HairStrands.StrandsPoints.SourcePointCount;
 		OutBulkData.Header.MinPointPerCurve = MinPointPerCurve;
 		OutBulkData.Header.MaxPointPerCurve = MaxPointPerCurve;
 		OutBulkData.Header.AvgPointPerCurve = FMath::DivideAndRoundDown(AccPointPerCurve, NumCurves);
@@ -2147,6 +2149,22 @@ bool FGroomBuilder::BuildHairDescriptionGroups(const FHairDescription& HairDescr
 
 		int32 CurveNumVertices = StrandNumVertices[StrandID];
 
+		FHairStrandsRawDatas* CurrentHairStrandsDatas = nullptr;
+		FHairDescriptionGroup& Group = FindOrAdd(StrandID, Out);
+
+		// Source curve/point count needs to be kept as vanilla as possible, so we count even invalid curve and invalid point
+		if (bIsGuide)
+		{
+			Group.Guides.StrandsCurves.SourceCurveCount += 1;
+			Group.Guides.StrandsPoints.SourcePointCount += CurveNumVertices;
+		}
+		else
+		{
+			Group.Strands.StrandsCurves.SourceCurveCount += 1;
+			Group.Strands.StrandsPoints.SourcePointCount += CurveNumVertices;
+		}
+
+		// Check if the current curve has any data
 		if (CurveNumVertices <= 0)
 		{
 			continue;
@@ -2167,8 +2185,6 @@ bool FGroomBuilder::BuildHairDescriptionGroups(const FHairDescription& HairDescr
 			}
 		}
 
-		FHairStrandsRawDatas* CurrentHairStrandsDatas = nullptr;
-		FHairDescriptionGroup& Group = FindOrAdd(StrandID, Out);
 		bool bNumCurveValid = false;
 		if (bHasInvalidPosition)
 		{
@@ -2847,6 +2863,8 @@ void Decimate(
 	OutData.StrandsCurves.SetNum(OutCurveCount, InAttributes, bBuildSourceMapping);
 	OutData.StrandsPoints.SetNum(OutTotalPointCount, InAttributes, bBuildSourceMapping);
 	OutData.HairDensity = InData.HairDensity;
+	OutData.StrandsCurves.SourceCurveCount = InData.StrandsCurves.SourceCurveCount;
+	OutData.StrandsPoints.SourcePointCount = InData.StrandsPoints.SourcePointCount;
 
 	uint32 OutPointOffset = 0; 
 	for (uint32 OutCurveIndex = 0; OutCurveIndex < OutCurveCount; ++OutCurveIndex)
diff --git a/Engine/Plugins/Runtime/HairStrands/Source/HairStrandsCore/Private/GroomTextureBuilder.cpp b/Engine/Plugins/Runtime/HairStrands/Source/HairStrandsCore/Private/GroomTextureBuilder.cpp
--- a/Engine/Plugins/Runtime/HairStrands/Source/HairStrandsCore/Private/GroomTextureBuilder.cpp
+++ b/Engine/Plugins/Runtime/HairStrands/Source/HairStrandsCore/Private/GroomTextureBuilder.cpp
@@ -1170,6 +1170,8 @@ static bool TraceTextures(
 			Instance.Common.GroupCount = GroupCount;
 			Instance.Common.PointCount = GroupResource.Strands.RestResource->BulkData.GetNumPoints();
 			Instance.Common.CurveCount = GroupResource.Strands.RestResource->BulkData.GetNumCurves();
+			Instance.Common.SourcePointCount = GroupResource.Strands.RestResource->BulkData.GetNumSourcePoints();
+			Instance.Common.SourceCurveCount = GroupResource.Strands.RestResource->BulkData.GetNumSourceCurves();
 			Instance.Common.Radius = RenderingData.GeometrySettings.HairWidth * 0.5f;
 			Instance.Common.RootScale = 1.0f;
 			Instance.Common.TipScale = 1.0f;
diff --git a/Engine/Plugins/Runtime/HairStrands/Source/HairStrandsCore/Private/HairStrandsDatas.cpp b/Engine/Plugins/Runtime/HairStrands/Source/HairStrandsCore/Private/HairStrandsDatas.cpp
--- a/Engine/Plugins/Runtime/HairStrands/Source/HairStrandsCore/Private/HairStrandsDatas.cpp
+++ b/Engine/Plugins/Runtime/HairStrands/Source/HairStrandsCore/Private/HairStrandsDatas.cpp
@@ -54,6 +54,7 @@ void FHairStrandsCurves::Reset()
 	CurvesClosestGuideIDs.Reset();
 	CurvesClosestGuideWeights.Reset();
 	CurvesMapping.Reset();
+	SourceCurveCount = 0;
 }
 
 void FHairStrandsCurves::SetNum(const uint32 NumCurves, uint32 InAttributes, bool bHasMapping)
@@ -87,6 +88,7 @@ void FHairStrandsPoints::Reset()
 	PointsRoughness.Reset();
 	PointsAO.Reset();
 	PointsMapping.Reset();
+	SourcePointCount = 0;
 }
 
 void FHairStrandsPoints::SetNum(const uint32 NumPoints, uint32 InAttributes, bool bHasMapping)
@@ -607,6 +609,9 @@ void FHairStrandsBulkData::SerializeHeader(FArchive& Ar, UObject* Owner)
 	Ar << Header.ImportedAttributes;
 	Ar << Header.ImportedAttributeFlags;
 
+	Ar << Header.SourceCurveCount;
+	Ar << Header.SourcePointCount;
+
 	Ar << Header.CurveToPointCount;
 	Ar << Header.CoverageScales;
 
diff --git a/Engine/Plugins/Runtime/HairStrands/Source/HairStrandsCore/Private/HairStrandsInterpolation.cpp b/Engine/Plugins/Runtime/HairStrands/Source/HairStrandsCore/Private/HairStrandsInterpolation.cpp
--- a/Engine/Plugins/Runtime/HairStrands/Source/HairStrandsCore/Private/HairStrandsInterpolation.cpp
+++ b/Engine/Plugins/Runtime/HairStrands/Source/HairStrandsCore/Private/HairStrandsInterpolation.cpp
@@ -1817,6 +1817,8 @@ FHairGroupPublicData::FVertexFactoryInput InternalComputeHairStrandsVertexInputD
 	OutVFInput.Strands.Common.bSingleGuide = (Instance->Strands.InterpolationResource && (Instance->Strands.InterpolationResource->BulkData.Header.Flags & FHairStrandsInterpolationBulkData::DataFlags_HasSingleGuideData) != 0) ? 1 : 0;
 	OutVFInput.Strands.Common.bCurveRemapping = (Instance->Strands.RestResource && (Instance->Strands.RestResource->BulkData.Header.Flags & FHairStrandsBulkData::DataFlags_HasCurveMapping) != 0) ? 1 : 0;
 	OutVFInput.Strands.Common.bPointRemapping = (Instance->Strands.RestResource && (Instance->Strands.RestResource->BulkData.Header.Flags & FHairStrandsBulkData::DataFlags_HasPointMapping) != 0) ? 1 : 0;
+	OutVFInput.Strands.Common.SourcePointCount = bUseGuideAttributeOffsets ? Instance->Guides.RestResource->GetSourcePointCount() : Instance->Strands.RestResource->GetSourcePointCount();
+	OutVFInput.Strands.Common.SourceCurveCount = bUseGuideAttributeOffsets ? Instance->Guides.RestResource->GetSourceCurveCount() : Instance->Strands.RestResource->GetSourceCurveCount();
 
 	// Add local transform assignement here ...
 
diff --git a/Engine/Plugins/Runtime/HairStrands/Source/HairStrandsCore/Public/GroomResources.h b/Engine/Plugins/Runtime/HairStrands/Source/HairStrandsCore/Public/GroomResources.h
--- a/Engine/Plugins/Runtime/HairStrands/Source/HairStrandsCore/Public/GroomResources.h
+++ b/Engine/Plugins/Runtime/HairStrands/Source/HairStrandsCore/Public/GroomResources.h
@@ -208,6 +208,9 @@ struct FHairStrandsRestResource : public FHairCommonResource
 	inline uint32 GetPointCount() const { return BulkData.GetNumPoints(); }
 	inline uint32 GetCurveCount() const { return BulkData.GetNumCurves(); }
 
+	inline uint32 GetSourcePointCount() const { return BulkData.GetNumSourcePoints(); }
+	inline uint32 GetSourceCurveCount() const { return BulkData.GetNumSourceCurves(); }
+
 	uint32 CachedTangentPointCount = 0;
 };
 
diff --git a/Engine/Plugins/Runtime/HairStrands/Source/HairStrandsCore/Public/HairStrandsDatas.h b/Engine/Plugins/Runtime/HairStrands/Source/HairStrandsCore/Public/HairStrandsDatas.h
--- a/Engine/Plugins/Runtime/HairStrands/Source/HairStrandsCore/Public/HairStrandsDatas.h
+++ b/Engine/Plugins/Runtime/HairStrands/Source/HairStrandsCore/Public/HairStrandsDatas.h
@@ -528,6 +528,9 @@ struct FHairStrandsPoints
 
 	/** Points source mapping */
 	TArray<uint32> PointsMapping;
+
+	/** Source point count */
+	uint32 SourcePointCount = 0;
 };
 
 /** Hair strands Curves attribute */
@@ -578,6 +581,9 @@ struct FHairStrandsCurves
 	
 	/** Curves source mapping */
 	TArray<uint32> CurvesMapping;
+
+	/** Source curve count */
+	uint32 SourceCurveCount = 0;
 };
 
 /** Hair strands datas that are stored on CPU */
@@ -655,6 +661,8 @@ struct FHairStrandsBulkData : FHairStrandsBulkCommon
 
 	uint32 GetNumCurves() const { return Header.CurveCount;  };
 	uint32 GetNumPoints() const { return Header.PointCount; };
+	uint32 GetNumSourceCurves() const { return Header.SourceCurveCount; };
+	uint32 GetNumSourcePoints() const { return Header.SourcePointCount; };
 	float  GetMaxLength() const	{ return Header.MaxLength; };
 	float  GetMaxRadius() const { return Header.MaxRadius; }
 	FVector GetPositionOffset() const { return Header.BoundingBox.GetCenter(); }
@@ -684,6 +692,10 @@ struct FHairStrandsBulkData : FHairStrandsBulkCommon
 		uint32 ImportedAttributes = 0;
 		uint32 ImportedAttributeFlags = 0;
 
+		/** Imported (Source) point/curve info */
+		uint32 SourceCurveCount = 0;
+		uint32 SourcePointCount = 0;
+
 		// Map 'curve' count to 'point' count (used for CLOD)
 		TArray<uint32> CurveToPointCount;
 
diff --git a/Engine/Shaders/Private/HairStrands/HairStrandsVertexFactoryCommon.ush b/Engine/Shaders/Private/HairStrands/HairStrandsVertexFactoryCommon.ush
--- a/Engine/Shaders/Private/HairStrands/HairStrandsVertexFactoryCommon.ush
+++ b/Engine/Shaders/Private/HairStrands/HairStrandsVertexFactoryCommon.ush
@@ -68,6 +68,10 @@ float3 ReadCardPositionOffset(StructuredBuffer<float4> In, uint InIndex, uint In
 	uint						Prefix##_CurveCount; \
 	uint						Prefix##_bSimulation; \
 	uint						Prefix##_bSingleGuide; \
+	uint						Prefix##_bPointRemapping; \
+	uint						Prefix##_bCurveRemapping; \
+	uint						Prefix##_SourcePointCount; \
+	uint						Prefix##_SourceCurveCount; \
 	float4x4					Prefix##_LocalToWorldPrimitiveTransform; \
 	float4x4					Prefix##_LocalToTranslatedWorldPrimitiveTransform; \
 	uint4						Prefix##_CurveAttributeOffsets[HAIR_CURVE_ATTRIBUTE_OFFSET_COUNT]; \
diff --git a/Engine/Source/Runtime/Renderer/Public/HairStrandsInterface.h b/Engine/Source/Runtime/Renderer/Public/HairStrandsInterface.h
--- a/Engine/Source/Runtime/Renderer/Public/HairStrandsInterface.h
+++ b/Engine/Source/Runtime/Renderer/Public/HairStrandsInterface.h
@@ -59,6 +59,8 @@ BEGIN_SHADER_PARAMETER_STRUCT(FHairStrandsInstanceCommonParameters, RENDERER_API
 	SHADER_PARAMETER(uint32, bSingleGuide)
 	SHADER_PARAMETER(uint32, bPointRemapping)
 	SHADER_PARAMETER(uint32, bCurveRemapping)
+	SHADER_PARAMETER(uint32, SourcePointCount)
+	SHADER_PARAMETER(uint32, SourceCurveCount)
 	SHADER_PARAMETER(FVector3f, PositionOffset)
 	SHADER_PARAMETER(FVector3f, PrevPositionOffset)
 	SHADER_PARAMETER(FMatrix44f, LocalToWorldPrimitiveTransform)
