Fix incorrect cache interpolation.

Propagate hair description flags, and building flags from source to bulk data to select if groom cache interpolation data has or not extra ending point, and requires interpolation correction.

diff --git a/Engine/Plugins/Runtime/HairStrands/Source/HairStrandsCore/Private/GroomAsset.cpp b/Engine/Plugins/Runtime/HairStrands/Source/HairStrandsCore/Private/GroomAsset.cpp
--- a/Engine/Plugins/Runtime/HairStrands/Source/HairStrandsCore/Private/GroomAsset.cpp
+++ b/Engine/Plugins/Runtime/HairStrands/Source/HairStrandsCore/Private/GroomAsset.cpp
@@ -542,7 +542,7 @@ static bool BuildHairGroup_Strands(
 
 		FHairStrandsDatas StrandsData;
 		FHairStrandsDatas GuidesData;
-		FGroomBuilder::BuildData(HairGroup, InHairGroupsInterpolation[GroupIndex], OutHairGroupsInfo[GroupIndex], StrandsData, GuidesData, true, true, bBuildSourceMapping);
+		FGroomBuilder::BuildData(HairGroup, InHairGroupsInterpolation[GroupIndex], OutHairGroupsInfo[GroupIndex], StrandsData, GuidesData, true/*bAllowCurveReordering*/, true/*bApplyDecimation*/, bBuildSourceMapping);
 
 		FGroomBuilder::BuildBulkData(HairGroup.Info, GuidesData, OutHairGroupsData[GroupIndex].Guides.BulkData, false /*bAllowCompression*/);
 		FGroomBuilder::BuildBulkData(HairGroup.Info, StrandsData, OutHairGroupsData[GroupIndex].Strands.BulkData, true /*bAllowCompression*/);
diff --git a/Engine/Plugins/Runtime/HairStrands/Source/HairStrandsCore/Private/GroomBuilder.cpp b/Engine/Plugins/Runtime/HairStrands/Source/HairStrandsCore/Private/GroomBuilder.cpp
--- a/Engine/Plugins/Runtime/HairStrands/Source/HairStrandsCore/Private/GroomBuilder.cpp
+++ b/Engine/Plugins/Runtime/HairStrands/Source/HairStrandsCore/Private/GroomBuilder.cpp
@@ -54,7 +54,7 @@ bool DoesHairStrandsSupportCompressedPosition();
 
 FString FGroomBuilder::GetVersion()
 {
-	FString Out = TEXT("v20");
+	FString Out = TEXT("v20_debug4");
 	if (IsRenderCurveEnabled())
 	{
 		Out += TEXT("_Curve3e");
@@ -115,6 +115,7 @@ namespace HairStrandsBuilder
 
 		Out.StrandsCurves = In.StrandsCurves;
 		Out.StrandsPoints = In.StrandsPoints;
+		Out.Flags = In.Flags;
 		Out.BoundingBox.Min = {  FLT_MAX,  FLT_MAX ,  FLT_MAX };
 		Out.BoundingBox.Max = { -FLT_MAX, -FLT_MAX , -FLT_MAX };
 		Out.BoundingBox.IsValid = 0;
@@ -204,7 +205,7 @@ namespace HairStrandsBuilder
 	}
 
 	/** Build the bulk/packed datas for gpu rendering/simulation */
-	void BuildBulkData(const FHairStrandsDatas& HairStrands, const TArray<uint8>& RandomSeeds, const TArray<GroomBuilder_Voxelization::FCoverageScale>& InCoverageScales, FHairStrandsBulkData& OutBulkData, bool bAllowTranscoding, uint32 InGroupFlags)
+	void BuildBulkData(const FHairStrandsDatas& HairStrands, const TArray<uint8>& RandomSeeds, const TArray<GroomBuilder_Voxelization::FCoverageScale>& InCoverageScales, FHairStrandsBulkData& OutBulkData, bool bAllowTranscoding)
 	{
 		TRACE_CPUPROFILER_EVENT_SCOPE(HairStrandsBuilder::BuildBulkData);
 
@@ -619,14 +620,22 @@ namespace HairStrandsBuilder
 
 		// Transfer group trimmed pt./curve info from group to header. This allows to display this information 
 		// in the groom editor by reading the bulk data header
-		if (InGroupFlags & uint32(EHairGroupInfoFlags::HasTrimmedCurve))
+		if (HairStrands.HasFlags(EHairStrandsFlags::HasTrimmedCurve))
 		{
 			OutBulkData.Header.Flags |= FHairStrandsBulkData::DataFlags_HasTrimmedCurve;
 		}
-		if (InGroupFlags & uint32(EHairGroupInfoFlags::HasTrimmedPoint))
+		if (HairStrands.HasFlags(EHairStrandsFlags::HasTrimmedPoint))
 		{
 			OutBulkData.Header.Flags |= FHairStrandsBulkData::DataFlags_HasTrimmedPoint;
 		}
+		if (HairStrands.HasFlags(EHairStrandsFlags::HasShuffledCurve))
+		{
+			OutBulkData.Header.Flags |= FHairStrandsBulkData::DataFlags_HasShuffledCurves;
+		}
+		if (HairStrands.HasFlags(EHairStrandsFlags::HasExtraEndingPoint))
+		{
+			OutBulkData.Header.Flags |= FHairStrandsBulkData::DataFlags_HasExtraEndingPoint;
+		}
 	
 		// Transcoding
 		const bool bHasTranscodedPosition = OutTranscodedPositions.Num() > 0;
@@ -2030,8 +2039,6 @@ bool FGroomBuilder::BuildHairDescriptionGroups(const FHairDescription& HairDescr
 	TStrandAttributesConstRef<int> StrandIDs = HairDescription.StrandAttributes().GetAttributesRef<int>(HairAttribute::Strand::ID);
 	TStrandAttributesConstRef<FName> GroupNames = HairDescription.StrandAttributes().GetAttributesRef<FName>(HairAttribute::Strand::GroupName);
 
-	bool bImportGuides = true;
-
 	// For precomputed weights, there is two variants with single or triplet IDs/weights
 	// Single
 	TStrandAttributesConstRef<int> ClosestGuide		  = HairDescription.StrandAttributes().GetAttributesRef<int>(HairAttribute::Strand::ClosestGuides);
@@ -2044,7 +2051,7 @@ bool FGroomBuilder::BuildHairDescriptionGroups(const FHairDescription& HairDescr
 	// must include StrandID attribute since ClosestGuides references those IDs
 	const bool bPrecomputedWeight1 = ClosestGuide.IsValid()  && GuideWeight.IsValid();
 	const bool bPrecomputedWeight3 = ClosestGuides.IsValid() && GuideWeights.IsValid();
-	const bool bCanUseClosestGuidesAndWeights = bImportGuides && StrandIDs.IsValid() && (bPrecomputedWeight1 || bPrecomputedWeight3);
+	const bool bCanUseClosestGuidesAndWeights = StrandIDs.IsValid() && (bPrecomputedWeight1 || bPrecomputedWeight3);
 	check(bCanUseClosestGuidesAndWeights == HairDescription.HasAttribute(EHairAttribute::PrecomputedGuideWeights)); // Sanity check
 
 	auto FindOrAdd = [&GroupIDs, &GroupNames](FStrandID StrandID, FHairDescriptionGroups& OutGroups) -> FHairDescriptionGroup&
@@ -2185,12 +2192,25 @@ bool FGroomBuilder::BuildHairDescriptionGroups(const FHairDescription& HairDescr
 			}
 		}
 
+		auto SetFlag = [bIsGuide, &Group](EHairStrandsFlags In)
+		{
+			if (bIsGuide)
+			{
+				Group.Guides.Flags |= uint32(In);
+			}
+			else
+			{
+				Group.Strands.Flags |= uint32(In);
+			}
+		};
+
 		bool bNumCurveValid = false;
 		if (bHasInvalidPosition)
 		{
 			// Remove curve with NaN positions
 			GlobalVertexIndex += CurveNumVertices;
 			Group.Info.Flags |= uint32(EHairGroupInfoFlags::HasInvalidPoint);
+			SetFlag(EHairStrandsFlags::HasInvalidPoint);
 			continue;
 		}
 		else if (!bIsGuide)
@@ -2202,7 +2222,7 @@ bool FGroomBuilder::BuildHairDescriptionGroups(const FHairDescription& HairDescr
 				++Group.Info.NumCurves;
 			}
 		}
-		else if (bImportGuides)
+		else
 		{
 			CurrentHairStrandsDatas = &Group.Guides;
 			if (Group.Info.NumGuides < HAIR_MAX_NUM_CURVE_PER_GROUP)
@@ -2211,12 +2231,6 @@ bool FGroomBuilder::BuildHairDescriptionGroups(const FHairDescription& HairDescr
 				++Group.Info.NumGuides;
 			}
 		}
-		else
-		{
-			// A guide but don't want to import it, so skip it
-			GlobalVertexIndex += CurveNumVertices;
-			continue;
-		}
 
 		// If the current curve has more control point than allows, issue a warning
 		if (CurveNumVertices > HAIR_MAX_NUM_POINT_PER_CURVE && !bPointCountWarningIssued)
@@ -2226,6 +2240,7 @@ bool FGroomBuilder::BuildHairDescriptionGroups(const FHairDescription& HairDescr
 			// * This is visible in Groom editor
 			//UE_LOG(LogGroomBuilder, Warning, TEXT("[Groom] Groom contains strands with more than %d control points. Control points beyond that limit will be trimmed"), uint32(HAIR_MAX_NUM_POINT_PER_CURVE));
 			Group.Info.Flags |= uint32(EHairGroupInfoFlags::HasTrimmedPoint);
+			SetFlag(EHairStrandsFlags::HasTrimmedPoint);
 		}
 
 		// If the current group has reached the max number of curve, skip them
@@ -2240,7 +2255,7 @@ bool FGroomBuilder::BuildHairDescriptionGroups(const FHairDescription& HairDescr
 				bCurveCountWarningIssued = true;
 			}
 			Group.Info.Flags |= uint32(EHairGroupInfoFlags::HasTrimmedCurve);
-
+			SetFlag(EHairStrandsFlags::HasTrimmedCurve);
 
 			// Skip it
 			GlobalVertexIndex += CurveNumVertices;
@@ -2254,6 +2269,10 @@ bool FGroomBuilder::BuildHairDescriptionGroups(const FHairDescription& HairDescr
 
 		const bool bAddEndingControlPoint = NeedsToAddEndingControlPoint(CurveNumVertices) && bAllowAddEndControlPoint && !bIsGuide;
 		CurrentHairStrandsDatas->StrandsCurves.CurvesCount.Add(FMath::Min(uint32(CurveNumVertices + (bAddEndingControlPoint ? 1u : 0u)), HAIR_MAX_NUM_POINT_PER_CURVE));
+		if (bAddEndingControlPoint)
+		{
+			SetFlag(EHairStrandsFlags::HasExtraEndingPoint);
+		}
 
 		if (bCanUseClosestGuidesAndWeights)
 		{
@@ -2525,6 +2544,19 @@ void FGroomBuilder::BuildData(
 		{
 			OutRen = HairStrandsBuilder::BuildInternalData(InHairDescriptionGroup.Strands);
 			
+			if (bCurveReordering)
+			{
+				OutRen.Flags |= uint32(EHairStrandsFlags::HasShuffledCurve);
+			}
+			if (CurveDecimation < 1)
+			{
+				OutRen.Flags |= uint32(EHairStrandsFlags::HasDecimatedCurve);
+			}
+			if (VertexDecimation < 1)
+			{
+				OutRen.Flags |= uint32(EHairStrandsFlags::HasDecimatedPoint);
+			}
+			
 			// Decimate
 			if (CurveDecimation < 1 || VertexDecimation < 1 || bCurveReordering)
 			{
@@ -2624,7 +2656,7 @@ void FGroomBuilder::BuildBulkData(
 	}
 
 	// 3. Build bulk data
-	HairStrandsBuilder::BuildBulkData(InData, CurveSeeds, CoverageScales, OutBulkData, bAllowTranscoding, InInfo.Flags);
+	HairStrandsBuilder::BuildBulkData(InData, CurveSeeds, CoverageScales, OutBulkData, bAllowTranscoding);
 }
 
 void FGroomBuilder::BuildInterplationData(
@@ -2912,6 +2944,8 @@ void Decimate(
 		OutPointOffset += OutPointCount;
 	}
 
+	// Copy flags from the source (trimmed curve/point, added ending point, shuffle curve, ...)
+	OutData.Flags = InData.Flags;
 	OutData = HairStrandsBuilder::BuildInternalData(OutData);
 }
 } // namespace FHairStrandsDecimation
diff --git a/Engine/Plugins/Runtime/HairStrands/Source/HairStrandsCore/Private/GroomManager.cpp b/Engine/Plugins/Runtime/HairStrands/Source/HairStrandsCore/Private/GroomManager.cpp
--- a/Engine/Plugins/Runtime/HairStrands/Source/HairStrandsCore/Private/GroomManager.cpp
+++ b/Engine/Plugins/Runtime/HairStrands/Source/HairStrandsCore/Private/GroomManager.cpp
@@ -1111,7 +1111,7 @@ static void RunHairStrandsDeformation_Guide(
 					InstanceData.Instance->Guides.RestResource->GetCurveCount(),
 					InterpolationFactor,
 					1.f /*MaxHairRadius*/,
-					false /*bAllowAddedControlPoint - Guide can't have an extra control point (see FGroomBuilder::BuildHairDescriptionGroups())*/,
+					InstanceData.Instance->Guides.RestResource->BulkData.HasFlags(FHairStrandsBulkData::EDataFlags::DataFlags_HasExtraEndingPoint),
 					CacheResources0,
 					CacheResources1,
 					RegisterAsSRV(GraphBuilder, InstanceData.Instance->Guides.RestResource->PositionBuffer),
@@ -1597,7 +1597,7 @@ static void RunHairStrandsInterpolation_Strands(
 					InstanceData.ActiveCurveCount,
 					InterpolationFactor,
 					InstanceData.Instance->Strands.Modifier.HairWidth * 0.5f /*InMaxHairRadius*/,
-					true /*bAllowAddedControlPoint - Render strands can have an extra control point*/,
+					InstanceData.Instance->Strands.RestResource->BulkData.HasFlags(FHairStrandsBulkData::EDataFlags::DataFlags_HasExtraEndingPoint),
 					CacheResources0,
 					CacheResources1,
 					RegisterAsSRV(GraphBuilder, InstanceData.Instance->Strands.RestResource->PositionBuffer),
diff --git a/Engine/Plugins/Runtime/HairStrands/Source/HairStrandsCore/Private/HairStrandsInterpolation.cpp b/Engine/Plugins/Runtime/HairStrands/Source/HairStrandsCore/Private/HairStrandsInterpolation.cpp
--- a/Engine/Plugins/Runtime/HairStrands/Source/HairStrandsCore/Private/HairStrandsInterpolation.cpp
+++ b/Engine/Plugins/Runtime/HairStrands/Source/HairStrandsCore/Private/HairStrandsInterpolation.cpp
@@ -279,7 +279,7 @@ void AddGroomCacheUpdatePass(
 	uint32 CurveCount,
 	float InterpolationFactor,
 	float InMaxHairRadius,
-	bool bAllowAddedControlPoint,
+	bool bHasAddedControlPoint,
 	FGroomCacheResources CacheResources0,
 	FGroomCacheResources CacheResources1,
 	FRDGBufferSRVRef InRestPositionBuffer,
@@ -312,7 +312,7 @@ void AddGroomCacheUpdatePass(
 	Parameters->InterpolationFactor = InterpolationFactor;
 	Parameters->MaxHairRadius = InMaxHairRadius;
 	Parameters->bHasRadiusData = CacheResources0.bHasRadiusData ? 1u : 0u;
-	Parameters->bHasAddedControlPoint = (bAllowAddedControlPoint && GetHairStrandsUsesTriangleStrips()) ? 1u : 0u;
+	Parameters->bHasAddedControlPoint = bHasAddedControlPoint ? 1u : 0u;
 
 	const FIntVector DispatchCount = FIntVector(FMath::DivideAndRoundUp(PointCount, FGroomCacheUpdatePassCS::GetGroupSize()), 1, 1);
 	TShaderMapRef<FGroomCacheUpdatePassCS> ComputeShader(ShaderMap);
diff --git a/Engine/Plugins/Runtime/HairStrands/Source/HairStrandsCore/Public/HairStrandsDatas.h b/Engine/Plugins/Runtime/HairStrands/Source/HairStrandsCore/Public/HairStrandsDatas.h
--- a/Engine/Plugins/Runtime/HairStrands/Source/HairStrandsCore/Public/HairStrandsDatas.h
+++ b/Engine/Plugins/Runtime/HairStrands/Source/HairStrandsCore/Public/HairStrandsDatas.h
@@ -494,6 +494,17 @@ struct FHairStrandsBulkCommon
 ///////////////////////////////////////////////////////////////////////////////////////////////////
 // Strands Data/Bulk
 
+enum class EHairStrandsFlags : uint8
+{
+	HasTrimmedPoint		= 1<<0,
+	HasTrimmedCurve		= 1<<1,
+	HasInvalidPoint		= 1<<2,
+	HasExtraEndingPoint	= 1<<3,
+	HasShuffledCurve	= 1<<4,
+	HasDecimatedPoint	= 1<<5,
+	HasDecimatedCurve	= 1<<6,
+};
+
 /** Hair strands points attribute */
 struct FHairStrandsPoints
 {
@@ -597,6 +608,7 @@ struct FHairStrandsCommonDatas
 
 	UE_API uint32 GetAttributes() const;
 	UE_API uint32 GetAttributeFlags() const;
+	UE_API bool HasFlags(EHairStrandsFlags In) const { return !!(Flags & uint32(In)); }
 
 	UE_API void Reset();
 
@@ -607,6 +619,9 @@ struct FHairStrandsCommonDatas
 
 	/** List of all the strands curves */
 	FHairStrandsCurves StrandsCurves;
+
+	/** General flags (shuffled curve, extra ending point, ...) */
+	uint32 Flags = 0;
 };
 
 /** Hair strands raw datas that are stored on CPU */
@@ -643,7 +658,9 @@ struct FHairStrandsBulkData : FHairStrandsBulkCommon
 		DataFlags_HasTrimmedCurve = 0x8u,		// Source data had more curves per group than the limit and some curves were trimmed
 		DataFlags_HasTrimmedPoint = 0x10u,		// Source data had more points per curve than the limit and some points were trimmed
 		DataFlags_HasPointMapping = 0x20u,		// Contains platform to source points mapping data
-		DataFlags_HasCurveMapping = 0x40u		// Contains platform to source curves mapping data
+		DataFlags_HasCurveMapping = 0x40u,		// Contains platform to source curves mapping data
+		DataFlags_HasShuffledCurves = 0x80u,	// Curve data has been shuffled
+		DataFlags_HasExtraEndingPoint = 0x100u	// Curve contains an extra ending control point
 	};
 
 	UE_API virtual void SerializeHeader(FArchive& Ar, UObject* Owner) override;
@@ -668,6 +685,7 @@ struct FHairStrandsBulkData : FHairStrandsBulkCommon
 	FVector GetPositionOffset() const { return Header.BoundingBox.GetCenter(); }
 	const FBox& GetBounds() const { return Header.BoundingBox; }
 	UE_API uint32 GetSize() const;
+	bool HasFlags(EDataFlags In) const { return !!(Header.Flags & In); }
 
 	uint32 GetCurveAttributeSizeInBytes(uint32 InCurveCount=HAIR_MAX_NUM_CURVE_PER_GROUP) const			{ return InCurveCount > 0 ? FMath::DivideAndRoundUp(FMath::Min(Header.CurveCount, InCurveCount), Header.Strides.CurveAttributeChunkElementCount) * Header.Strides.CurveAttributeChunkStride : 0; }
 	uint32 GetPointAttributeSizeInBytes(uint32 InPointCount=HAIR_MAX_NUM_POINT_PER_GROUP) const			{ return InPointCount > 0 ? FMath::DivideAndRoundUp(FMath::Min(Header.PointCount, InPointCount), Header.Strides.PointAttributeChunkElementCount) * Header.Strides.PointAttributeChunkStride : 0; }
