Fix groom deformer issues

#virtualized

diff --git a/Engine/Plugins/Runtime/HairStrands/Shaders/Private/Deformers/DeformerGroomGuidesRead.ush b/Engine/Plugins/Runtime/HairStrands/Shaders/Private/Deformers/DeformerGroomGuidesRead.ush
--- a/Engine/Plugins/Runtime/HairStrands/Shaders/Private/Deformers/DeformerGroomGuidesRead.ush
+++ b/Engine/Plugins/Runtime/HairStrands/Shaders/Private/Deformers/DeformerGroomGuidesRead.ush
@@ -32,7 +32,7 @@ Buffer<float4> {DataInterfaceName}_TriangleDeformedPositions;
 Buffer<uint> {DataInterfaceName}_CurveBarycentricCoordinates;
 Buffer<uint> {DataInterfaceName}_CurveTriangleIndices;
 
-float3x4 BuildSkinnedTransform(uint CurveIndex, Buffer<float4> TrianglePositions,
+float3x4 BuildGuidesSkinnedTransform(uint CurveIndex, Buffer<float4> TrianglePositions,
 	Buffer<uint> BarycentricCoordinates, Buffer<uint> TriangleIndices)
 {
 	const float3 BarycentricCoord = UnpackBarycentrics(BarycentricCoordinates[CurveIndex]);
@@ -91,7 +91,7 @@ float3x4 ReadCurveRestTransform_{DataInterfaceName}(uint CurveIndex)
 {
 	CurveIndex -= {DataInterfaceName}_BaseCurveIndex;	
 #if ENABLE_SKINNED_TRANSFORM == 1
-		return BuildSkinnedTransform(CurveIndex, {DataInterfaceName}_TriangleRestPositions,
+		return BuildGuidesSkinnedTransform(CurveIndex, {DataInterfaceName}_TriangleRestPositions,
 			{DataInterfaceName}_CurveBarycentricCoordinates, {DataInterfaceName}_CurveTriangleIndices);
 #endif
 
@@ -104,7 +104,7 @@ float3x4 ReadCurveDeformedTransform_{DataInterfaceName}(uint CurveIndex)
 {
 	CurveIndex -= {DataInterfaceName}_BaseCurveIndex;
 #if ENABLE_SKINNED_TRANSFORM == 1
-		return BuildSkinnedTransform(CurveIndex, {DataInterfaceName}_TriangleDeformedPositions,
+		return BuildGuidesSkinnedTransform(CurveIndex, {DataInterfaceName}_TriangleDeformedPositions,
 			{DataInterfaceName}_CurveBarycentricCoordinates, {DataInterfaceName}_CurveTriangleIndices);
 #endif
 
diff --git a/Engine/Plugins/Runtime/HairStrands/Shaders/Private/Deformers/DeformerGroomStrandsRead.ush b/Engine/Plugins/Runtime/HairStrands/Shaders/Private/Deformers/DeformerGroomStrandsRead.ush
--- a/Engine/Plugins/Runtime/HairStrands/Shaders/Private/Deformers/DeformerGroomStrandsRead.ush
+++ b/Engine/Plugins/Runtime/HairStrands/Shaders/Private/Deformers/DeformerGroomStrandsRead.ush
@@ -3,6 +3,7 @@
 #include "/Plugin/ComputeFramework/Private/ComputeKernelCommon.ush"
 #include "/Engine/Private/HairStrands/HairStrandsVertexFactoryCommon.ush"
 #include "/Engine/Private/HairStrands/HairStrandsAttributeCommon.ush"
+#include "/Engine/Private/HairStrands/HairStrandsBindingCommon.ush"
 
 HAIR_STRANDS_INSTANCE_PARAMETERS({DataInterfaceName})
 
@@ -33,13 +34,13 @@ Buffer<float4> {DataInterfaceName}_TriangleDeformedPositions;
 Buffer<uint> {DataInterfaceName}_CurveBarycentricCoordinates;
 Buffer<uint> {DataInterfaceName}_CurveTriangleIndices;
 
-float3x4 BuildSkinnedTransform(uint CurveIndex, Buffer<float4> TrianglePositions,
+float3x4 BuildStrandsSkinnedTransform(uint CurveIndex, Buffer<float4> TrianglePositions,
 	Buffer<uint> CurveBarycentricCoordinates, Buffer<uint> CurveTriangleIndices)
 {
 	const float3 BarycentricCoord = UnpackBarycentrics(CurveBarycentricCoordinates[CurveIndex]);
 	const uint TriangleIndex = CurveTriangleIndices[CurveIndex];
 
-	const FHairMeshBasis MeshBasis = GetHairMeshBasis(TriangleIndex, TrianglePositions, RootBarycentric);
+	const FHairMeshBasis MeshBasis = GetHairMeshBasis(TriangleIndex, TrianglePositions, BarycentricCoord);
 	
 	return float3x4(MeshBasis.BaseT.x, MeshBasis.BaseB.x, MeshBasis.BaseN.x, MeshBasis.BaseO.x,
 					MeshBasis.BaseT.y, MeshBasis.BaseB.y, MeshBasis.BaseN.y, MeshBasis.BaseO.y,
@@ -165,6 +166,7 @@ uint ReadCurveNumPoints_{DataInterfaceName}(uint CurveIndex)
 
 uint ReadPointCurveIndex_{DataInterfaceName}(uint PointIndex)
 {
+	return {DataInterfaceName}_BaseCurveIndex;
 	PointIndex -= {DataInterfaceName}_BasePointIndex;
 	return ReadHairPointToCurveIndex({DataInterfaceName}_PointToCurveBuffer, PointIndex) + {DataInterfaceName}_BaseCurveIndex;
 }
@@ -201,7 +203,7 @@ float3x4 ReadCurveRestTransform_{DataInterfaceName}(uint CurveIndex)
 {
 	CurveIndex -= {DataInterfaceName}_BaseCurveIndex;
 #if ENABLE_SKINNED_TRANSFORM == 1
-		return BuildSkinnedTransform(CurveIndex, {DataInterfaceName}_TriangleRestPositions,
+		return BuildStrandsSkinnedTransform(CurveIndex, {DataInterfaceName}_TriangleRestPositions,
 			{DataInterfaceName}_CurveBarycentricCoordinates, {DataInterfaceName}_CurveTriangleIndices);
 #endif
 
@@ -214,7 +216,7 @@ float3x4 ReadCurveDeformedTransform_{DataInterfaceName}(uint CurveIndex)
 {
 	CurveIndex -= {DataInterfaceName}_BaseCurveIndex;
 #if ENABLE_SKINNED_TRANSFORM == 1
-		return BuildSkinnedTransform(CurveIndex, {DataInterfaceName}_TriangleDeformedPositions,
+		return BuildStrandsSkinnedTransform(CurveIndex, {DataInterfaceName}_TriangleDeformedPositions,
 			{DataInterfaceName}_CurveBarycentricCoordinates, {DataInterfaceName}_CurveTriangleIndices);
 #endif
 
@@ -223,14 +225,14 @@ float3x4 ReadCurveDeformedTransform_{DataInterfaceName}(uint CurveIndex)
 #endif
 }
 
-uint ReadSourceCurveIndex_{DataInterfaceName}(uint CurveIndex)
+uint ReadCurveSourceIndex_{DataInterfaceName}(uint CurveIndex)
 {
 	CurveIndex -= {DataInterfaceName}_BaseCurveIndex;
-	return ReadHairSourceCurveIndex({DataInterfaceName}_CurveRemappingBuffer, CurveIndex, {DataInterfaceName}_bCurveRemapping);
+	return ReadHairSourceCurveIndex({DataInterfaceName}_CurveMappingBuffer, CurveIndex, {DataInterfaceName}_bCurveRemapping);
 }
 
-uint ReadSourcePointIndex_{DataInterfaceName}(uint PointIndex)
+uint ReadPointSourceIndex_{DataInterfaceName}(uint PointIndex)
 {
 	PointIndex -= {DataInterfaceName}_BasePointIndex;
-	return ReadHairSourcePointIndex({DataInterfaceName}_PointRemappingBuffer, PointIndex, {DataInterfaceName}_bPointRemapping);
+	return ReadHairSourcePointIndex({DataInterfaceName}_PointMappingBuffer, PointIndex, {DataInterfaceName}_bPointRemapping);
 }
\ No newline at end of file
diff --git a/Engine/Plugins/Runtime/HairStrands/Source/HairStrandsCore/Private/GroomEdit.cpp b/Engine/Plugins/Runtime/HairStrands/Source/HairStrandsCore/Private/GroomEdit.cpp
--- a/Engine/Plugins/Runtime/HairStrands/Source/HairStrandsCore/Private/GroomEdit.cpp
+++ b/Engine/Plugins/Runtime/HairStrands/Source/HairStrandsCore/Private/GroomEdit.cpp
@@ -246,7 +246,7 @@ void ConvertToGroomAsset(UGroomAsset* OutGroom, const FEditableGroom* InGroom, u
 	// Commit & build data for the groom asset
 	// Internally, this will derived the build data, recreate resources, and 
 	// recreate all components using this groom assets
-	OutGroom->CommitHairDescription(MoveTemp(HairDescription), UGroomAsset::EHairDescriptionType::Edit);
+	OutGroom->CommitHairDescription(MoveTemp(HairDescription), UGroomAsset::EHairDescriptionType::Source);
 	OutGroom->CacheDerivedDatas();
 	OutGroom->OnGroomAssetChanged.Broadcast();
 #endif
diff --git a/Engine/Plugins/Runtime/HairStrands/Source/HairStrandsCore/Private/GroomManager.cpp b/Engine/Plugins/Runtime/HairStrands/Source/HairStrandsCore/Private/GroomManager.cpp
--- a/Engine/Plugins/Runtime/HairStrands/Source/HairStrandsCore/Private/GroomManager.cpp
+++ b/Engine/Plugins/Runtime/HairStrands/Source/HairStrandsCore/Private/GroomManager.cpp
@@ -2600,7 +2600,7 @@ static FHairLOD ComputeHairLODIndex(const FHairGroupInstance* Instance, const TA
 	// Auto LOD
 	const bool bNeedAutoLOD = MaxLODIndexWithStrands != INDEX_NONE && int32(Out.HairLODIndex) <= MaxLODIndexWithStrands;
 	const float OriginalHairLODIndex = Out.HairLODIndex;
-	if (Instance->Strands.ClusterResource && Instance->Strands.RestResource)
+	if (Instance->Strands.ClusterResource && Instance->Strands.RestResource && Instance->bSupportStreaming)
 	{
 		uint32 EffectiveCurveCount = 0;
 		if (bNeedAutoLOD && (Instance->HairGroupPublicData->bAutoLOD || UseHairStrandsForceAutoLOD()))
diff --git a/Engine/Plugins/Runtime/HairStrands/Source/HairStrandsDataflow/Private/GroomAssetTerminalNode.cpp b/Engine/Plugins/Runtime/HairStrands/Source/HairStrandsDataflow/Private/GroomAssetTerminalNode.cpp
--- a/Engine/Plugins/Runtime/HairStrands/Source/HairStrandsDataflow/Private/GroomAssetTerminalNode.cpp
+++ b/Engine/Plugins/Runtime/HairStrands/Source/HairStrandsDataflow/Private/GroomAssetTerminalNode.cpp
@@ -72,7 +72,7 @@ namespace UE::Groom::Private
 	}
 
 	template<typename FacadeType, typename AttributeType>
-	FORCEINLINE void BuildVerticesAttribute(const FManagedArrayCollection& InCollection, FManagedArrayCollection* OutCollection, const int32 NumPoints, const FName& AttributeName)
+	static void BuildVerticesAttribute(const FManagedArrayCollection& InCollection, FManagedArrayCollection* OutCollection, const int32 NumPoints, const FName& AttributeName)
 	{
 		const TManagedArray<AttributeType>& VerticesAttribute = InCollection.GetAttribute<AttributeType>(AttributeName, FacadeType::VerticesGroup);
 		if(OutCollection->NumElements(FacadeType::PointsGroup) != NumPoints)
@@ -92,12 +92,13 @@ namespace UE::Groom::Private
 	}
 
 	template<typename FacadeType>
-	FORCEINLINE void TransferVerticesAttributes(const FManagedArrayCollection& InCollection, FManagedArrayCollection* OutCollection, const int32 NumPoints)
+	static void TransferVerticesAttributes(const FManagedArrayCollection& InCollection, FManagedArrayCollection* OutCollection, const int32 NumPoints)
 	{
 		// Transfer vertices weight maps onto the points to be stored onto the rest collection
 		const TArray<FName> AttributeNames = InCollection.AttributeNames(FacadeType::VerticesGroup);
 		for(const FName& AttributeName : AttributeNames)
 		{
+			
 			if(InCollection.GetAttributeType(AttributeName, FacadeType::VerticesGroup) == EManagedArrayType::FFloatType)
 			{
 				BuildVerticesAttribute<FacadeType, float>(InCollection, OutCollection, NumPoints, AttributeName);
@@ -212,6 +213,11 @@ void FGroomAssetTerminalDataflowNode::SetAssetValue(TObjectPtr<UObject> Asset, U
 				// Convert to groom asset
 				ConvertToGroomAsset(const_cast<UGroomAsset*>(GroomAsset), &EditGroom, EEditableGroomOperations::ControlPoints_Modified);
 			}
+			// To prevent future reconstruction in the BuildData we set the type to be imported
+			for(FHairGroupsInterpolation& GroupInterpolation : GroomAsset->GetHairGroupsInterpolation())
+			{
+				GroupInterpolation.InterpolationSettings.GuideType = EGroomGuideType::Imported;
+			}
 		}  
 
 		if(OutCollection)
diff --git a/Engine/Shaders/Private/HairStrands/HairStrandsBindingCommon.ush b/Engine/Shaders/Private/HairStrands/HairStrandsBindingCommon.ush
--- a/Engine/Shaders/Private/HairStrands/HairStrandsBindingCommon.ush
+++ b/Engine/Shaders/Private/HairStrands/HairStrandsBindingCommon.ush
@@ -1,4 +1,5 @@
 // Copyright Epic Games, Inc. All Rights Reserved.
+#pragma once
 
 #include "HairStrandsPack.ush"
 
diff --git a/Engine/Shaders/Private/HairStrands/HairStrandsInterpolation.usf b/Engine/Shaders/Private/HairStrands/HairStrandsInterpolation.usf
--- a/Engine/Shaders/Private/HairStrands/HairStrandsInterpolation.usf
+++ b/Engine/Shaders/Private/HairStrands/HairStrandsInterpolation.usf
@@ -779,13 +779,19 @@ void MainCS(uint3 DispatchThreadId : SV_DispatchThreadID, uint ThreadIndex : SV_
 				#endif
 
 				// 3. Apply deformer
-				#if PERMUTATION_DEFORMER
+				#if PERMUTATION_DEFORMER == DEFORMER_ADDITIVE
 				{
 					const FPackedHairPosition PackedDeformerPosition = ReadPackedHairPosition(RenDeformerPositionBuffer, VertexIndex);
 					const float3 DeformerPosition = UnpackHairControlPointPosition(PackedDeformerPosition, InRenHairPositionOffset);
 					OutPosition = ApplyDeformer(RenTri, RestPosition, DeformerPosition, OutPosition);
 					CopyPackedAttribute(PackedRestPosition, PackedDeformerPosition);
 				}
+				#elif PERMUTATION_DEFORMER == DEFORMER_OVERWRITE
+				{
+					const FPackedHairPosition PackedDeformerPosition = ReadPackedHairPosition(RenDeformerPositionBuffer, VertexIndex);
+					OutPosition = UnpackHairControlPointPosition(PackedDeformerPosition, InRenHairPositionOffset);
+					CopyPackedAttribute(PackedRestPosition, PackedDeformerPosition);
+				}
 				#endif
 
 				// 4. Write out the final position
