Groom attribute external buffers

diff --git a/Engine/Plugins/Runtime/HairStrands/Source/HairStrandsCore/Private/GroomResources.cpp b/Engine/Plugins/Runtime/HairStrands/Source/HairStrandsCore/Private/GroomResources.cpp
--- a/Engine/Plugins/Runtime/HairStrands/Source/HairStrandsCore/Private/GroomResources.cpp
+++ b/Engine/Plugins/Runtime/HairStrands/Source/HairStrandsCore/Private/GroomResources.cpp
@@ -1375,6 +1375,11 @@ void FHairStrandsRestResource::InternalRelease()
 	PointMappingBuffer.Release();
 	MaxAvailableCurveCount = 0;
 	CachedTangentPointCount = 0;
+	
+	for (TPair<FString, FRDGExternalBuffer>& ExternalBuffer : ExternalBuffers)
+	{
+		ExternalBuffer.Value.Release();
+	}
 }
 
 void FHairStrandsRestResource::InternalResetLoadedSize()
@@ -1466,6 +1471,11 @@ void FHairStrandsDeformedResource::InternalRelease()
 
 	DeformedOffsetBuffer[0].Release();
 	DeformedOffsetBuffer[1].Release();
+
+	for (TPair<FString, FRDGExternalBuffer>& ExternalBuffer : ExternalBuffers)
+	{
+		ExternalBuffer.Value.Release();
+	}
 }
 
 /////////////////////////////////////////////////////////////////////////////////////////
diff --git a/Engine/Plugins/Runtime/HairStrands/Source/HairStrandsCore/Public/GroomResources.h b/Engine/Plugins/Runtime/HairStrands/Source/HairStrandsCore/Public/GroomResources.h
--- a/Engine/Plugins/Runtime/HairStrands/Source/HairStrandsCore/Public/GroomResources.h
+++ b/Engine/Plugins/Runtime/HairStrands/Source/HairStrandsCore/Public/GroomResources.h
@@ -212,6 +212,9 @@ struct FHairStrandsRestResource : public FHairCommonResource
 	inline uint32 GetSourceCurveCount() const { return BulkData.GetNumSourceCurves(); }
 
 	uint32 CachedTangentPointCount = 0;
+
+	/** List of external buffers */
+	TMap<FString,FRDGExternalBuffer> ExternalBuffers;
 };
 
 struct FHairStrandsDeformedResource : public FHairCommonResource
@@ -291,6 +294,9 @@ struct FHairStrandsDeformedResource : public FHairCommonResource
 	HAIRSTRANDSCORE_API FRDGExternalBuffer& GetDeformerBuffer(FRDGBuilder& GraphBuilder);
 	HAIRSTRANDSCORE_API FRDGExternalBuffer& GetDeformerCurveAttributeBuffer(FRDGBuilder& GraphBuilder);
 	HAIRSTRANDSCORE_API FRDGExternalBuffer& GetDeformerPointAttributeBuffer(FRDGBuilder& GraphBuilder);
+
+	/** List of external buffers */
+	TMap<FString,FRDGExternalBuffer> ExternalBuffers;
 };
 
 
diff --git a/Engine/Plugins/Runtime/HairStrands/Source/HairStrandsDeformer/Private/DeformerGroomAttributeRead.cpp b/Engine/Plugins/Runtime/HairStrands/Source/HairStrandsDeformer/Private/DeformerGroomAttributeRead.cpp
--- a/Engine/Plugins/Runtime/HairStrands/Source/HairStrandsDeformer/Private/DeformerGroomAttributeRead.cpp
+++ b/Engine/Plugins/Runtime/HairStrands/Source/HairStrandsDeformer/Private/DeformerGroomAttributeRead.cpp
@@ -277,7 +277,7 @@ FOptimusGroomAttributeReadProviderProxy::FOptimusGroomAttributeReadProviderProxy
 	TArray<const UGroomComponent*> GroomComponents;
 	UE::Groom::Private::GatherGroomComponents(MeshComponent, GroomComponents);
 	UE::Groom::Private::GroomComponentsToInstances(GroomComponents, GroupInstances);
-	UE::Groom::Private::GetGroomInvocationElementGroups(GroomComponents,UOptimusGroomExecDataInterface::GetExecutionDomainName(GroomAttributeGroup), GroupElements);
+	UE::Groom::Private::GetGroomInvocationElementGroups(GroomComponents,UOptimusGroomExecDataInterface::GetExecutionDomainName(GroomAttributeGroup), GroupElements, INDEX_NONE, true);
 }
 
 bool FOptimusGroomAttributeReadProviderProxy::IsValid(FValidationData const& InValidationData) const
@@ -500,98 +500,124 @@ FORCEINLINE int32 NumAttributeValues<FVector2f>()
 }
 
 template<typename DataType>
-FORCEINLINE void CreateInternalBuffers(const UGroomAsset* GroomAsset, const UE::Groom::Private::FGroupElements& GroupElements, const FName& AttributeName,
-	const FName& GroupName, FRDGBuilder& GraphBuilder, TArray<FRDGBufferRef>& AttributeBuffers)
+static void CreateInternalBuffers(const UGroomAsset* GroomAsset, const UE::Groom::Private::FGroupElements& GroupElements, const FName& AttributeName,
+	const EOptimusGroomExecDomain& AttributeGroup, FRDGBuilder& GraphBuilder, TArray<FRDGBufferSRVRef>& AttributeResources)
 {
+	const FName GroupName = UOptimusGroomExecDataInterface::GetExecutionDomainName(AttributeGroup);
+	
 	if(GroomAsset && GroomAsset->GetDataflowSettings().GetRestCollection() &&
      		GroomAsset->GetDataflowSettings().GetRestCollection()->HasAttribute(AttributeName, GroupName))
     {
      	const TManagedArray<DataType>& AssetAttributeValues =
      			GroomAsset->GetDataflowSettings().GetRestCollection()->GetAttribute<DataType>(AttributeName,GroupName);
-     		
+
+		const bool bIsGuide = ((AttributeGroup == EOptimusGroomExecDomain::GuidesPoints) || (AttributeGroup == EOptimusGroomExecDomain::GuidesEdges) ||
+					(AttributeGroup == EOptimusGroomExecDomain::GuidesCurves) || (AttributeGroup == EOptimusGroomExecDomain::GuidesObjects));
+		const bool bIsStrand = ((AttributeGroup == EOptimusGroomExecDomain::ControlPoint) || (AttributeGroup == EOptimusGroomExecDomain::StrandsEdges) ||
+			   (AttributeGroup == EOptimusGroomExecDomain::Curve) || (AttributeGroup == EOptimusGroomExecDomain::StrandsObjects));
+
+		int32 InstanceIndex = 0;
      	for(int32 GroupIndex : GroupElements.GroupIndices)
      	{
-     		const uint32 NumAttributes = GroupElements.GroupOffsets[GroupIndex+1]-GroupElements.GroupOffsets[GroupIndex];
-     			
-     		TArray<float> GroupAttributeValues;
-     		GroupAttributeValues.SetNum(NumAttributes*NumAttributeValues<DataType>());
- 
-     		int32 DataIndex = 0;
-     		for(uint32 AttributeIndex = 0; AttributeIndex < NumAttributes; ++AttributeIndex)
+     		if(const FHairGroupInstance* GroupInstance = GroupElements.GroupInstances[InstanceIndex])
      		{
-     			ConvertAttributeData<DataType>(AssetAttributeValues[AttributeIndex+GroupElements.GroupOffsets[GroupIndex]], GroupAttributeValues, DataIndex);
+     			const FString ResourceType = bIsGuide ? FString("Guides.") : FString("Strands.");
+     			const FString BufferName = FString("Hair.Deformer.") + ResourceType + AttributeName.ToString();
+     			
+     			if(FHairStrandsRestResource* RestResource = bIsGuide ? GroupInstance->Guides.RestResource : bIsStrand ? GroupInstance->Strands.RestResource : nullptr)
+     			{
+     				FRDGExternalBuffer* ExternalBuffer = RestResource->ExternalBuffers.Find(BufferName);
+     				if(!ExternalBuffer)
+     				{
+     					const uint32 NumAttributes = GroupElements.GroupOffsets[GroupIndex+1]-GroupElements.GroupOffsets[GroupIndex];
+     				
+			        	TArray<float> GroupAttributeValues;
+			        	GroupAttributeValues.SetNum(NumAttributes*NumAttributeValues<DataType>());
+	 
+			        	int32 DataIndex = 0;
+			        	for(uint32 AttributeIndex = 0; AttributeIndex < NumAttributes; ++AttributeIndex)
+			        	{
+			        		ConvertAttributeData<DataType>(AssetAttributeValues[AttributeIndex+GroupElements.GroupOffsets[GroupIndex]], GroupAttributeValues, DataIndex);
+			        	}
+			        	FRDGBufferRef TransientBuffer = GraphBuilder.CreateBuffer(FRDGBufferDesc::CreateStructuredDesc(
+			        		sizeof(float), GroupAttributeValues.Num()), *BufferName);
+     					ExternalBuffer = &RestResource->ExternalBuffers.Add(BufferName);
+
+     					ExternalBuffer->Buffer = GraphBuilder.ConvertToExternalBuffer(TransientBuffer);
+     					
+			        	GraphBuilder.QueueBufferUpload(TransientBuffer, GroupAttributeValues.GetData(),
+			        		sizeof(float) * GroupAttributeValues.Num(), ERDGInitialDataFlags::None);
+			        }
+     				AttributeResources.Add(RegisterAsSRV(GraphBuilder, *ExternalBuffer));
+     			}
      		}
- 
-     		const FString BufferName = FString("Hair.Deformer.") + AttributeName.ToString();
-     		FRDGBufferRef& GroupAttributeBuffer = AttributeBuffers.Add_GetRef( GraphBuilder.CreateBuffer(FRDGBufferDesc::CreateStructuredDesc(sizeof(float),
-     		 GroupAttributeValues.Num()), *BufferName));
-     		GraphBuilder.QueueBufferUpload(GroupAttributeBuffer, GroupAttributeValues.GetData(), sizeof(float) * GroupAttributeValues.Num(), ERDGInitialDataFlags::None);
+     		++InstanceIndex;
      	}
     }
 }
 
-FORCEINLINE void AddGroupResources(UGroomAsset* GroomAsset, const UE::Groom::Private::FGroupElements& GroupElements, const EOptimusGroomAttributeTypes GroomAttributeType,
-	const FName& GroomAttributeName,  const FName& GroomAttributeGroup, FRDGBuilder& GraphBuilder, TArray<FRDGBufferRef>& GroomAttributeBuffers)
+static void AddGroupResources(UGroomAsset* GroomAsset, const UE::Groom::Private::FGroupElements& GroupElements, const EOptimusGroomAttributeTypes GroomAttributeType,
+	const FName& GroomAttributeName,  const EOptimusGroomExecDomain& GroomAttributeGroup, FRDGBuilder& GraphBuilder, TArray<FRDGBufferSRVRef>& GroomAttributeResources)
 {
 	if(GroomAttributeType == EOptimusGroomAttributeTypes::Bool)
 	{
 		CreateInternalBuffers<bool>(GroomAsset, GroupElements, GroomAttributeName,
-			GroomAttributeGroup, GraphBuilder, GroomAttributeBuffers);
+			GroomAttributeGroup, GraphBuilder, GroomAttributeResources);
 	}
 	else if(GroomAttributeType == EOptimusGroomAttributeTypes::Int || GroomAttributeType == EOptimusGroomAttributeTypes::Uint)
 	{
 		CreateInternalBuffers<int32>(GroomAsset, GroupElements, GroomAttributeName,
-			GroomAttributeGroup, GraphBuilder, GroomAttributeBuffers);
+			GroomAttributeGroup, GraphBuilder, GroomAttributeResources);
 	}
 	else if(GroomAttributeType == EOptimusGroomAttributeTypes::IntVector2)
 	{
 		CreateInternalBuffers<FIntVector2>(GroomAsset, GroupElements, GroomAttributeName,
-			GroomAttributeGroup, GraphBuilder, GroomAttributeBuffers);
+			GroomAttributeGroup, GraphBuilder, GroomAttributeResources);
 	}
 	else if(GroomAttributeType == EOptimusGroomAttributeTypes::IntVector3)
 	{
 		CreateInternalBuffers<FIntVector3>(GroomAsset, GroupElements, GroomAttributeName,
-			GroomAttributeGroup, GraphBuilder, GroomAttributeBuffers);
+			GroomAttributeGroup, GraphBuilder, GroomAttributeResources);
 	}
 	else if(GroomAttributeType == EOptimusGroomAttributeTypes::IntVector4)
 	{
 		CreateInternalBuffers<FIntVector4>(GroomAsset, GroupElements, GroomAttributeName,
-			GroomAttributeGroup, GraphBuilder, GroomAttributeBuffers);
+			GroomAttributeGroup, GraphBuilder, GroomAttributeResources);
 	}
 	else if(GroomAttributeType == EOptimusGroomAttributeTypes::Float)
 	{
 		CreateInternalBuffers<float>(GroomAsset, GroupElements, GroomAttributeName,
-			GroomAttributeGroup, GraphBuilder, GroomAttributeBuffers);
+			GroomAttributeGroup, GraphBuilder, GroomAttributeResources);
 	}
 	else if(GroomAttributeType == EOptimusGroomAttributeTypes::Vector2)
 	{
 		CreateInternalBuffers<FVector2f>(GroomAsset, GroupElements,  GroomAttributeName,
-			GroomAttributeGroup, GraphBuilder, GroomAttributeBuffers);
+			GroomAttributeGroup, GraphBuilder, GroomAttributeResources);
 	}
 	else if(GroomAttributeType == EOptimusGroomAttributeTypes::Vector3 || GroomAttributeType == EOptimusGroomAttributeTypes::Rotator)
 	{
 		CreateInternalBuffers<FVector3f>(GroomAsset, GroupElements, GroomAttributeName,
-			GroomAttributeGroup, GraphBuilder, GroomAttributeBuffers);
+			GroomAttributeGroup, GraphBuilder, GroomAttributeResources);
 	}
 	else if(GroomAttributeType == EOptimusGroomAttributeTypes::Vector4)
 	{
 		CreateInternalBuffers<FVector4f>(GroomAsset, GroupElements, GroomAttributeName,
-			GroomAttributeGroup, GraphBuilder, GroomAttributeBuffers);
+			GroomAttributeGroup, GraphBuilder, GroomAttributeResources);
 	}
 	else if(GroomAttributeType == EOptimusGroomAttributeTypes::LinearColor)
 	{
 		CreateInternalBuffers<FLinearColor>(GroomAsset, GroupElements, GroomAttributeName,
-			GroomAttributeGroup, GraphBuilder, GroomAttributeBuffers);
+			GroomAttributeGroup, GraphBuilder, GroomAttributeResources);
 	}
 	else if(GroomAttributeType == EOptimusGroomAttributeTypes::Quat)
 	{
 		CreateInternalBuffers<FQuat4f>(GroomAsset, GroupElements, GroomAttributeName,
-			GroomAttributeGroup, GraphBuilder, GroomAttributeBuffers);
+			GroomAttributeGroup, GraphBuilder, GroomAttributeResources);
 	}
 	else if(GroomAttributeType == EOptimusGroomAttributeTypes::Transform || GroomAttributeType == EOptimusGroomAttributeTypes::Matrix3x4)
 	{
 		CreateInternalBuffers<FTransform3f>(GroomAsset, GroupElements, GroomAttributeName,
-			GroomAttributeGroup, GraphBuilder, GroomAttributeBuffers);
+			GroomAttributeGroup, GraphBuilder, GroomAttributeResources);
 	}
 }
 
@@ -600,31 +626,20 @@ void FOptimusGroomAttributeReadProviderProxy::CreateInternalBuffers(FRDGBuilder&
 	for (const TPair<UGroomAsset*,UE::Groom::Private::FGroupElements>& AssetPair : GroupElements)
 	{
 		AddGroupResources(AssetPair.Key, AssetPair.Value, GroomAttributeType, GroomAttributeName,
-				UOptimusGroomExecDataInterface::GetExecutionDomainName(GroomAttributeGroup),GraphBuilder,
-				AttributeValuesBuffers);
+				GroomAttributeGroup,GraphBuilder,AttributeValuesResources);
 	}
 }
 
 void FOptimusGroomAttributeReadProviderProxy::AllocateResources(FRDGBuilder& GraphBuilder, FAllocationData const& InAllocationData)
 {
 	if (!FallbackStructuredSRV) { FallbackStructuredSRV  = GraphBuilder.CreateSRV(GSystemTextures.GetDefaultStructuredBuffer(GraphBuilder, 16u)); }
-	
-	CreateInternalBuffers(GraphBuilder);
-	
+
 	AttributeValuesResources.Reset();
-	for (const TPair<UGroomAsset*,UE::Groom::Private::FGroupElements>& AssetPair : GroupElements)
-	{
-		for(int32 GroupIndex = 0; GroupIndex < AttributeValuesBuffers.Num(); ++GroupIndex)
-		{
-			FRDGBufferSRVRef& AttributeValuesResource = AttributeValuesResources.AddDefaulted_GetRef();
-			AttributeValuesResource = GraphBuilder.CreateSRV(AttributeValuesBuffers[GroupIndex]);
-		}
-	}
+	CreateInternalBuffers(GraphBuilder);
 }
 
 void FOptimusGroomAttributeReadProviderProxy::GatherDispatchData(FDispatchData const& InDispatchData)
 {
-	check(InDispatchData.NumInvocations == GroupInstances.Num());
 	check(InDispatchData.NumInvocations == AttributeValuesResources.Num());
 	
 	const TStridedView<FParameters> ParameterArray = MakeStridedParameterView<FParameters>(InDispatchData);
@@ -632,10 +647,11 @@ void FOptimusGroomAttributeReadProviderProxy::GatherDispatchData(FDispatchData c
 	int32 InvocationIndex = 0;
 	for (const TPair<UGroomAsset*,UE::Groom::Private::FGroupElements>& AssetPair : GroupElements)
 	{
+		int32 InstanceIndex = 0;
 		for(int32 GroupIndex : AssetPair.Value.GroupIndices)
 		{
 			FParameters& Parameters = ParameterArray[InvocationIndex];
-			if (const FHairGroupInstance* GroupInstance = GroupInstances[InvocationIndex])
+			if (const FHairGroupInstance* GroupInstance = AssetPair.Value.GroupInstances[InstanceIndex])
 			{
 				Parameters.NumValues = AssetPair.Value.GroupOffsets[GroupIndex+1]-AssetPair.Value.GroupOffsets[GroupIndex];
 				Parameters.ValueTypedData = AttributeValuesResources[InvocationIndex];
@@ -645,7 +661,8 @@ void FOptimusGroomAttributeReadProviderProxy::GatherDispatchData(FDispatchData c
 				Parameters.NumValues = 0;
 				Parameters.ValueTypedData = FallbackStructuredSRV;
 			}
+			++InstanceIndex;
 			++InvocationIndex;
 		}
 	}
-}
\ No newline at end of file
+}
diff --git a/Engine/Plugins/Runtime/HairStrands/Source/HairStrandsDeformer/Private/DeformerGroomAttributeRead.h b/Engine/Plugins/Runtime/HairStrands/Source/HairStrandsDeformer/Private/DeformerGroomAttributeRead.h
--- a/Engine/Plugins/Runtime/HairStrands/Source/HairStrandsDeformer/Private/DeformerGroomAttributeRead.h
+++ b/Engine/Plugins/Runtime/HairStrands/Source/HairStrandsDeformer/Private/DeformerGroomAttributeRead.h
@@ -145,9 +145,6 @@ private:
 
 	/** Groom assets group indices the proxy instance is linked to */
 	TArray<TPair<UGroomAsset*, UE::Groom::Private::FGroupElements>> GroupElements;
-
-	/** Groom attribute buffers built from the deformer */
-	TArray<FRDGBufferRef> AttributeValuesBuffers;
 	
 	/** Attribute Resources used to dispatch CS on GPU */
 	TArray<FRDGBufferSRVRef> AttributeValuesResources;
diff --git a/Engine/Plugins/Runtime/HairStrands/Source/HairStrandsDeformer/Private/DeformerGroomInterfaceUtils.h b/Engine/Plugins/Runtime/HairStrands/Source/HairStrandsDeformer/Private/DeformerGroomInterfaceUtils.h
--- a/Engine/Plugins/Runtime/HairStrands/Source/HairStrandsDeformer/Private/DeformerGroomInterfaceUtils.h
+++ b/Engine/Plugins/Runtime/HairStrands/Source/HairStrandsDeformer/Private/DeformerGroomInterfaceUtils.h
@@ -1,4 +1,4 @@
-﻿// Copyright Epic Games, Inc. All Rights Reserved.
+// Copyright Epic Games, Inc. All Rights Reserved.
 
 #pragma once
 #include "GroomComponent.h"
@@ -62,11 +62,13 @@ namespace UE::Groom::Private
  			}
  		}
  	}
-	static void GroomComponentsToSkelmeshes(const TArray<const UGroomComponent*>& GroomComponents, TArray<const FSkeletalMeshObject*>& SkeletalMeshes, TArray<FMatrix44f>& SkeletalTransforms, TArray<TArray<FMatrix44f>>& BonesRefToLocals)
+	FORCEINLINE void GroomComponentsToSkelmeshes(const TArray<const UGroomComponent*>& GroomComponents, TArray<const FSkeletalMeshObject*>& SkeletalMeshes, TArray<FMatrix44f>& SkeletalTransforms,
+		TArray<TArray<FMatrix44f>>& BonesRefToLocals, TArray<const FHairGroupInstance*>& GroupInstances)
 	{
 		SkeletalMeshes.Reset();
 		SkeletalTransforms.Reset();
 		BonesRefToLocals.Reset();
+		GroupInstances.Reset();
 
 		TArray<FTransform> BonesTransforms;
 		TArray<FMatrix44f> RefToLocals;
@@ -115,6 +117,7 @@ namespace UE::Groom::Private
 									SkeletalTransforms.Add(FMatrix44f(BonesTransform.ToMatrixWithScale() * GroupTransform.ToInverseMatrixWithScale()));
 									BonesRefToLocals.Add(RefToLocals);
 									SkeletalMeshes.Add(SkelMesh->MeshObject);
+									GroupInstances.Add(GroomComponent->GetGroupInstance(GroupIndex));
 								}
 							}
 						}
@@ -148,12 +151,12 @@ namespace UE::Groom::Private
 		GroomComponentsToInstances(GroomComponents, GroupInstances);
 	}
 
-	FORCEINLINE void GatherGroupSkelmeshes(const UActorComponent* ActorComponent, TArray<const FSkeletalMeshObject*>& SkeletalMeshes, TArray<FMatrix44f>& SkeletalTransforms, TArray<TArray<FMatrix44f>>& BonesRefToLocals)
+	FORCEINLINE void GatherGroupSkelmeshes(const UActorComponent* ActorComponent, TArray<const FSkeletalMeshObject*>& SkeletalMeshes, TArray<FMatrix44f>& SkeletalTransforms, TArray<TArray<FMatrix44f>>& BonesRefToLocals, TArray<const FHairGroupInstance*>& GroupInstances)
 	{
 		TArray<const UGroomComponent*> GroomComponents;
 		GatherGroomComponents(ActorComponent, GroomComponents);
 
-		GroomComponentsToSkelmeshes(GroomComponents, SkeletalMeshes, SkeletalTransforms, BonesRefToLocals);
+		GroomComponentsToSkelmeshes(GroomComponents, SkeletalMeshes, SkeletalTransforms, BonesRefToLocals, GroupInstances);
 	}
 
 	template<typename DataType>
@@ -255,9 +258,9 @@ namespace UE::Groom::Private
 		return true;
 	}
 
-	FORCEINLINE void GetGroomGroupElementOffsets(const UGroomComponent* GroomComponent, const FName DomainName, TArray<int32>& GroupOffsets, const int32 GroupSize = INDEX_NONE)
+	FORCEINLINE void GetGroomGroupElementOffsets(const UGroomComponent* GroomComponent, const FName DomainName, TArray<int32>& GroupOffsets, const int32 GroupSize = INDEX_NONE, const bool bSourceElements = false)
 	{
-		if(GroomComponent)
+		if(GroomComponent) 
 		{
 			const uint32 NumGroups = GroomComponent->GetGroupCount();
 			const float GroupAlignment = static_cast<float>(GroupSize);
@@ -295,30 +298,32 @@ namespace UE::Groom::Private
 					{
 						if(DomainName == UOptimusGroomAssetComponentSource::FStrandsExecutionDomains::Curves)
 						{
-							NumElements += Instance->Strands.GetData().GetNumCurves();
+							NumElements += bSourceElements ? Instance->Strands.GetData().GetNumSourceCurves() : Instance->Strands.GetData().GetNumCurves();
 						}
 						else if (DomainName == UOptimusGroomAssetComponentSource::FStrandsExecutionDomains::Edges)
 						{
-							NumElements += Instance->Strands.GetData().GetNumPoints() - Instance->Strands.GetData().GetNumCurves();
+							NumElements +=  bSourceElements ? Instance->Strands.GetData().GetNumSourcePoints() - Instance->Strands.GetData().GetNumSourceCurves() :
+									Instance->Strands.GetData().GetNumPoints() - Instance->Strands.GetData().GetNumCurves();
 						}
 						else if (DomainName == UOptimusGroomAssetComponentSource::FStrandsExecutionDomains::Points)
 						{
-							NumElements += Instance->Strands.GetData().GetNumPoints();
+							NumElements += bSourceElements ? Instance->Strands.GetData().GetNumSourcePoints() : Instance->Strands.GetData().GetNumPoints();
 						}
 					}
 					if(Instance->Guides.IsValid())
 					{
 						if(DomainName == UOptimusGroomAssetComponentSource::FGuidesExecutionDomains::Curves)
 						{
-							NumElements += Instance->Guides.GetData().GetNumCurves();
+							NumElements += bSourceElements ? Instance->Guides.GetData().GetNumSourceCurves() : Instance->Guides.GetData().GetNumCurves();
 						}
 						else if (DomainName == UOptimusGroomAssetComponentSource::FGuidesExecutionDomains::Edges)
 						{
-							NumElements += Instance->Guides.GetData().GetNumPoints() - Instance->Guides.GetData().GetNumCurves();
+							NumElements += bSourceElements ? Instance->Guides.GetData().GetNumSourcePoints() - Instance->Guides.GetData().GetNumSourceCurves() :
+								Instance->Guides.GetData().GetNumPoints() - Instance->Guides.GetData().GetNumCurves();
 						}
 						else if (DomainName == UOptimusGroomAssetComponentSource::FGuidesExecutionDomains::Points)
 						{
-							NumElements += Instance->Guides.GetData().GetNumPoints();
+							NumElements += bSourceElements ? Instance->Guides.GetData().GetNumSourcePoints() : Instance->Guides.GetData().GetNumPoints();
 						}
 					}
 				}
@@ -329,7 +334,7 @@ namespace UE::Groom::Private
 		}
 	}
 
-	FORCEINLINE int32 GetGroomInvocationElementCounts(const TArray<const UGroomComponent*>& GroomComponents, const FName DomainName, TArray<int32>& InvocationCounts, int32 GroupSize = INDEX_NONE)
+	FORCEINLINE int32 GetGroomInvocationElementCounts(const TArray<const UGroomComponent*>& GroomComponents, const FName DomainName, TArray<int32>& InvocationCounts, int32 GroupSize = INDEX_NONE, const bool bSourceElements = false)
 	{
 		TArray<int32> GroupOffsets;
 		int32 TotalCount = 0;
@@ -338,7 +343,7 @@ namespace UE::Groom::Private
 		{
 			if(GroomComponent)
 			{
-				GetGroomGroupElementOffsets(GroomComponent, DomainName, GroupOffsets, GroupSize);
+				GetGroomGroupElementOffsets(GroomComponent, DomainName, GroupOffsets, GroupSize, bSourceElements);
 
 				const uint32 NumGroups = GroomComponent->GetGroupCount();
 				for (uint32 GroupIndex = 0; GroupIndex < NumGroups; ++GroupIndex)
@@ -358,9 +363,10 @@ namespace UE::Groom::Private
 	{
 		TArray<int32> GroupIndices;
 		TArray<int32> GroupOffsets;
+		TArray<const FHairGroupInstance*> GroupInstances;
 	};
 	
-	FORCEINLINE void GetGroomInvocationElementGroups(const TArray<const UGroomComponent*>& GroomComponents, const FName DomainName, TArray<TPair<UGroomAsset*, FGroupElements>>& InvocationGroups, int32 GroupSize = INDEX_NONE)
+	FORCEINLINE void GetGroomInvocationElementGroups(const TArray<const UGroomComponent*>& GroomComponents, const FName DomainName, TArray<TPair<UGroomAsset*, FGroupElements>>& InvocationGroups, int32 GroupSize = INDEX_NONE, const bool bSourceElements = false)
 	{
 		InvocationGroups.Reset();
 		for(const UGroomComponent* GroomComponent : GroomComponents)
@@ -368,7 +374,7 @@ namespace UE::Groom::Private
 			if(GroomComponent)
 			{
 				FGroupElements GroupData;
-				GetGroomGroupElementOffsets(GroomComponent, DomainName, GroupData.GroupOffsets, GroupSize);
+				GetGroomGroupElementOffsets(GroomComponent, DomainName, GroupData.GroupOffsets, GroupSize, bSourceElements);
 				
 				const uint32 NumGroups = GroomComponent->GetGroupCount();
 				for (uint32 GroupIndex = 0; GroupIndex < NumGroups; ++GroupIndex)
@@ -376,6 +382,7 @@ namespace UE::Groom::Private
 					if(HasDeformationEnabledOrHasMeshDeformer(GroomComponent, GroupIndex) && GroomComponent->GetGroupInstance(GroupIndex))
 					{
 						GroupData.GroupIndices.Add(GroupIndex);
+						GroupData.GroupInstances.Add(GroomComponent->GetGroupInstance(GroupIndex));
 					}
 				}
 				InvocationGroups.Add(MakeTuple(GroomComponent->GroomAsset, GroupData));
diff --git a/Engine/Plugins/Runtime/HairStrands/Source/HairStrandsDeformer/Private/DeformerGroomMeshesRead.cpp b/Engine/Plugins/Runtime/HairStrands/Source/HairStrandsDeformer/Private/DeformerGroomMeshesRead.cpp
--- a/Engine/Plugins/Runtime/HairStrands/Source/HairStrandsDeformer/Private/DeformerGroomMeshesRead.cpp
+++ b/Engine/Plugins/Runtime/HairStrands/Source/HairStrandsDeformer/Private/DeformerGroomMeshesRead.cpp
@@ -1,4 +1,4 @@
-﻿// Copyright Epic Games, Inc. All Rights Reserved.
+// Copyright Epic Games, Inc. All Rights Reserved.
 
 #include "DeformerGroomMeshesRead.h"
 
@@ -99,7 +99,7 @@ FComputeDataProviderRenderProxy* UOptimusGroomMeshesReadDataProvider::GetRenderP
 
 FOptimusGroomMeshesReadDataProviderProxy::FOptimusGroomMeshesReadDataProviderProxy(UMeshComponent* MeshComponent)
 {
-	UE::Groom::Private::GatherGroupSkelmeshes(MeshComponent, SkeletalMeshObjects, SkeletalMeshTransforms, BonesRefToLocals);
+	UE::Groom::Private::GatherGroupSkelmeshes(MeshComponent, SkeletalMeshObjects, SkeletalMeshTransforms, BonesRefToLocals, GroupInstances);
 }
 
 bool FOptimusGroomMeshesReadDataProviderProxy::IsValid(FValidationData const& InValidationData) const
@@ -116,31 +116,52 @@ bool FOptimusGroomMeshesReadDataProviderProxy::IsValid(FValidationData const& In
 	return true;
 }
 
-FORCEINLINE void CreateBonesBuffers(const FSkeletalMeshObject* SkeletalMeshObject, const FMatrix44f& SkeletalMeshTransform, const TArray<FMatrix44f>& BoneRefToLocals, FRDGBuilder& GraphBuilder,
-	TArray<FRDGBufferRef>& BonesMatricesBuffers, TArray<FRDGBufferSRVRef>& BoneMatricesResources)
+static void CreateInternalBuffers(const FSkeletalMeshObject* SkeletalMeshObject, const FMatrix44f& SkeletalMeshTransform, const TArray<FMatrix44f>& BoneRefToLocals,
+	const FHairGroupInstance* GroupInstance, FRDGBuilder& GraphBuilder, TArray<FRDGBufferSRVRef>& AttributeResources)
 {
-	if(SkeletalMeshObject)
+	if(SkeletalMeshObject && GroupInstance)
 	{
 		//const TArray<FMatrix44f>& BoneRefToLocals = SkeletalMeshObject->GetReferenceToLocalMatrices();
-		const int32 NumBones = BoneRefToLocals.Num();
-		
-		TArray<FVector4f> BoneVectors;
-		BoneVectors.Init(FVector4f::Zero(), NumBones * 3);
-
-		int32 VectorIndex = 0;
-		for(int32 BoneIndex = 0; BoneIndex < NumBones; ++BoneIndex)
+		const FString BufferName("Hair.Deformer.Strands.BoneTransformMatrices");
+		if(FHairStrandsDeformedResource* DeformedResource = GroupInstance->Strands.DeformedResource)
 		{
-			const FMatrix44f LocalMatrix = BoneRefToLocals[BoneIndex] * SkeletalMeshTransform;
-			LocalMatrix.To3x4MatrixTranspose(static_cast<float*>(&BoneVectors[VectorIndex].X));
-
-			VectorIndex += 3;
+			FRDGBufferRef TransientBuffer = nullptr;
+			FRDGExternalBuffer* ExternalBuffer = DeformedResource->ExternalBuffers.Find(BufferName);
+			if(ExternalBuffer)
+			{
+				if(!GraphBuilder.FindExternalBuffer(ExternalBuffer->Buffer))
+				{
+					TransientBuffer = GraphBuilder.RegisterExternalBuffer(ExternalBuffer->Buffer);
+				}
+			}
+			else
+			{
+				TransientBuffer = GraphBuilder.CreateBuffer(FRDGBufferDesc::CreateStructuredDesc(
+					sizeof(FVector4f), BoneRefToLocals.Num() * 3 ), *BufferName);
+				ExternalBuffer = &DeformedResource->ExternalBuffers.Add(BufferName);
+				
+				ExternalBuffer->Buffer = GraphBuilder.ConvertToExternalBuffer(TransientBuffer);
+			}
+			if(TransientBuffer)
+			{
+				const int32 NumBones = BoneRefToLocals.Num();
+	
+				TArray<FVector4f> BoneVectors;
+				BoneVectors.Init(FVector4f::Zero(), NumBones * 3);
+
+				int32 VectorIndex = 0;
+				for(int32 BoneIndex = 0; BoneIndex < NumBones; ++BoneIndex)
+				{
+					const FMatrix44f LocalMatrix = BoneRefToLocals[BoneIndex] * SkeletalMeshTransform;
+					LocalMatrix.To3x4MatrixTranspose(static_cast<float*>(&BoneVectors[VectorIndex].X));
+					 
+					VectorIndex += 3;
+				}
+				GraphBuilder.QueueBufferUpload(TransientBuffer, BoneVectors.GetData(),
+					sizeof(FVector4f) * BoneVectors.Num(), ERDGInitialDataFlags::None);
+			}
+			AttributeResources.Add(RegisterAsSRV(GraphBuilder,*ExternalBuffer));
 		}
-
-		FRDGBufferRef& BonesBuffer = BonesMatricesBuffers.Add_GetRef( GraphBuilder.CreateBuffer(
-			FRDGBufferDesc::CreateStructuredDesc(sizeof(FVector4f), BoneVectors.Num()), TEXT("Hair.Deformer.BoneTransformMatrices")));
-		GraphBuilder.QueueBufferUpload(BonesBuffer, BoneVectors.GetData(), sizeof(FVector4f) * BoneVectors.Num(), ERDGInitialDataFlags::None);
-
-		BoneMatricesResources.Add(GraphBuilder.CreateSRV(BonesBuffer));
 	}
 }
 
@@ -149,12 +170,11 @@ void FOptimusGroomMeshesReadDataProviderProxy::AllocateResources(FRDGBuilder& Gr
 	if (!FallbackStructuredSRV) { FallbackStructuredSRV  = GraphBuilder.CreateSRV(GSystemTextures.GetDefaultStructuredBuffer(GraphBuilder, 16u)); }
 
 	BoneMatricesResources.Reset();
-	BoneMatricesBuffers.Reset();
 
 	int32 SkelIndex = 0;
 	for (const FSkeletalMeshObject* SkeletalMeshObject : SkeletalMeshObjects)
 	{
-		CreateBonesBuffers(SkeletalMeshObject, SkeletalMeshTransforms[SkelIndex], BonesRefToLocals[SkelIndex], GraphBuilder, BoneMatricesBuffers, BoneMatricesResources);
+		CreateInternalBuffers(SkeletalMeshObject, SkeletalMeshTransforms[SkelIndex], BonesRefToLocals[SkelIndex], GroupInstances[SkelIndex], GraphBuilder, BoneMatricesResources);
 		++SkelIndex;
 	}
 }
@@ -178,4 +198,4 @@ void FOptimusGroomMeshesReadDataProviderProxy::GatherDispatchData(FDispatchData
 			Parameters.BoneTransformMatrices = FallbackStructuredSRV;
 		}
 	}
-}
\ No newline at end of file
+}
diff --git a/Engine/Plugins/Runtime/HairStrands/Source/HairStrandsDeformer/Private/DeformerGroomMeshesRead.h b/Engine/Plugins/Runtime/HairStrands/Source/HairStrandsDeformer/Private/DeformerGroomMeshesRead.h
--- a/Engine/Plugins/Runtime/HairStrands/Source/HairStrandsDeformer/Private/DeformerGroomMeshesRead.h
+++ b/Engine/Plugins/Runtime/HairStrands/Source/HairStrandsDeformer/Private/DeformerGroomMeshesRead.h
@@ -1,4 +1,4 @@
-﻿// Copyright Epic Games, Inc. All Rights Reserved.
+// Copyright Epic Games, Inc. All Rights Reserved.
 
 #pragma once
 
@@ -84,9 +84,11 @@ private:
 
 	/** Bones refs to locals matrices */
 	TArray<TArray<FMatrix44f>> BonesRefToLocals;
+
+	/** List of instances (invocations) used in that data interface */
+	TArray<const FHairGroupInstance*> GroupInstances;
 	
 	/** Bones Resources used to dispatch CS on GPU */
-	TArray<FRDGBufferRef> BoneMatricesBuffers;
 	TArray<FRDGBufferSRVRef> BoneMatricesResources;
 
 	/** Fallback resources */
