Reduce the number of stored attributes in groom dataflow

diff --git a/Engine/Plugins/Runtime/HairStrands/Source/HairStrandsDataflow/Private/GroomAssetTerminalNode.cpp b/Engine/Plugins/Runtime/HairStrands/Source/HairStrandsDataflow/Private/GroomAssetTerminalNode.cpp
--- a/Engine/Plugins/Runtime/HairStrands/Source/HairStrandsDataflow/Private/GroomAssetTerminalNode.cpp
+++ b/Engine/Plugins/Runtime/HairStrands/Source/HairStrandsDataflow/Private/GroomAssetTerminalNode.cpp
@@ -92,61 +92,63 @@ namespace UE::Groom::Private
 	}
 
 	template<typename FacadeType>
-	static void TransferVerticesAttributes(const FManagedArrayCollection& InCollection, FManagedArrayCollection* OutCollection, const int32 NumPoints)
+	static void TransferVerticesAttributes(const FManagedArrayCollection& InCollection, FManagedArrayCollection* OutCollection,
+		const int32 NumPoints, const TArray<FName>& AttributesToSkip)
 	{
 		// Transfer vertices weight maps onto the points to be stored onto the rest collection
 		const TArray<FName> AttributeNames = InCollection.AttributeNames(FacadeType::VerticesGroup);
 		for(const FName& AttributeName : AttributeNames)
 		{
-			
-			if(InCollection.GetAttributeType(AttributeName, FacadeType::VerticesGroup) == EManagedArrayType::FFloatType)
-			{
-				BuildVerticesAttribute<FacadeType, float>(InCollection, OutCollection, NumPoints, AttributeName);
-			}
-			else if(InCollection.GetAttributeType(AttributeName, FacadeType::VerticesGroup) == EManagedArrayType::FVector4fType)
-			{
-				BuildVerticesAttribute<FacadeType, FVector4f>(InCollection, OutCollection, NumPoints, AttributeName);
-			}
-			else if(InCollection.GetAttributeType(AttributeName, FacadeType::VerticesGroup) == EManagedArrayType::FVectorType)
-			{
-				BuildVerticesAttribute<FacadeType, FVector3f>(InCollection, OutCollection, NumPoints, AttributeName);
-			}
-			else if(InCollection.GetAttributeType(AttributeName, FacadeType::VerticesGroup) == EManagedArrayType::FVector2DType)
-			{
-				BuildVerticesAttribute<FacadeType, FVector2f>(InCollection, OutCollection, NumPoints, AttributeName);
-			}
-			if(InCollection.GetAttributeType(AttributeName, FacadeType::VerticesGroup) == EManagedArrayType::FInt32Type)
-			{
-				BuildVerticesAttribute<FacadeType, int32>(InCollection, OutCollection, NumPoints, AttributeName);
-			}
-			else if(InCollection.GetAttributeType(AttributeName, FacadeType::VerticesGroup) == EManagedArrayType::FIntVector4Type)
-			{
-				BuildVerticesAttribute<FacadeType, FIntVector4>(InCollection, OutCollection, NumPoints, AttributeName);
-			}
-			else if(InCollection.GetAttributeType(AttributeName, FacadeType::VerticesGroup) == EManagedArrayType::FIntVectorType)
+			if(!AttributesToSkip.Contains(AttributeName))
 			{
-				BuildVerticesAttribute<FacadeType, FIntVector3>(InCollection, OutCollection, NumPoints, AttributeName);
-			}
-			else if(InCollection.GetAttributeType(AttributeName, FacadeType::VerticesGroup) == EManagedArrayType::FIntVector2Type)
-			{
-				BuildVerticesAttribute<FacadeType, FIntVector2>(InCollection, OutCollection, NumPoints, AttributeName);
-			}
-			else if(InCollection.GetAttributeType(AttributeName, FacadeType::VerticesGroup) == EManagedArrayType::FBoolType)
-			{
-				BuildVerticesAttribute<FacadeType, bool>(InCollection, OutCollection, NumPoints, AttributeName);
-			}
-			else if(InCollection.GetAttributeType(AttributeName, FacadeType::VerticesGroup) == EManagedArrayType::FLinearColorType)
-			{
-				BuildVerticesAttribute<FacadeType, FLinearColor>(InCollection, OutCollection, NumPoints, AttributeName);
-			}
-			else if(InCollection.GetAttributeType(AttributeName, FacadeType::VerticesGroup) == EManagedArrayType::FQuatType)
-			{
-				BuildVerticesAttribute<FacadeType, FQuat4f>(InCollection, OutCollection, NumPoints, AttributeName);
-			}
-			else if(InCollection.GetAttributeType(AttributeName, FacadeType::VerticesGroup) == EManagedArrayType::FTransform3fType)
-			{
-				BuildVerticesAttribute<FacadeType, FTransform3f>(InCollection, OutCollection, NumPoints, AttributeName); 
-
+				if(InCollection.GetAttributeType(AttributeName, FacadeType::VerticesGroup) == EManagedArrayType::FFloatType)
+				{
+					BuildVerticesAttribute<FacadeType, float>(InCollection, OutCollection, NumPoints, AttributeName);
+				}
+				else if(InCollection.GetAttributeType(AttributeName, FacadeType::VerticesGroup) == EManagedArrayType::FVector4fType)
+				{
+					BuildVerticesAttribute<FacadeType, FVector4f>(InCollection, OutCollection, NumPoints, AttributeName);
+				}
+				else if(InCollection.GetAttributeType(AttributeName, FacadeType::VerticesGroup) == EManagedArrayType::FVectorType)
+				{
+					BuildVerticesAttribute<FacadeType, FVector3f>(InCollection, OutCollection, NumPoints, AttributeName);
+				}
+				else if(InCollection.GetAttributeType(AttributeName, FacadeType::VerticesGroup) == EManagedArrayType::FVector2DType)
+				{
+					BuildVerticesAttribute<FacadeType, FVector2f>(InCollection, OutCollection, NumPoints, AttributeName);
+				}
+				if(InCollection.GetAttributeType(AttributeName, FacadeType::VerticesGroup) == EManagedArrayType::FInt32Type)
+				{
+					BuildVerticesAttribute<FacadeType, int32>(InCollection, OutCollection, NumPoints, AttributeName);
+				}
+				else if(InCollection.GetAttributeType(AttributeName, FacadeType::VerticesGroup) == EManagedArrayType::FIntVector4Type)
+				{
+					BuildVerticesAttribute<FacadeType, FIntVector4>(InCollection, OutCollection, NumPoints, AttributeName);
+				}
+				else if(InCollection.GetAttributeType(AttributeName, FacadeType::VerticesGroup) == EManagedArrayType::FIntVectorType)
+				{
+					BuildVerticesAttribute<FacadeType, FIntVector3>(InCollection, OutCollection, NumPoints, AttributeName);
+				}
+				else if(InCollection.GetAttributeType(AttributeName, FacadeType::VerticesGroup) == EManagedArrayType::FIntVector2Type)
+				{
+					BuildVerticesAttribute<FacadeType, FIntVector2>(InCollection, OutCollection, NumPoints, AttributeName);
+				}
+				else if(InCollection.GetAttributeType(AttributeName, FacadeType::VerticesGroup) == EManagedArrayType::FBoolType)
+				{
+					BuildVerticesAttribute<FacadeType, bool>(InCollection, OutCollection, NumPoints, AttributeName);
+				}
+				else if(InCollection.GetAttributeType(AttributeName, FacadeType::VerticesGroup) == EManagedArrayType::FLinearColorType)
+				{
+					BuildVerticesAttribute<FacadeType, FLinearColor>(InCollection, OutCollection, NumPoints, AttributeName);
+				}
+				else if(InCollection.GetAttributeType(AttributeName, FacadeType::VerticesGroup) == EManagedArrayType::FQuatType)
+				{
+					BuildVerticesAttribute<FacadeType, FQuat4f>(InCollection, OutCollection, NumPoints, AttributeName);
+				}
+				else if(InCollection.GetAttributeType(AttributeName, FacadeType::VerticesGroup) == EManagedArrayType::FTransform3fType)
+				{
+					BuildVerticesAttribute<FacadeType, FTransform3f>(InCollection, OutCollection, NumPoints, AttributeName); 
+				}
 			}
 		}
 	}
@@ -180,23 +182,27 @@ void FGroomAssetTerminalDataflowNode::SetAssetValue(TObjectPtr<UObject> Asset, U
 
 		if(StrandsFacade.IsValid())
 		{
+			// Skip default vertex attributes as not defined by the user
+			const TArray<FName> AttributesToSkip = {UE::Groom::FGroomStrandsFacade::VertexLinearColorsAttribute};
+			
 			// Transfer vertices weight maps onto the points to be stored onto the rest collection
-			UE::Groom::Private::TransferVerticesAttributes<UE::Groom::FGroomStrandsFacade>(InCollection, OutCollection, StrandsFacade.GetNumPoints());
+			UE::Groom::Private::TransferVerticesAttributes<UE::Groom::FGroomStrandsFacade>(InCollection, OutCollection, StrandsFacade.GetNumPoints(), AttributesToSkip);
 		}
 		if(GuidesFacade.IsValid())
 		{
 			// Only copy the attributes not stored already in the asset
 			const TArray<TTuple<FName, FName>> AttributesToCopy = {
-				{UE::Groom::FGroomGuidesFacade::PointsGroup, UE::Groom::FGroomGuidesFacade::PointBoneIndicesAttribute},
-				{UE::Groom::FGroomGuidesFacade::PointsGroup, UE::Groom::FGroomGuidesFacade::PointBoneWeightsAttribute},
 				{UE::Groom::FGroomGuidesFacade::CurvesGroup, UE::Groom::FGroomGuidesFacade::CurveParentIndicesAttribute},
 				{UE::Groom::FGroomGuidesFacade::CurvesGroup, UE::Groom::FGroomGuidesFacade::CurveLodIndicesAttribute}};
 
 			// Copy attributes from input collection
 			UE::Groom::Private::CopyCollectionAttributes(&InCollection, OutCollection, AttributesToCopy);
+
+			// Skip default vertex attributes as not defined by the user
+			const TArray<FName> AttributesToSkip = {UE::Groom::FGroomGuidesFacade::VertexLinearColorsAttribute};
 			
 			// Transfer vertices weight maps onto the points to be stored onto the rest collection (VertexKinematicWeights should be included here)
-			UE::Groom::Private::TransferVerticesAttributes<UE::Groom::FGroomGuidesFacade>(InCollection, OutCollection, GuidesFacade.GetNumPoints());
+			UE::Groom::Private::TransferVerticesAttributes<UE::Groom::FGroomGuidesFacade>(InCollection, OutCollection, GuidesFacade.GetNumPoints(), AttributesToSkip);
 
 			// Build an editable groom asset for the strands
 			FEditableGroom EditGroom;
