Fix use after free issue when rebuilding groom data

diff --git a/Engine/Plugins/Runtime/HairStrands/Source/HairStrandsCore/Private/GroomAsset.cpp b/Engine/Plugins/Runtime/HairStrands/Source/HairStrandsCore/Private/GroomAsset.cpp
--- a/Engine/Plugins/Runtime/HairStrands/Source/HairStrandsCore/Private/GroomAsset.cpp
+++ b/Engine/Plugins/Runtime/HairStrands/Source/HairStrandsCore/Private/GroomAsset.cpp
@@ -87,6 +87,9 @@ static TAutoConsoleVariable<int32> GHairStrandsWarningLogVerbosity(
 static int32 GHairStrandsDDCLogEnable = 0;
 static FAutoConsoleVariableRef CVarHairStrandsDDCLogEnable(TEXT("r.HairStrands.DDCLog"), GHairStrandsDDCLogEnable, TEXT("Enable DDC logging for groom assets and groom binding assets"));
 
+static int32 GHairStrandsReleaseResourceBeforeBuildingData = 1;
+static FAutoConsoleVariableRef CVarHairStrandsReleaseResourceBeforeBuildingData(TEXT("r.HairStrands.ReleaseResourceBeforeBuildingData"), GHairStrandsReleaseResourceBeforeBuildingData, TEXT("Experimental settings for releasing resources prior to building data. This avoids issue when reloading groom."));
+
 static EGroomLODMode InternalGetHairStrandsLODMode()
 {
 	return GetHairStrandsLODMode() > 0 ? EGroomLODMode::Auto : EGroomLODMode::Manual;
@@ -2507,6 +2510,15 @@ bool UGroomAsset::CacheDerivedDatas()
 	const bool bIsGameThread = IsInGameThread();
 	FGroomComponentRecreateRenderStateContext RecreateContext(bIsGameThread ? this : nullptr);
 
+	// (Temporary) Workaround when reloading groom asset several time, which can cause resources to be still alive while the data are being updated 
+	// by CacheDerivedData. The BulkData referenced by the resources becomes no longer valid, causing use-after-free crash.
+	const bool bReleaseResourceBeforeCacheDerivedData = bIsGameThread && GHairStrandsReleaseResourceBeforeBuildingData > 1;
+	if (bReleaseResourceBeforeCacheDerivedData)
+	{
+		ReleaseResource();
+		FlushRenderingCommands();
+	}
+
 	const uint32 GroupCount = GetHairGroupsInterpolation().Num();
 	bool bSucceed = true;
 	for (uint32 GroupIndex = 0; GroupIndex < GroupCount; ++GroupIndex)
@@ -2521,7 +2533,10 @@ bool UGroomAsset::CacheDerivedDatas()
 	
 		if (bIsGameThread)
 		{
-			ReleaseResource();
+			if (!bReleaseResourceBeforeCacheDerivedData)
+			{
+				ReleaseResource();
+			}
 			InitResources();
 		}
 	}
