Add debug visualization of position offset for groom.

This allows to debug quantization issue which might happen when skel. mesh are far from its local origin.

diff --git a/Engine/Plugins/Runtime/HairStrands/Source/HairStrandsCore/Private/GroomDebug.cpp b/Engine/Plugins/Runtime/HairStrands/Source/HairStrandsCore/Private/GroomDebug.cpp
--- a/Engine/Plugins/Runtime/HairStrands/Source/HairStrandsCore/Private/GroomDebug.cpp
+++ b/Engine/Plugins/Runtime/HairStrands/Source/HairStrandsCore/Private/GroomDebug.cpp
@@ -101,6 +101,8 @@ class FHairProjectionHairDebugCS : public FGlobalShader
 		SHADER_PARAMETER(uint32, Sim_HairDeformedFrames)
 		SHADER_PARAMETER(uint32, Sim_HairDeformedSamples)
 
+		SHADER_PARAMETER_RDG_BUFFER_SRV(StructuredBuffer<float4>, RestPositionOffset)
+		SHADER_PARAMETER_RDG_BUFFER_SRV(StructuredBuffer<float4>, DeformedPositionOffset)
 		SHADER_PARAMETER_RDG_BUFFER_SRV(Buffer, RestPositionBuffer)
 		SHADER_PARAMETER_RDG_BUFFER_SRV(Buffer, DeformedPositionBuffer)
 		SHADER_PARAMETER_RDG_BUFFER_SRV(StructuredBuffer, RestSamplePositionsBuffer)
@@ -144,6 +146,8 @@ static void AddDebugProjectionHairPass(
 	FIntRect Viewport,
 	const TUniformBufferRef<FViewUniformShaderParameters>& ViewUniformBuffer,
 	const int32 MeshLODIndex,
+	const FHairStrandsRestResource* RestResources,
+	const FHairStrandsDeformedResource* DeformedResources,
 	const FHairStrandsRestRootResource* RestRootResources,
 	const FHairStrandsDeformedRootResource* DeformedRootResources,
 	const FTransform& LocalToWorld,
@@ -235,6 +239,9 @@ static void AddDebugProjectionHairPass(
 	}
 
 	Parameters->RootToUniqueTriangleIndexBuffer = RegisterAsSRV(GraphBuilder, RestLODDatas.RootToUniqueTriangleIndexBuffer);
+	
+	Parameters->RestPositionOffset = RegisterAsSRV(GraphBuilder, RestResources->PositionOffsetBuffer);	
+	Parameters->DeformedPositionOffset = RegisterAsSRV(GraphBuilder, DeformedResources->GetPositionOffsetBuffer(FHairStrandsDeformedResource::EFrameType::Current));
 
 	Parameters->RestPositionBuffer = RegisterAsSRV(GraphBuilder, RestLODDatas.RestUniqueTrianglePositionBuffer);
 	Parameters->DeformedPositionBuffer = RegisterAsSRV(GraphBuilder, DeformedLODDatas.GetDeformedUniqueTrianglePositionBuffer(FHairStrandsLODDeformedRootResource::Current));
@@ -1267,15 +1274,24 @@ void RunHairStrandsDebug(
 					if (!Instance->HairGroupPublicData || Instance->BindingType != EHairBindingType::Skinning)
 						continue;
 
+					FHairStrandsRestResource* RestResource = nullptr;
+					FHairStrandsDeformedResource* DeformedResource = nullptr;
+
 					FHairStrandsRestRootResource* RestRootResource = nullptr;
 					FHairStrandsDeformedRootResource* DeformedRootResource = nullptr;
 					if (bSim)
 					{
+						RestResource 	 = Instance->Guides.RestResource;
+						DeformedResource = Instance->Guides.DeformedResource;
+
 						RestRootResource 	 = Instance->Guides.RestRootResource;
 						DeformedRootResource = Instance->Guides.DeformedRootResource;
 					}
 					else if (Instance->GetHairGeometry() == EHairGeometryType::Strands)
 					{
+						RestResource 	 = Instance->Strands.RestResource;
+						DeformedResource = Instance->Strands.DeformedResource;
+
 						RestRootResource 	 = Instance->Strands.RestRootResource;
 						DeformedRootResource = Instance->Strands.DeformedRootResource;
 					}
@@ -1284,19 +1300,22 @@ void RunHairStrandsDebug(
 						const int32 LODIndex = Instance->HairGroupPublicData->GetIntLODIndex();
 						if (Instance->Cards.IsValid(LODIndex))
 						{
+							RestResource 	 = Instance->Cards.LODs[LODIndex].Guides.RestResource;
+							DeformedResource = Instance->Cards.LODs[LODIndex].Guides.DeformedResource;
+
 							RestRootResource 	 = Instance->Cards.LODs[LODIndex].Guides.RestRootResource;
 							DeformedRootResource = Instance->Cards.LODs[LODIndex].Guides.DeformedRootResource;
 						}
 					}
-					if (RestRootResource == nullptr || DeformedRootResource == nullptr)
+					if (RestRootResource == nullptr || DeformedRootResource == nullptr || RestResource == nullptr || DeformedResource == nullptr)
 						continue;
 
 					const int32 MeshLODIndex = Instance->Debug.MeshLODIndex;
 
-					AddDebugProjectionHairPass(GraphBuilder, LocalView, ShaderMap, ShaderPrintData, Viewport, ViewUniformBuffer, MeshLODIndex, RestRootResource, DeformedRootResource, Instance->HairGroupPublicData->VFInput.LocalToWorldTransform, DrawIndex++, bSim, EDebugProjectionHairType::HairRoots,    HairStrandsTriangleType::RestPose);
-					AddDebugProjectionHairPass(GraphBuilder, LocalView, ShaderMap, ShaderPrintData, Viewport, ViewUniformBuffer, MeshLODIndex, RestRootResource, DeformedRootResource, Instance->HairGroupPublicData->VFInput.LocalToWorldTransform, DrawIndex++, bSim, EDebugProjectionHairType::HairSamples,  HairStrandsTriangleType::RestPose);
-					AddDebugProjectionHairPass(GraphBuilder, LocalView, ShaderMap, ShaderPrintData, Viewport, ViewUniformBuffer, MeshLODIndex, RestRootResource, DeformedRootResource, Instance->HairGroupPublicData->VFInput.LocalToWorldTransform, DrawIndex++, bSim, EDebugProjectionHairType::HairRoots,    HairStrandsTriangleType::DeformedPose);
-					AddDebugProjectionHairPass(GraphBuilder, LocalView, ShaderMap, ShaderPrintData, Viewport, ViewUniformBuffer, MeshLODIndex, RestRootResource, DeformedRootResource, Instance->HairGroupPublicData->VFInput.LocalToWorldTransform, DrawIndex++, bSim, EDebugProjectionHairType::HairSamples,  HairStrandsTriangleType::DeformedPose);
+					AddDebugProjectionHairPass(GraphBuilder, LocalView, ShaderMap, ShaderPrintData, Viewport, ViewUniformBuffer, MeshLODIndex, RestResource, DeformedResource, RestRootResource, DeformedRootResource, Instance->HairGroupPublicData->VFInput.LocalToWorldTransform, DrawIndex++, bSim, EDebugProjectionHairType::HairRoots,    HairStrandsTriangleType::RestPose);
+					AddDebugProjectionHairPass(GraphBuilder, LocalView, ShaderMap, ShaderPrintData, Viewport, ViewUniformBuffer, MeshLODIndex, RestResource, DeformedResource, RestRootResource, DeformedRootResource, Instance->HairGroupPublicData->VFInput.LocalToWorldTransform, DrawIndex++, bSim, EDebugProjectionHairType::HairSamples,  HairStrandsTriangleType::RestPose);
+					AddDebugProjectionHairPass(GraphBuilder, LocalView, ShaderMap, ShaderPrintData, Viewport, ViewUniformBuffer, MeshLODIndex, RestResource, DeformedResource, RestRootResource, DeformedRootResource, Instance->HairGroupPublicData->VFInput.LocalToWorldTransform, DrawIndex++, bSim, EDebugProjectionHairType::HairRoots,    HairStrandsTriangleType::DeformedPose);
+					AddDebugProjectionHairPass(GraphBuilder, LocalView, ShaderMap, ShaderPrintData, Viewport, ViewUniformBuffer, MeshLODIndex, RestResource, DeformedResource, RestRootResource, DeformedRootResource, Instance->HairGroupPublicData->VFInput.LocalToWorldTransform, DrawIndex++, bSim, EDebugProjectionHairType::HairSamples,  HairStrandsTriangleType::DeformedPose);
 				}
 			};
 
diff --git a/Engine/Plugins/Runtime/HairStrands/Source/HairStrandsCore/Public/GroomResources.h b/Engine/Plugins/Runtime/HairStrands/Source/HairStrandsCore/Public/GroomResources.h
--- a/Engine/Plugins/Runtime/HairStrands/Source/HairStrandsCore/Public/GroomResources.h
+++ b/Engine/Plugins/Runtime/HairStrands/Source/HairStrandsCore/Public/GroomResources.h
@@ -282,6 +282,7 @@ struct FHairStrandsDeformedResource : public FHairCommonResource
 	inline uint32 GetIndex(EFrameType T) const					{ return T == EFrameType::Current ? CurrentIndex : 1u - CurrentIndex; }
 	inline FRDGExternalBuffer& GetBuffer(EFrameType T)			{ return DeformedPositionBuffer[GetIndex(T)];  }
 	inline FVector& GetPositionOffset(EFrameType T)				{ return PositionOffset[GetIndex(T)]; }
+	inline const FRDGExternalBuffer& GetPositionOffsetBuffer(EFrameType T) const { return DeformedOffsetBuffer[GetIndex(T)]; }
 	inline FRDGExternalBuffer& GetPositionOffsetBuffer(EFrameType T) { return DeformedOffsetBuffer[GetIndex(T)]; }
 	inline const FVector& GetPositionOffset(EFrameType T) const { return PositionOffset[GetIndex(T)]; }
 	inline void SetPositionOffset(EFrameType T, const FVector& Offset)  { PositionOffset[GetIndex(T)] = Offset; }
diff --git a/Engine/Shaders/Private/HairStrands/HairStrandsDebug.usf b/Engine/Shaders/Private/HairStrands/HairStrandsDebug.usf
--- a/Engine/Shaders/Private/HairStrands/HairStrandsDebug.usf
+++ b/Engine/Shaders/Private/HairStrands/HairStrandsDebug.usf
@@ -387,6 +387,9 @@ StructuredBuffer<float4> RestSamplePositionsBuffer;
 StructuredBuffer<float4> DeformedSamplePositionsBuffer;
 #endif
 
+StructuredBuffer<float4> RestPositionOffset;
+StructuredBuffer<float4> DeformedPositionOffset;
+
 void Move(inout FShaderPrintContext Ctx, uint2 InPos)
 {
 	if (Ctx.bIsActive)
@@ -589,6 +592,35 @@ void MainCS(uint3 DispatchThreadId : SV_DispatchThreadID)
 		// Settings
 		NormalScale = AddSlider(Ctx, TEXT("Size"), 0.02f, FontWhite) * 10.f; 
 
+		// Position Offsets
+		{
+			bool bRest   = false;
+			bool bDeform = false;
+
+			Move(Ctx, StartPos + OffsetPos*3);
+			Print(Ctx, TEXT("Position Offsets"), FontYellow); Newline(Ctx);
+			{
+				bRest     = AddCheckbox(Ctx, TEXT("Offset Rest    "), false, FontWhite);  Print(Ctx, TEXT("||"), FontCyan); Newline(Ctx);
+				bDeform   = AddCheckbox(Ctx, TEXT("Offset Deformed"), false, FontWhite);  Print(Ctx, TEXT("||"), FontMagenta); Newline(Ctx);
+				Newline(Ctx);	
+			}
+
+			if (DispatchThreadId.x == 0)
+			{
+				if (bRest)
+				{
+					const float3 PositionOffset = mul(RestPositionOffset[0], RootLocalToWorld).xyz;
+					AddCrossWS(PositionOffset, NormalScale, ColorCyan);
+				}
+
+				if (bDeform)
+				{
+					const float3 PositionOffset = mul(DeformedPositionOffset[0], RootLocalToWorld).xyz;
+					AddCrossWS(PositionOffset, NormalScale, ColorMagenta);
+				}
+			}
+		}
+
 		if (!bValid)
 		{
 			return;
