Fix typo and issue render command flushing only when needed.

diff --git a/Engine/Plugins/Runtime/HairStrands/Source/HairStrandsCore/Private/GroomAsset.cpp b/Engine/Plugins/Runtime/HairStrands/Source/HairStrandsCore/Private/GroomAsset.cpp
--- a/Engine/Plugins/Runtime/HairStrands/Source/HairStrandsCore/Private/GroomAsset.cpp
+++ b/Engine/Plugins/Runtime/HairStrands/Source/HairStrandsCore/Private/GroomAsset.cpp
@@ -2512,11 +2512,15 @@ bool UGroomAsset::CacheDerivedDatas()
 
 	// (Temporary) Workaround when reloading groom asset several time, which can cause resources to be still alive while the data are being updated 
 	// by CacheDerivedData. The BulkData referenced by the resources becomes no longer valid, causing use-after-free crash.
-	const bool bReleaseResourceBeforeCacheDerivedData = bIsGameThread && GHairStrandsReleaseResourceBeforeBuildingData > 1;
+	const bool bReleaseResourceBeforeCacheDerivedData = bIsGameThread && GHairStrandsReleaseResourceBeforeBuildingData > 0;
 	if (bReleaseResourceBeforeCacheDerivedData)
 	{
+		const bool bFlushCommands = bIsInitialized;
 		ReleaseResource();
-		FlushRenderingCommands();
+		if (bFlushCommands)
+		{
+			FlushRenderingCommands();
+		}
 	}
 
 	const uint32 GroupCount = GetHairGroupsInterpolation().Num();
