Fix rigged groom.

diff --git a/Engine/Plugins/Runtime/HairStrands/Source/HairStrandsCore/Private/GroomAsset.cpp b/Engine/Plugins/Runtime/HairStrands/Source/HairStrandsCore/Private/GroomAsset.cpp
--- a/Engine/Plugins/Runtime/HairStrands/Source/HairStrandsCore/Private/GroomAsset.cpp
+++ b/Engine/Plugins/Runtime/HairStrands/Source/HairStrandsCore/Private/GroomAsset.cpp
@@ -3545,11 +3545,15 @@ bool UGroomAsset::IsVisible(int32 GroupIndex, int32 LODIndex) const
 bool UGroomAsset::IsDeformationEnable(int32 GroupIndex) const
 {
 	return 
-		(GetHairGroupsInterpolation().IsValidIndex(GroupIndex) && 
-		GetHairGroupsInterpolation()[GroupIndex].InterpolationSettings.GuideType == EGroomGuideType::Rigged) ||
+		(GetGuideType(GroupIndex) == EGroomGuideType::Rigged) ||
 		(GetHairGroupsPhysics().IsValidIndex(GroupIndex) &&	GetHairGroupsPhysics()[GroupIndex].SolverSettings.bEnableDeformation);
 }
 
+EGroomGuideType UGroomAsset::GetGuideType(int32 GroupIndex) const
+{
+	return GetHairGroupsInterpolation().IsValidIndex(GroupIndex) ? GetHairGroupsInterpolation()[GroupIndex].InterpolationSettings.GuideType : EGroomGuideType::Generated;
+}
+
 bool UGroomAsset::IsSimulationEnable(int32 GroupIndex, int32 LODIndex) const 
 {
 	if (!GetHairGroupsLOD().IsValidIndex(GroupIndex) || LODIndex >= GetHairGroupsLOD()[GroupIndex].LODs.Num() || !IsHairStrandsSimulationEnable())
diff --git a/Engine/Plugins/Runtime/HairStrands/Source/HairStrandsCore/Private/GroomComponent.cpp b/Engine/Plugins/Runtime/HairStrands/Source/HairStrandsCore/Private/GroomComponent.cpp
--- a/Engine/Plugins/Runtime/HairStrands/Source/HairStrandsCore/Private/GroomComponent.cpp
+++ b/Engine/Plugins/Runtime/HairStrands/Source/HairStrandsCore/Private/GroomComponent.cpp
@@ -1495,7 +1495,7 @@ bool UGroomComponent::IsSimulationEnable(int32 GroupIndex, int32 LODIndex) const
 bool UGroomComponent::IsDeformationEnable(int32 GroupIndex) const
 {
 	// To enable deformation groom asset should allow it and we must have a mesh deformer locally or being register to a solver component
-	return GroomAsset->IsDeformationEnable(GroupIndex) && (GroomSolver || MeshDeformer);
+	return GroomAsset->IsDeformationEnable(GroupIndex) && (GroomSolver || MeshDeformer || GroomAsset->GetGuideType(GroupIndex) == EGroomGuideType::Rigged);
 }
 
 void UGroomComponent::ReleaseHairSimulation()
diff --git a/Engine/Plugins/Runtime/HairStrands/Source/HairStrandsCore/Public/GroomAsset.h b/Engine/Plugins/Runtime/HairStrands/Source/HairStrandsCore/Public/GroomAsset.h
--- a/Engine/Plugins/Runtime/HairStrands/Source/HairStrandsCore/Public/GroomAsset.h
+++ b/Engine/Plugins/Runtime/HairStrands/Source/HairStrandsCore/Public/GroomAsset.h
@@ -767,6 +767,7 @@ public:
 	UE_API bool IsGlobalInterpolationEnable(int32 GroupIndex, int32 LODIndex) const;
 	UE_API bool NeedsInterpolationData(int32 GroupIndex) const;
 	UE_API bool NeedsInterpolationData() const;
+	EGroomGuideType GetGuideType(int32 GroupIndex) const;
 
 	UE_API void UpdateHairGroupsInfo();
 	UE_API bool HasGeometryType(EGroomGeometryType Type) const;
