Delay groom binding platform data release in order to accomodate existing resources still in flight while async building. These resources points to bulk data, which needs to not be released until the resources are actually released.

diff --git a/Engine/Plugins/Runtime/HairStrands/Source/HairStrandsCore/Private/GroomBindingAsset.cpp b/Engine/Plugins/Runtime/HairStrands/Source/HairStrandsCore/Private/GroomBindingAsset.cpp
--- a/Engine/Plugins/Runtime/HairStrands/Source/HairStrandsCore/Private/GroomBindingAsset.cpp
+++ b/Engine/Plugins/Runtime/HairStrands/Source/HairStrandsCore/Private/GroomBindingAsset.cpp
@@ -34,6 +34,9 @@
 
 LLM_DECLARE_TAG(Groom);
 
+static int32 GHairStrandsDelayBindingDataRelease = 1;
+static FAutoConsoleVariableRef CVarHairStrandsDelayBindingDataRelease(TEXT("r.HairStrands.Binding.DelayBindingDataRelease"), GHairStrandsDelayBindingDataRelease, TEXT("Delay groom binding data release (Experimental)"));
+
 /////////////////////////////////////////////////////////////////////////////////////////
 
 bool IsHairStrandsDDCLogEnable();
@@ -430,6 +433,18 @@ void UGroomBindingAsset::ReleaseResource(bool bResetLoadedSize)
 		});
 		//	#hair_todo: release cards root resources
 	}
+
+	// Process data to be deleted (should happen only in editor)
+	FHairGroupPlatformDataArray* DataToDelete = nullptr;
+	while (HairGroupPlatformDataToDelete.Dequeue(DataToDelete))
+	{
+		ENQUEUE_RENDER_COMMAND(ReleaseHairStrandsPlatformDataCommand)(UE::RenderCommandPipe::Groom,
+		[DataToDelete](FRHICommandList& RHICmdList)
+		{
+			DataToDelete->Data.Empty();
+			delete DataToDelete;
+		});
+	}
 }
 
 void UGroomBindingAsset::FlushRenderingCommandIfUsed(const UGroomBindingAsset* In)
@@ -1248,6 +1263,17 @@ void UGroomBindingAsset::ExecuteCacheDerivedDatas(FGroomBindingBuildContext& Con
 	const bool bAnyGroupNeedRebuild = true;
 #endif
 
+	const bool bDelayDataRelease = GHairStrandsDelayBindingDataRelease > 0;
+	if (bIsValid && bAnyGroupNeedRebuild && bDelayDataRelease)
+	{
+		TArray<UGroomBindingAsset::FHairGroupPlatformData>& PlatformData = GetHairGroupsPlatformData();
+		FHairGroupPlatformDataArray* DataToBeDeleted = new FHairGroupPlatformDataArray();
+		DataToBeDeleted->Data = MoveTemp(PlatformData);
+		HairGroupPlatformDataToDelete.Enqueue(DataToBeDeleted);
+
+		GetHairGroupsPlatformData().SetNum(GroupCount);
+	}
+
 	// 3. Build or retrieve from cache, binding data for each group
 	bIsValid = true;
 	bool bReloadResource = false;
@@ -1259,7 +1285,7 @@ void UGroomBindingAsset::ExecuteCacheDerivedDatas(FGroomBindingBuildContext& Con
 			bool bGroupValid = true;
 			bool bGroupReloadResource = false;
 #if WITH_EDITOR
-			if (GroupDerivedDataKeys[GroupIndex] != CachedDerivedDataKey[GroupIndex])
+			if (GroupDerivedDataKeys[GroupIndex] != CachedDerivedDataKey[GroupIndex] || bDelayDataRelease)
 			{
 				// Fetch the built data from the DDC or build it
 
diff --git a/Engine/Plugins/Runtime/HairStrands/Source/HairStrandsCore/Public/GroomBindingAsset.h b/Engine/Plugins/Runtime/HairStrands/Source/HairStrandsCore/Public/GroomBindingAsset.h
--- a/Engine/Plugins/Runtime/HairStrands/Source/HairStrandsCore/Public/GroomBindingAsset.h
+++ b/Engine/Plugins/Runtime/HairStrands/Source/HairStrandsCore/Public/GroomBindingAsset.h
@@ -368,10 +368,18 @@ public:
 
 private:
 	/** Queue of resources which needs to be deleted. This queue is needed for keeping valid pointer on the group resources 
-	   when the binding asset is recomputed */
+	    when the binding asset is recomputed */
 	UE_DEPRECATED(5.3, "Please do not access this member directly; use UGroomBindingAsset::GetHairGroupPlatformData().")
 	TQueue<FHairGroupResource> HairGroupResourcesToDelete;
 
+	/** Queue of data which needs to be deleted. This queue is needed for keeping valid pointer on the group resources 
+	    when the binding asset is recomputed */
+	struct FHairGroupPlatformDataArray
+	{
+		TArray<FHairGroupPlatformData> Data;
+	};
+	TQueue<FHairGroupPlatformDataArray*> HairGroupPlatformDataToDelete;
+
 	/** Platform data for each hair groups */
 	UE_DEPRECATED(5.3, "Please do not access this member directly; use UGroomBindingAsset::GetHairGroupPlatformData().")
 	TArray<FHairGroupPlatformData> HairGroupsPlatformData;
