Fix several shader compilation issue related to read/write operation in groom deformer.

diff --git a/Engine/Plugins/Runtime/HairStrands/Shaders/Private/Deformers/DeformerGroomStrandsRead.ush b/Engine/Plugins/Runtime/HairStrands/Shaders/Private/Deformers/DeformerGroomStrandsRead.ush
--- a/Engine/Plugins/Runtime/HairStrands/Shaders/Private/Deformers/DeformerGroomStrandsRead.ush
+++ b/Engine/Plugins/Runtime/HairStrands/Shaders/Private/Deformers/DeformerGroomStrandsRead.ush
@@ -16,7 +16,8 @@ uint {DataInterfaceName}_TotalCurveCount;
 #define HAIR_STRANDS_ATTRIBUTE_ACCESSORS(Name) {DataInterfaceName}_##Name
 #include "/Engine/Private/HairStrands/HairStrandsAttributeTemplate.ush"
 
-ByteAddressBuffer {DataInterfaceName}_InterpolationBuffer;
+ByteAddressBuffer {DataInterfaceName}_CurveInterpolationBuffer;
+ByteAddressBuffer {DataInterfaceName}_PointInterpolationBuffer;
 
 #if ENABLE_SKINNED_TRANSFORM == 0
 float4x4 {DataInterfaceName}_RigidRestTransform;
diff --git a/Engine/Plugins/Runtime/HairStrands/Shaders/Private/Deformers/DeformerGroomStrandsWrite.ush b/Engine/Plugins/Runtime/HairStrands/Shaders/Private/Deformers/DeformerGroomStrandsWrite.ush
--- a/Engine/Plugins/Runtime/HairStrands/Shaders/Private/Deformers/DeformerGroomStrandsWrite.ush
+++ b/Engine/Plugins/Runtime/HairStrands/Shaders/Private/Deformers/DeformerGroomStrandsWrite.ush
@@ -14,9 +14,8 @@ uint {DataInterfaceName}_TotalCurveCount;
 // because we write out data which would be bound as SRV otherwise (position, attributes, ...)
 StructuredBuffer<float4> {DataInterfaceName}_DeformedPositionOffset;
 ByteAddressBuffer		{DataInterfaceName}_PointRestPositions;
-RWByteAddressBuffer		{DataInterfaceName}_PointDeformedPositions;
-RWByteAddressBuffer 	{DataInterfaceName}_PointAttributeBuffer;
-RWByteAddressBuffer 	{DataInterfaceName}_PointAttributeBuffer;
+RWByteAddressBuffer		{DataInterfaceName}_OutPointDeformedPositions;
+RWByteAddressBuffer 	{DataInterfaceName}_OutPointAttributeBuffer;
 
 // Helper functions
 
@@ -82,7 +81,7 @@ void WritePointDeformedPosition_{DataInterfaceName}(uint PointIndex, float3 Posi
 	if (WriteIndex < {DataInterfaceName}_PointCount)
 	{
 		const FPackedHairPosition PackedCP = ReadPackedHairPosition({DataInterfaceName}_PointRestPositions, WriteIndex);
-		WritePackedHairControlPointPosition({DataInterfaceName}_PointDeformedPositions, WriteIndex, PackedCP, Position, {DataInterfaceName}_GetDeformedPositionOffset());
+		WritePackedHairControlPointPosition({DataInterfaceName}_OutPointDeformedPositions, WriteIndex, PackedCP, Position, {DataInterfaceName}_GetDeformedPositionOffset());
 	}
 }
 
@@ -95,7 +94,7 @@ void WritePointRadius_{DataInterfaceName}(uint PointIndex, float Radius)
 	{
 		const FPackedHairPosition PackedCP = ReadPackedHairPosition({DataInterfaceName}_PointRestPositions, WriteIndex);
 		const FPackedHairPosition NewPackedCP = PackHairControlPointRadius(PackedCP, Radius, {DataInterfaceName}_Radius);
-		WritePackedHairPosition({DataInterfaceName}_PointDeformedPositions, WriteIndex, NewPackedCP);
+		WritePackedHairPosition({DataInterfaceName}_OutPointDeformedPositions, WriteIndex, NewPackedCP);
 	}
 }
 
@@ -110,7 +109,7 @@ void WritePointPositionRadius_{DataInterfaceName}(uint PointIndex, float4 Positi
 		FPackedHairPosition NewPackedCP = 0;
 		NewPackedCP = PackHairControlPointPosition(PackedCP, PositionAndRadius.xyz, {DataInterfaceName}_GetDeformedPositionOffset());
 		NewPackedCP = PackHairControlPointRadius(NewPackedCP, PositionAndRadius.w, {DataInterfaceName}_Radius);
-		WritePackedHairPosition({DataInterfaceName}_PointDeformedPositions, WriteIndex, NewPackedCP);
+		WritePackedHairPosition({DataInterfaceName}_OutPointDeformedPositions, WriteIndex, NewPackedCP);
 	}
 }
 
@@ -121,7 +120,7 @@ void WriteCurveRootUV_{DataInterfaceName}(uint CurveIndex, float2 RootUV)
 	const uint WriteIndex = CurveIndex;
 	if (WriteIndex < {DataInterfaceName}_CurveCount && HasHairStrandsRootUV())
 	{
-		InternalSetHairStrandsAttributes32(CurveIndex, HAIR_CURVE_ATTRIBUTE_OFFSET_ROOTUV({DataInterfaceName}_CurveAttributeOffsets), HAIR_CURVE_ATTRIBUTE_STRIDE_ROOTUV, {DataInterfaceName}_PointAttributeBuffer, PackHairRootUV(RootUV), {DataInterfaceName}_CurveAttributeIndexToChunkDivAsShift, {DataInterfaceName}_CurveAttributeChunkElementCount), {DataInterfaceName}_CurveAttributeChunkStrideInBytes);
+		InternalSetHairStrandsAttributes32(CurveIndex, HAIR_CURVE_ATTRIBUTE_OFFSET_ROOTUV({DataInterfaceName}_CurveAttributeOffsets), HAIR_CURVE_ATTRIBUTE_STRIDE_ROOTUV, {DataInterfaceName}_OutPointAttributeBuffer, PackHairRootUV(RootUV), {DataInterfaceName}_CurveAttributeIndexToChunkDivAsShift, {DataInterfaceName}_CurveAttributeChunkElementCount, {DataInterfaceName}_CurveAttributeChunkStrideInBytes);
 	}
 }
 
@@ -132,7 +131,7 @@ void WriteCurveSeed_{DataInterfaceName}(uint CurveIndex, float Seed)
 	const uint WriteIndex = CurveIndex;
 	if (WriteIndex < {DataInterfaceName}_CurveCount)
 	{
-		InternalSetHairStrandsAttributes(CurveIndex, HAIR_CURVE_ATTRIBUTE_OFFSET_SEED({DataInterfaceName}_CurveAttributeOffsets), HAIR_CURVE_ATTRIBUTE_STRIDE_SEED, {DataInterfaceName}_PointAttributeBuffer, PackHairSeed(Seed), {DataInterfaceName}_CurveAttributeIndexToChunkDivAsShift, {DataInterfaceName}_CurveAttributeChunkElementCount), {DataInterfaceName}_CurveAttributeChunkStrideInBytes);
+		InternalSetHairStrandsAttributes(CurveIndex, HAIR_CURVE_ATTRIBUTE_OFFSET_SEED({DataInterfaceName}_CurveAttributeOffsets), HAIR_CURVE_ATTRIBUTE_STRIDE_SEED, {DataInterfaceName}_OutPointAttributeBuffer, PackHairSeed(Seed), {DataInterfaceName}_CurveAttributeIndexToChunkDivAsShift, {DataInterfaceName}_CurveAttributeChunkElementCount, {DataInterfaceName}_CurveAttributeChunkStrideInBytes);
 	}
 }
 
@@ -143,7 +142,7 @@ void WriteCurveClumpId_{DataInterfaceName}(uint CurveIndex, uint ClumpId)
 	const uint WriteIndex = CurveIndex;
 	if (WriteIndex < {DataInterfaceName}_CurveCount && HasHairStrandsClumpId())
 	{
-		InternalSetHairStrandsAttributes(CurveIndex, HAIR_CURVE_ATTRIBUTE_OFFSET_CLUMPID({DataInterfaceName}_CurveAttributeOffsets), HAIR_CURVE_ATTRIBUTE_STRIDE_CLUMPID, {DataInterfaceName}_PointAttributeBuffer, PackHairClumpID(ClumpId), {DataInterfaceName}_CurveAttributeIndexToChunkDivAsShift, {DataInterfaceName}_CurveAttributeChunkElementCount), {DataInterfaceName}_CurveAttributeChunkStrideInBytes);
+		InternalSetHairStrandsAttributes(CurveIndex, HAIR_CURVE_ATTRIBUTE_OFFSET_CLUMPID({DataInterfaceName}_CurveAttributeOffsets), HAIR_CURVE_ATTRIBUTE_STRIDE_CLUMPID, {DataInterfaceName}_OutPointAttributeBuffer, PackHairClumpID(ClumpId), {DataInterfaceName}_CurveAttributeIndexToChunkDivAsShift, {DataInterfaceName}_CurveAttributeChunkElementCount, {DataInterfaceName}_CurveAttributeChunkStrideInBytes);
 	}
 }
 
@@ -154,7 +153,7 @@ void WritePointColor_{DataInterfaceName}(uint PointIndex, float3 Color)
 	const uint WriteIndex = PointIndex;
 	if (WriteIndex < {DataInterfaceName}_PointCount && HasHairStrandsColor())
 	{
-		InternalSetHairStrandsAttributes32(PointIndex, HAIR_POINT_ATTRIBUTE_OFFSET_COLOR({DataInterfaceName}_PointAttributeOffsets), HAIR_POINT_ATTRIBUTE_STRIDE_COLOR, {DataInterfaceName}_PointAttributeBuffer, PackHairColor(Color), {DataInterfaceName}_PointAttributeIndexToChunkDivAsShift, {DataInterfaceName}_PointAttributeChunkElementCount), {DataInterfaceName}_PointAttributeChunkStrideInBytes);
+		InternalSetHairStrandsAttributes32(PointIndex, HAIR_POINT_ATTRIBUTE_OFFSET_COLOR({DataInterfaceName}_PointAttributeOffsets), HAIR_POINT_ATTRIBUTE_STRIDE_COLOR, {DataInterfaceName}_OutPointAttributeBuffer, PackHairColor(Color), {DataInterfaceName}_PointAttributeIndexToChunkDivAsShift, {DataInterfaceName}_PointAttributeChunkElementCount, {DataInterfaceName}_PointAttributeChunkStrideInBytes);
 	}
 }
 
@@ -165,7 +164,7 @@ void WritePointRoughness_{DataInterfaceName}(uint PointIndex, float Roughness)
 	const uint WriteIndex = PointIndex;
 	if (WriteIndex < {DataInterfaceName}_PointCount && HasHairStrandsRoughness())
 	{
-		InternalSetHairStrandsAttributes(PointIndex, HAIR_POINT_ATTRIBUTE_OFFSET_ROUGHNESS({DataInterfaceName}_PointAttributeOffsets), HAIR_POINT_ATTRIBUTE_STRIDE_ROUGHNESS, {DataInterfaceName}_PointAttributeBuffer, PackHairRoughness(Roughness), {DataInterfaceName}_PointAttributeIndexToChunkDivAsShift, {DataInterfaceName}_PointAttributeChunkElementCount), {DataInterfaceName}_PointAttributeChunkStrideInBytes);
+		InternalSetHairStrandsAttributes(PointIndex, HAIR_POINT_ATTRIBUTE_OFFSET_ROUGHNESS({DataInterfaceName}_PointAttributeOffsets), HAIR_POINT_ATTRIBUTE_STRIDE_ROUGHNESS, {DataInterfaceName}_OutPointAttributeBuffer, PackHairRoughness(Roughness), {DataInterfaceName}_PointAttributeIndexToChunkDivAsShift, {DataInterfaceName}_PointAttributeChunkElementCount, {DataInterfaceName}_PointAttributeChunkStrideInBytes);
 	}
 }
 
@@ -176,6 +175,6 @@ void WritePointAO_{DataInterfaceName}(uint PointIndex, float Ao)
 	const uint WriteIndex = PointIndex;
 	if (WriteIndex < {DataInterfaceName}_PointCount && HasHairStrandsAO())
 	{
-		InternalSetHairStrandsAttributes(PointIndex, HAIR_POINT_ATTRIBUTE_OFFSET_AO({DataInterfaceName}_PointAttributeOffsets), HAIR_POINT_ATTRIBUTE_STRIDE_AO, {DataInterfaceName}_PointAttributeBuffer, PackHairAO(Ao), {DataInterfaceName}_PointAttributeIndexToChunkDivAsShift, {DataInterfaceName}_PointAttributeChunkElementCount), {DataInterfaceName}_PointAttributeChunkStrideInBytes);
+		InternalSetHairStrandsAttributes(PointIndex, HAIR_POINT_ATTRIBUTE_OFFSET_AO({DataInterfaceName}_PointAttributeOffsets), HAIR_POINT_ATTRIBUTE_STRIDE_AO, {DataInterfaceName}_OutPointAttributeBuffer, PackHairAO(Ao), {DataInterfaceName}_PointAttributeIndexToChunkDivAsShift, {DataInterfaceName}_PointAttributeChunkElementCount, {DataInterfaceName}_PointAttributeChunkStrideInBytes);
 	}
 }
\ No newline at end of file
diff --git a/Engine/Plugins/Runtime/HairStrands/Source/HairStrandsDeformer/Private/DeformerGroomStrandsRead.h b/Engine/Plugins/Runtime/HairStrands/Source/HairStrandsDeformer/Private/DeformerGroomStrandsRead.h
--- a/Engine/Plugins/Runtime/HairStrands/Source/HairStrandsDeformer/Private/DeformerGroomStrandsRead.h
+++ b/Engine/Plugins/Runtime/HairStrands/Source/HairStrandsDeformer/Private/DeformerGroomStrandsRead.h
@@ -84,6 +84,8 @@ private:
 		FRDGBufferSRVRef TriangleDeformedPositions = nullptr;
 		FRDGBufferSRVRef CurveBarycentricCoordinates = nullptr;
 		FRDGBufferSRVRef CurveTriangleIndices = nullptr;
+		FRDGBufferSRVRef CurveInterpolationBuffer = nullptr;
+		FRDGBufferSRVRef PointInterpolationBuffer = nullptr;
 	};
 	TArray<FBindingResources> BindingResources;
 
diff --git a/Engine/Plugins/Runtime/HairStrands/Source/HairStrandsDeformer/Private/DeformerGroomStrandsWrite.cpp b/Engine/Plugins/Runtime/HairStrands/Source/HairStrandsDeformer/Private/DeformerGroomStrandsWrite.cpp
--- a/Engine/Plugins/Runtime/HairStrands/Source/HairStrandsDeformer/Private/DeformerGroomStrandsWrite.cpp
+++ b/Engine/Plugins/Runtime/HairStrands/Source/HairStrandsDeformer/Private/DeformerGroomStrandsWrite.cpp
@@ -118,9 +118,9 @@ BEGIN_SHADER_PARAMETER_STRUCT(FOptimusGroomStrandsWriteParameters, )
 	SHADER_PARAMETER(uint32, TotalCurveCount)
 	SHADER_PARAMETER_RDG_BUFFER_SRV(StructuredBuffer, DeformedPositionOffset)
 	SHADER_PARAMETER_RDG_BUFFER_SRV(ByteAddressBuffer, PointRestPositions)
-	SHADER_PARAMETER_RDG_BUFFER_UAV(RWByteAddressBuffer, PointDeformedPositions)
-	SHADER_PARAMETER_RDG_BUFFER_UAV(RWByteAddressBuffer, CurveAttributeBuffer)
-	SHADER_PARAMETER_RDG_BUFFER_UAV(RWByteAddressBuffer, PointAttributeBuffer)
+	SHADER_PARAMETER_RDG_BUFFER_UAV(RWByteAddressBuffer, OutPointDeformedPositions)
+	SHADER_PARAMETER_RDG_BUFFER_UAV(RWByteAddressBuffer, OutCurveAttributeBuffer)
+	SHADER_PARAMETER_RDG_BUFFER_UAV(RWByteAddressBuffer, OutPointAttributeBuffer)
 END_SHADER_PARAMETER_STRUCT()
 
 void UOptimusGroomStrandsWriteDataInterface::GetShaderParameters(TCHAR const* UID, FShaderParametersMetadataBuilder& InOutBuilder, FShaderParametersMetadataAllocations& InOutAllocations) const
@@ -208,7 +208,7 @@ void FOptimusGroomStrandsWriteProviderProxy::AllocateResources(FRDGBuilder& Grap
 			{
 				Resource.DeformedPositionOffset 	= Register(GraphBuilder, GroupInstance->Strands.RestResource->PositionOffsetBuffer, ERDGImportedBufferFlags::CreateSRV).SRV;
 				Resource.PointRestPositions 		= Register(GraphBuilder, GroupInstance->Strands.RestResource->PositionBuffer, ERDGImportedBufferFlags::CreateSRV).SRV;
-				Resource.PointDeformedPositions 	= Register(GraphBuilder, GroupInstance->Strands.DeformedResource->GetDeformerBuffer(GraphBuilder), ERDGImportedBufferFlags::CreateUAV).UAV;
+				Resource.OutPointDeformedPositions 	= Register(GraphBuilder, GroupInstance->Strands.DeformedResource->GetDeformerBuffer(GraphBuilder), ERDGImportedBufferFlags::CreateUAV).UAV;
 			}
 			
 			// Always copy the attributes from the rest asset, so that if deformer write different attribute at different tick, it all remains consistent with the source data
@@ -217,32 +217,32 @@ void FOptimusGroomStrandsWriteProviderProxy::AllocateResources(FRDGBuilder& Grap
 				// Curve Attributes
 				if (OutputMask & 0x38)
 				{
-					FRDGExternalBuffer CurveAttributeBufferExt = GroupInstance->Strands.DeformedResource->GetDeformerCurveAttributeBuffer(GraphBuilder);
-					if (CurveAttributeBufferExt.Buffer)
+					FRDGExternalBuffer OutCurveAttributeBufferExt = GroupInstance->Strands.DeformedResource->GetDeformerCurveAttributeBuffer(GraphBuilder);
+					if (OutCurveAttributeBufferExt.Buffer)
 					{
-						FRDGImportedBuffer DstCurveAttributeBuffer = Register(GraphBuilder, CurveAttributeBufferExt, ERDGImportedBufferFlags::CreateUAV);
+						FRDGImportedBuffer DstCurveAttributeBuffer = Register(GraphBuilder, OutCurveAttributeBufferExt, ERDGImportedBufferFlags::CreateUAV);
 						FRDGImportedBuffer SrcCurveAttributeBuffer = Register(GraphBuilder, GroupInstance->Strands.RestResource->CurveAttributeBuffer, ERDGImportedBufferFlags::None);
 						if (SrcCurveAttributeBuffer.Buffer)
 						{
 							AddCopyBufferPass(GraphBuilder, DstCurveAttributeBuffer.Buffer, SrcCurveAttributeBuffer.Buffer);
 						}
-						Resource.CurveAttributeBuffer = DstCurveAttributeBuffer.UAV;
+						Resource.OutCurveAttributeBuffer = DstCurveAttributeBuffer.UAV;
 					}
 				}
 
 				// Point Attributes
 				if (OutputMask & 0x1c0)
 				{
-					FRDGExternalBuffer PointAttributeBufferExt = GroupInstance->Strands.DeformedResource->GetDeformerPointAttributeBuffer(GraphBuilder);
-					if (PointAttributeBufferExt.Buffer)
+					FRDGExternalBuffer OutPointAttributeBufferExt = GroupInstance->Strands.DeformedResource->GetDeformerPointAttributeBuffer(GraphBuilder);
+					if (OutPointAttributeBufferExt.Buffer)
 					{
-						FRDGImportedBuffer DstPointAttributeBuffer = Register(GraphBuilder, PointAttributeBufferExt, ERDGImportedBufferFlags::CreateUAV);
+						FRDGImportedBuffer DstPointAttributeBuffer = Register(GraphBuilder, OutPointAttributeBufferExt, ERDGImportedBufferFlags::CreateUAV);
 						FRDGImportedBuffer SrcPointAttributeBuffer = Register(GraphBuilder, GroupInstance->Strands.RestResource->PointAttributeBuffer, ERDGImportedBufferFlags::None);
 						if(SrcPointAttributeBuffer.Buffer)
 						{
 							AddCopyBufferPass(GraphBuilder, DstPointAttributeBuffer.Buffer, SrcPointAttributeBuffer.Buffer);
 						}
-						Resource.PointAttributeBuffer = DstPointAttributeBuffer.UAV;
+						Resource.OutPointAttributeBuffer = DstPointAttributeBuffer.UAV;
 					}
 				}
 			}
@@ -262,7 +262,7 @@ void FOptimusGroomStrandsWriteProviderProxy::GatherDispatchData(FDispatchData co
 		{
 			FDispatchResources& Resource = GroupResources[InvocationIndex];
 
-			const bool bValid = Resource.PointDeformedPositions != nullptr && Resource.PointRestPositions != nullptr && Resource.DeformedPositionOffset != nullptr;
+			const bool bValid = Resource.OutPointDeformedPositions != nullptr && Resource.PointRestPositions != nullptr && Resource.DeformedPositionOffset != nullptr;
 			const FHairGroupPublicData::FVertexFactoryInput VFInput = ComputeHairStrandsVertexInputData(GroupInstance, EGroomViewMode::None);		
 
 			FParameters& Parameters = ParameterArray[InvocationIndex];
@@ -277,9 +277,9 @@ void FOptimusGroomStrandsWriteProviderProxy::GatherDispatchData(FDispatchData co
 			
 			Parameters.DeformedPositionOffset 	= bValid ? Resource.DeformedPositionOffset : FallbackPositionBufferSRV;
 			Parameters.PointRestPositions 		= bValid ? Resource.PointRestPositions : FallbackPositionBufferSRV;
-			Parameters.PointDeformedPositions 	= bValid ? Resource.PointDeformedPositions : FallbackPositionBufferUAV;
-			Parameters.PointAttributeBuffer 	= bValid ? Resource.PointAttributeBuffer : FallbackAttributeBufferUAV;
-			Parameters.CurveAttributeBuffer 	= bValid ? Resource.CurveAttributeBuffer : FallbackAttributeBufferUAV;
+			Parameters.OutPointDeformedPositions= bValid ? Resource.OutPointDeformedPositions : FallbackPositionBufferUAV;
+			Parameters.OutPointAttributeBuffer 	= bValid ? Resource.OutPointAttributeBuffer : FallbackAttributeBufferUAV;
+			Parameters.OutCurveAttributeBuffer 	= bValid ? Resource.OutCurveAttributeBuffer : FallbackAttributeBufferUAV;
 		}
 	}
 
diff --git a/Engine/Plugins/Runtime/HairStrands/Source/HairStrandsDeformer/Private/DeformerGroomStrandsWrite.h b/Engine/Plugins/Runtime/HairStrands/Source/HairStrandsDeformer/Private/DeformerGroomStrandsWrite.h
--- a/Engine/Plugins/Runtime/HairStrands/Source/HairStrandsDeformer/Private/DeformerGroomStrandsWrite.h
+++ b/Engine/Plugins/Runtime/HairStrands/Source/HairStrandsDeformer/Private/DeformerGroomStrandsWrite.h
@@ -88,9 +88,9 @@ private:
 	{
 		FRDGBufferSRVRef DeformedPositionOffset = nullptr;
 		FRDGBufferSRVRef PointRestPositions = nullptr;
-		FRDGBufferUAVRef PointDeformedPositions = nullptr;
-		FRDGBufferUAVRef PointAttributeBuffer = nullptr;
-		FRDGBufferUAVRef CurveAttributeBuffer = nullptr;
+		FRDGBufferUAVRef OutPointDeformedPositions = nullptr;
+		FRDGBufferUAVRef OutPointAttributeBuffer = nullptr;
+		FRDGBufferUAVRef OutCurveAttributeBuffer = nullptr;
 	};
 	TArray<FDispatchResources> GroupResources;
 
