Fix crash in hair RBF sample init pass due to skin cache data not being ready.

The skin cache task was not synced prior to the binding surface data update.

diff --git a/Engine/Plugins/Runtime/HairStrands/Source/HairStrandsCore/Private/GroomManager.cpp b/Engine/Plugins/Runtime/HairStrands/Source/HairStrandsCore/Private/GroomManager.cpp
--- a/Engine/Plugins/Runtime/HairStrands/Source/HairStrandsCore/Private/GroomManager.cpp
+++ b/Engine/Plugins/Runtime/HairStrands/Source/HairStrandsCore/Private/GroomManager.cpp
@@ -646,8 +646,8 @@ static FCachedGeometry GetCacheGeometryForHair(
 			if (const FSkeletalMeshSceneProxy* SceneProxy = static_cast<const FSkeletalMeshSceneProxy*>(PrimitiveSceneInfo->Proxy))
 			{
 				// 1. Get cached geometry from GPU SkinCache, if enabled
-				SceneProxy->GetCachedGeometry(GraphBuilder, Out);
-				if (Out.Sections.Num() > 0)
+				const bool bValidCachedGeometry = SceneProxy->GetCachedGeometry(GraphBuilder, Out);
+				if (bValidCachedGeometry && Out.Sections.Num() > 0)
 				{
 					OutHairGeometryCache.AddDebug(Instance, SceneProxy, Out, PositionUpdateType, FHairGeometryCache::ECacheType::SkinCache);
 				}
diff --git a/Engine/Source/Runtime/Engine/Private/SkeletalRenderGPUSkin.cpp b/Engine/Source/Runtime/Engine/Private/SkeletalRenderGPUSkin.cpp
--- a/Engine/Source/Runtime/Engine/Private/SkeletalRenderGPUSkin.cpp
+++ b/Engine/Source/Runtime/Engine/Private/SkeletalRenderGPUSkin.cpp
@@ -2474,7 +2474,7 @@ bool FSkeletalMeshObjectGPUSkin::GetCachedGeometry(FRDGBuilder& GraphBuilder, FC
 		{
 			// Get the cached geometry SRVs from the skin cache.
 			FRWBuffer* PositionBuffer = FGPUSkinCache::GetPositionBuffer(GraphBuilder, SkinCacheEntry, SectionIndex);
-			if (!PositionBuffer)
+			if (!PositionBuffer || !PositionBuffer->SRV)
 			{
 				return false;
 			}
diff --git a/Engine/Source/Runtime/Renderer/Private/SceneVisibility.cpp b/Engine/Source/Runtime/Renderer/Private/SceneVisibility.cpp
--- a/Engine/Source/Runtime/Renderer/Private/SceneVisibility.cpp
+++ b/Engine/Source/Runtime/Renderer/Private/SceneVisibility.cpp
@@ -4767,6 +4767,9 @@ void FSceneRenderer::PreVisibilityFrameSetup(FRDGBuilder& GraphBuilder)
 			{
 				if (Scene && IsHairStrandsEnabled(EHairStrandsShaderType::All, Scene->GetShaderPlatform()))
 				{
+					// Need GPUSkinCache results to be available in order to get correct skel. mesh data
+					Scene->GetGPUSkinCacheTask().Wait();
+
 					// 1.Update binding surfaces
 					// Prepare surface data (MeshLODData)
 					RunHairStrandsBookmark(GraphBuilder, EHairStrandsBookmark::ProcessBindingSurfaceUpdate, Parameters);
