Fix conversion from/to groom edit causing extra attribute to be written while they shouldn't.

Work done by [REDACTED_PERSON]. Checking this on his behalf.

diff --git a/Engine/Plugins/Runtime/HairStrands/Source/HairStrandsCore/Private/GroomEdit.cpp b/Engine/Plugins/Runtime/HairStrands/Source/HairStrandsCore/Private/GroomEdit.cpp
--- a/Engine/Plugins/Runtime/HairStrands/Source/HairStrandsCore/Private/GroomEdit.cpp
+++ b/Engine/Plugins/Runtime/HairStrands/Source/HairStrandsCore/Private/GroomEdit.cpp
@@ -20,9 +20,12 @@ static void HairStrandsDataToEditableHairGuides(const FHairStrandsDatas& In, TAr
 	Out.SetNum(StrandsCount);
 	for (uint32 StrandsIt = 0; StrandsIt < StrandsCount; ++StrandsIt)
 	{
-		Out[StrandsIt].GuideID= bHasValidStrandsIDs ? In.StrandsCurves.StrandIDs[StrandsIt] : StrandsIt;
+		Out[StrandsIt].GuideID = bHasValidStrandsIDs ? In.StrandsCurves.StrandIDs[StrandsIt] : StrandsIt;
 		Out[StrandsIt].RootUV = bHasValidRootUVs ? In.StrandsCurves.CurvesRootUV[StrandsIt] : FVector2f(0.0f,0.0f);
 
+		Out[StrandsIt].bHasGuideID = bHasValidStrandsIDs;
+		Out[StrandsIt].bHasRootUV = bHasValidRootUVs;
+
 		const uint32 PointCount = In.StrandsCurves.CurvesCount[StrandsIt];
 		const uint32 PointOffset = In.StrandsCurves.CurvesOffset[StrandsIt];
 
@@ -39,7 +42,7 @@ static void HairStrandsDataToEditableHairGuides(const FHairStrandsDatas& In, TAr
 static void HairStrandsDataToEditableHairStrands(const FHairStrandsDatas& In, TArray<FEditableHairStrand>& Out)
 {
 	const bool bHasValidStrandsIDs = In.StrandsCurves.StrandIDs.Num() > 0;
-	const bool bHasValidClumpIDs = In.StrandsCurves.ClumpIDs.Num() > 0;	
+	const bool bHasValidClumpIDs = In.StrandsCurves.ClumpIDs.Num() > 0;
 	const bool bHasValidClosestGuide = In.StrandsCurves.CurvesClosestGuideIDs.Num() > 0;
 	const bool bHasValidRootUVs = In.StrandsCurves.CurvesRootUV.Num() > 0;
 	
@@ -60,12 +63,18 @@ static void HairStrandsDataToEditableHairStrands(const FHairStrandsDatas& In, TA
 		Out[StrandsIt].GuideWeights[1] = GuideWeights.Y;
 		Out[StrandsIt].GuideWeights[2] = GuideWeights.Z;
 
+		Out[StrandsIt].bHasStrandID = bHasValidStrandsIDs;
+		Out[StrandsIt].bHasClumpID = bHasValidClumpIDs;
+		Out[StrandsIt].bHasClosestGuide = bHasValidClosestGuide;
+		Out[StrandsIt].bHasRootUV = bHasValidRootUVs;
+
 		const uint32 PointCount  = In.StrandsCurves.CurvesCount[StrandsIt];
 		const uint32 PointOffset = In.StrandsCurves.CurvesOffset[StrandsIt];
+		const uint32 MaxPoint = PointOffset + PointCount - 1;
 
-		const bool bHasValidColor = (In.StrandsPoints.PointsBaseColor.Num() == PointCount);
-		const bool bHasValidRoughness = (In.StrandsPoints.PointsRoughness.Num() == PointCount);
-		const bool bHasValidAO = (In.StrandsPoints.PointsAO.Num() == PointCount);
+		const bool bHasValidColor = (In.StrandsPoints.PointsBaseColor.IsValidIndex(MaxPoint));
+		const bool bHasValidRoughness = (In.StrandsPoints.PointsRoughness.IsValidIndex(MaxPoint));
+		const bool bHasValidAO = (In.StrandsPoints.PointsAO.IsValidIndex(MaxPoint));
 
 		Out[StrandsIt].ControlPoints.SetNum(PointCount);
 		for (uint32 PointIt = 0; PointIt < PointCount; ++PointIt)
@@ -74,6 +83,10 @@ static void HairStrandsDataToEditableHairStrands(const FHairStrandsDatas& In, TA
 			Out[StrandsIt].ControlPoints[PointIt].Position	= In.StrandsPoints.PointsPosition[EffectiveIndex];
 			Out[StrandsIt].ControlPoints[PointIt].U			= In.StrandsPoints.PointsCoordU[EffectiveIndex];
 			Out[StrandsIt].ControlPoints[PointIt].Radius	= In.StrandsPoints.PointsRadius[EffectiveIndex];
+
+			Out[StrandsIt].ControlPoints[PointIt].bHasColor = bHasValidColor;
+			Out[StrandsIt].ControlPoints[PointIt].bHasRoughness = bHasValidRoughness;
+			Out[StrandsIt].ControlPoints[PointIt].bHasAO = bHasValidAO;
 			
 			if(bHasValidColor)
 			{
@@ -99,8 +112,14 @@ static void EditableGroomGroupToHairDescription(FHairDescription& OutHairDescrip
 
 		const uint32 PointCount = Guide.ControlPoints.Num();
 		SetHairStrandAttribute(OutHairDescription, StrandID, HairAttribute::Strand::VertexCount, (int32)PointCount);
-		SetHairStrandAttribute(OutHairDescription, StrandID, HairAttribute::Strand::RootUV,      Guide.RootUV);
-		SetHairStrandAttribute(OutHairDescription, StrandID, HairAttribute::Strand::ID,          int(Guide.GuideID));
+		if(Guide.bHasRootUV)
+		{ 
+			SetHairStrandAttribute(OutHairDescription, StrandID, HairAttribute::Strand::RootUV,      Guide.RootUV);
+		}
+		if (Guide.bHasGuideID)
+		{
+			SetHairStrandAttribute(OutHairDescription, StrandID, HairAttribute::Strand::ID,          int(Guide.GuideID));
+		}
 		SetHairStrandAttribute(OutHairDescription, StrandID, HairAttribute::Strand::GroupID,     int(GroupID));
 		SetHairStrandAttribute(OutHairDescription, StrandID, HairAttribute::Strand::GroupName,   GroupName);
 		SetHairStrandAttribute(OutHairDescription, StrandID, HairAttribute::Strand::Guide,       1);
@@ -131,26 +150,38 @@ static void EditableGroomGroupToHairDescription(FHairDescription& OutHairDescrip
 
 		const uint32 PointCount = Strands.ControlPoints.Num();
 		SetHairStrandAttribute(OutHairDescription, StrandID, HairAttribute::Strand::VertexCount, (int32)PointCount);
-		SetHairStrandAttribute(OutHairDescription, StrandID, HairAttribute::Strand::RootUV,      Strands.RootUV);
-		SetHairStrandAttribute(OutHairDescription, StrandID, HairAttribute::Strand::ID,          int(Strands.StrandID));
-		SetHairStrandAttribute(OutHairDescription, StrandID, HairAttribute::Strand::ClumpID,     int(Strands.ClumpID));
+		if(Strands.bHasRootUV)
+		{ 
+			SetHairStrandAttribute(OutHairDescription, StrandID, HairAttribute::Strand::RootUV,      Strands.RootUV);
+		}
+		if (Strands.bHasStrandID)
+		{
+			SetHairStrandAttribute(OutHairDescription, StrandID, HairAttribute::Strand::ID,          int(Strands.StrandID));
+		}
+		if (Strands.bHasClumpID)
+		{
+			SetHairStrandAttribute(OutHairDescription, StrandID, HairAttribute::Strand::ClumpID,     int(Strands.ClumpID));
+		}
 		SetHairStrandAttribute(OutHairDescription, StrandID, HairAttribute::Strand::GroupID,     int(GroupID));
 		SetHairStrandAttribute(OutHairDescription, StrandID, HairAttribute::Strand::GroupName,   GroupName);
 		SetHairStrandAttribute(OutHairDescription, StrandID, HairAttribute::Strand::Guide,       0);
 
-		if (bHasTripleWeights)
+		if (Strands.bHasClosestGuide)
 		{
-			const FVector3f ClosestGuides = FVector3f(Strands.GuideIDs[0], Strands.GuideIDs[1], Strands.GuideIDs[2]);
-			const FVector3f GuideWeights = FVector3f(Strands.GuideWeights[0], Strands.GuideWeights[1], Strands.GuideWeights[2]);
-			SetHairStrandAttribute(OutHairDescription, StrandID, HairAttribute::Strand::ClosestGuides, ClosestGuides);
-			SetHairStrandAttribute(OutHairDescription, StrandID, HairAttribute::Strand::GuideWeights, GuideWeights);
-		}
-		else if (bHasSingleWeights)
-		{
-			int ClosestGuide = Strands.GuideIDs[0];
-			float GuideWeight = Strands.GuideWeights[0];
-			SetHairStrandAttribute(OutHairDescription, StrandID, HairAttribute::Strand::ClosestGuides, ClosestGuide);
-			SetHairStrandAttribute(OutHairDescription, StrandID, HairAttribute::Strand::GuideWeights, GuideWeight);
+			if (bHasTripleWeights)
+			{
+				const FVector3f ClosestGuides = FVector3f(Strands.GuideIDs[0], Strands.GuideIDs[1], Strands.GuideIDs[2]);
+				const FVector3f GuideWeights = FVector3f(Strands.GuideWeights[0], Strands.GuideWeights[1], Strands.GuideWeights[2]);
+				SetHairStrandAttribute(OutHairDescription, StrandID, HairAttribute::Strand::ClosestGuides, ClosestGuides);
+				SetHairStrandAttribute(OutHairDescription, StrandID, HairAttribute::Strand::GuideWeights, GuideWeights);
+			}
+			else if (bHasSingleWeights)
+			{
+				int ClosestGuide = Strands.GuideIDs[0];
+				float GuideWeight = Strands.GuideWeights[0];
+				SetHairStrandAttribute(OutHairDescription, StrandID, HairAttribute::Strand::ClosestGuides, ClosestGuide);
+				SetHairStrandAttribute(OutHairDescription, StrandID, HairAttribute::Strand::GuideWeights, GuideWeight);
+			}
 		}
 
 		for (const FEditableHairStrandControlPoint& Point : Strands.ControlPoints)
@@ -158,9 +189,18 @@ static void EditableGroomGroupToHairDescription(FHairDescription& OutHairDescrip
 			FVertexID VertexID = (FVertexID)VertexIndex++;
 			SetHairVertexAttribute(OutHairDescription, VertexID, HairAttribute::Vertex::Position,  Point.Position);
 			SetHairVertexAttribute(OutHairDescription, VertexID, HairAttribute::Vertex::Width,     Point.Radius * 2.f);
-			SetHairVertexAttribute(OutHairDescription, VertexID, HairAttribute::Vertex::Color,     FVector3f(Point.BaseColor.R, Point.BaseColor.G, Point.BaseColor.B));
-			SetHairVertexAttribute(OutHairDescription, VertexID, HairAttribute::Vertex::Roughness, Point.Roughness);
-			SetHairVertexAttribute(OutHairDescription, VertexID, HairAttribute::Vertex::AO,        Point.AO);
+			if(Point.bHasColor)
+			{ 
+				SetHairVertexAttribute(OutHairDescription, VertexID, HairAttribute::Vertex::Color,     FVector3f(Point.BaseColor.R, Point.BaseColor.G, Point.BaseColor.B));
+			}
+			if (Point.bHasRoughness)
+			{
+				SetHairVertexAttribute(OutHairDescription, VertexID, HairAttribute::Vertex::Roughness, Point.Roughness);
+			}
+			if (Point.bHasAO)
+			{
+				SetHairVertexAttribute(OutHairDescription, VertexID, HairAttribute::Vertex::AO,        Point.AO);
+			}
 		}
 	}
 }
diff --git a/Engine/Plugins/Runtime/HairStrands/Source/HairStrandsCore/Public/GroomEdit.h b/Engine/Plugins/Runtime/HairStrands/Source/HairStrandsCore/Public/GroomEdit.h
--- a/Engine/Plugins/Runtime/HairStrands/Source/HairStrandsCore/Public/GroomEdit.h
+++ b/Engine/Plugins/Runtime/HairStrands/Source/HairStrandsCore/Public/GroomEdit.h
@@ -18,6 +18,10 @@ struct FEditableHairStrandControlPoint
 	FLinearColor BaseColor;
 	float Roughness;
 	float AO;
+
+	uint8 bHasColor : 1 = false;
+	uint8 bHasRoughness : 1 = false;
+	uint8 bHasAO : 1 = false;
 };
 
 struct FEditableHairStrand
@@ -31,6 +35,11 @@ struct FEditableHairStrand
 	// Pre-computed Simulation
 	uint32 GuideIDs[3];
 	float GuideWeights[3];
+
+	uint8 bHasStrandID : 1 = false;
+	uint8 bHasClumpID : 1 = false;
+	uint8 bHasClosestGuide : 1 = false;
+	uint8 bHasRootUV : 1 = false;
 };
 
 ///////////////////////////////////////////////////////////////////////////////////////////////////
@@ -47,6 +56,9 @@ struct FEditableHairGuide
 
 	uint32 GuideID;
 	FVector2f RootUV;
+
+	uint8 bHasGuideID : 1 = false;
+	uint8 bHasRootUV : 1 = false;
 };
 
 ///////////////////////////////////////////////////////////////////////////////////////////////////
