Fix crash on CitySample.

Always use the same data source for skel. mesh data (SkinCache or manual SkinCache, not a mix of both).

diff --git a/Engine/Plugins/Runtime/HairStrands/Source/HairStrandsCore/Private/GroomManager.cpp b/Engine/Plugins/Runtime/HairStrands/Source/HairStrandsCore/Private/GroomManager.cpp
--- a/Engine/Plugins/Runtime/HairStrands/Source/HairStrandsCore/Private/GroomManager.cpp
+++ b/Engine/Plugins/Runtime/HairStrands/Source/HairStrandsCore/Private/GroomManager.cpp
@@ -646,8 +646,7 @@ static FCachedGeometry GetCacheGeometryForHair(
 			if (const FSkeletalMeshSceneProxy* SceneProxy = static_cast<const FSkeletalMeshSceneProxy*>(PrimitiveSceneInfo->Proxy))
 			{
 				// 1. Get cached geometry from GPU SkinCache, if enabled
-				const bool bValidCachedGeometry = SceneProxy->GetCachedGeometry(GraphBuilder, Out);
-				if (bValidCachedGeometry && Out.Sections.Num() > 0)
+				if (Out.Sections.Num() > 0)
 				{
 					OutHairGeometryCache.AddDebug(Instance, SceneProxy, Out, PositionUpdateType, FHairGeometryCache::ECacheType::SkinCache);
 				}
