Fix typo causing the skin cache data to no longer be fetched.

diff --git a/Engine/Plugins/Runtime/HairStrands/Source/HairStrandsCore/Private/GroomManager.cpp b/Engine/Plugins/Runtime/HairStrands/Source/HairStrandsCore/Private/GroomManager.cpp
--- a/Engine/Plugins/Runtime/HairStrands/Source/HairStrandsCore/Private/GroomManager.cpp
+++ b/Engine/Plugins/Runtime/HairStrands/Source/HairStrandsCore/Private/GroomManager.cpp
@@ -646,6 +646,7 @@ static FCachedGeometry GetCacheGeometryForHair(
 			if (const FSkeletalMeshSceneProxy* SceneProxy = static_cast<const FSkeletalMeshSceneProxy*>(PrimitiveSceneInfo->Proxy))
 			{
 				// 1. Get cached geometry from GPU SkinCache, if enabled
+				SceneProxy->GetCachedGeometry(GraphBuilder, Out);
 				if (Out.Sections.Num() > 0)
 				{
 					OutHairGeometryCache.AddDebug(Instance, SceneProxy, Out, PositionUpdateType, FHairGeometryCache::ECacheType::SkinCache);
