Fix groom dataflow nodes attributes duplicate in between points/vertices groups

diff --git a/Engine/Plugins/Runtime/HairStrands/Source/HairStrandsDataflow/Private/AttachGuidesRootsNode.cpp b/Engine/Plugins/Runtime/HairStrands/Source/HairStrandsDataflow/Private/AttachGuidesRootsNode.cpp
--- a/Engine/Plugins/Runtime/HairStrands/Source/HairStrandsDataflow/Private/AttachGuidesRootsNode.cpp
+++ b/Engine/Plugins/Runtime/HairStrands/Source/HairStrandsDataflow/Private/AttachGuidesRootsNode.cpp
@@ -6,6 +6,17 @@
 
 #include UE_INLINE_GENERATED_CPP_BY_NAME(AttachGuidesRootsNode)
 
+namespace UE::Groom::Private
+{
+	FCollectionAttributeKey GetKinematicWeightsKey()
+	{
+		FCollectionAttributeKey Key;
+		Key.Group =  UE::Groom::FGroomGuidesFacade::VerticesGroup.ToString();
+		Key.Attribute = UE::Groom::FGroomGuidesFacade::PointKinematicWeightsAttribute.ToString();
+		return Key;
+	}
+}
+
 void FAttachGuidesRootsDataflowNode::Evaluate(UE::Dataflow::FContext& Context, const FDataflowOutput* Out) const
 {
 	if (Out->IsA<FManagedArrayCollection>(&Collection))
@@ -16,7 +27,7 @@ void FAttachGuidesRootsDataflowNode::Evaluate(UE::Dataflow::FContext& Context, c
 
 		if(GuidesFacade.IsValid())
 		{
-			TArray<float> KinematicWeights = GuidesFacade.GetVertexKinematicWeights();
+			TArray<float> KinematicWeights = GuidesFacade.GetPointKinematicWeights();
 			
 			const int32 LocalIndex = GroupIndex;
 			ParallelFor(GuidesFacade.GetNumPoints(), [&KinematicWeights, &GuidesFacade, LocalIndex](int32 PointIndex)
@@ -36,10 +47,15 @@ void FAttachGuidesRootsDataflowNode::Evaluate(UE::Dataflow::FContext& Context, c
 					}
 				}
 			}, EParallelForFlags::None);
-			GuidesFacade.SetVertexKinematicWeights(KinematicWeights);
+			GuidesFacade.SetPointKinematicWeights(KinematicWeights);
 		}
 		
 		SetValue(Context, MoveTemp(GroomCollection), &Collection);
 	}
+	else if (Out->IsA<FCollectionAttributeKey>(&KinematicWeightsKey))
+	{
+		SetValue(Context, UE::Groom::Private::GetKinematicWeightsKey(), &KinematicWeightsKey);
+	}
 }
 
+
diff --git a/Engine/Plugins/Runtime/HairStrands/Source/HairStrandsDataflow/Private/BuildGroomSkinningNodes.cpp b/Engine/Plugins/Runtime/HairStrands/Source/HairStrandsDataflow/Private/BuildGroomSkinningNodes.cpp
--- a/Engine/Plugins/Runtime/HairStrands/Source/HairStrandsDataflow/Private/BuildGroomSkinningNodes.cpp
+++ b/Engine/Plugins/Runtime/HairStrands/Source/HairStrandsDataflow/Private/BuildGroomSkinningNodes.cpp
@@ -230,7 +230,7 @@ namespace UE::Groom::Private
 	}
 }
 
-void FBuildGuidesSkinningDataflowNode::Evaluate(UE::Dataflow::FContext& Context, const FDataflowOutput* Out) const
+void FTransferSkinWeightsGroomNode::Evaluate(UE::Dataflow::FContext& Context, const FDataflowOutput* Out) const
 {
 	if (Out->IsA<FManagedArrayCollection>(&Collection))
 	{
@@ -244,117 +244,37 @@ void FBuildGuidesSkinningDataflowNode::Evaluate(UE::Dataflow::FContext& Context,
 		{
 			UE::Groom::Private::BuildSkinningData<UE::Groom::FGroomStrandsFacade>(GroomCollection, GroupIndex, SkeletalMesh, LODIndex, RelativeTransform);
 		}
-		
-		SetValue(Context, MoveTemp(GroomCollection), &Collection);
-	}
-}
-
-TArray<UE::Dataflow::FRenderingParameter> FBuildGuidesSkinningDataflowNode::GetRenderParametersImpl() const
-{
-	if(CurvesType == EGroomCollectionType::Guides)
-	{
-		return { {TEXT("GuidesRender"), FName("FGroomCollection"), {TEXT("Collection")}}};
-	}
-	else
-	{
-		return { {TEXT("StrandsRender"), FName("FGroomCollection"), {TEXT("Collection")}}};
-	}
-}
-
-void FUnpackSkinWeightsDataflowNode::Evaluate(UE::Dataflow::FContext& Context, const FDataflowOutput* Out) const
-{
-	using namespace UE::Dataflow;
-	
-	// Get the pin value if plugged
-	FCollectionAttributeKey BoneIndicesKeyValue = GetBoneIndicesKey();
-	FCollectionAttributeKey BoneWeightsKeyValue = GetBoneWeightsKey();
-
-	if (Out->IsA<FManagedArrayCollection>(&Collection))
-	{
-		// Evaluate in collection
-		FManagedArrayCollection GroomCollection = GetValue<FManagedArrayCollection>(Context, &Collection);
-
-		if(CurvesType == EGroomCollectionType::Guides)
-		{
-			UE::Groom::Private::ExtractSkinningData<UE::Groom::FGroomGuidesFacade>(GroomCollection, BoneIndicesKeyValue, BoneWeightsKeyValue);
-		}
-		else
-		{
-			UE::Groom::Private::ExtractSkinningData<UE::Groom::FGroomStrandsFacade>(GroomCollection, BoneIndicesKeyValue, BoneWeightsKeyValue);
-		}
-
 		SetValue(Context, MoveTemp(GroomCollection), &Collection);
 	}
 	else if (Out->IsA<FCollectionAttributeKey>(&BoneIndicesKey))
 	{
-		SetValue(Context, MoveTemp(BoneIndicesKeyValue), &BoneIndicesKey);
+		SetValue(Context, GetBoneIndicesKey(), &BoneIndicesKey);
 	}
 	else if (Out->IsA<FCollectionAttributeKey>(&BoneWeightsKey))
 	{
-		SetValue(Context, MoveTemp(BoneWeightsKeyValue), &BoneWeightsKey);
-	}
-}
-
-TArray<UE::Dataflow::FRenderingParameter> FUnpackSkinWeightsDataflowNode::GetRenderParametersImpl() const
-{
-	if(CurvesType == EGroomCollectionType::Guides)
-	{
-		return { {TEXT("GuidesRender"), FName("FGroomCollection"), {TEXT("Collection")}}};
-	}
-	else
-	{
-		return { {TEXT("StrandsRender"), FName("FGroomCollection"), {TEXT("Collection")}}};
+		SetValue(Context, GetBoneWeightsKey(), &BoneWeightsKey);
 	}
 }
 
-
-FCollectionAttributeKey FUnpackSkinWeightsDataflowNode::GetBoneIndicesKey() const
+FCollectionAttributeKey FTransferSkinWeightsGroomNode::GetBoneIndicesKey() const
 {
-	// Get the pin value if plugged
 	FCollectionAttributeKey Key;
 	Key.Group = (CurvesType == EGroomCollectionType::Guides) ? UE::Groom::FGroomGuidesFacade::VerticesGroup.ToString() :
 		UE::Groom::FGroomStrandsFacade::VerticesGroup.ToString();
-	Key.Attribute = BoneIndicesName;
+	Key.Attribute = UE::Groom::FGroomGuidesFacade::PointBoneIndicesAttribute.ToString();
 	return Key;
 }
 
-FCollectionAttributeKey FUnpackSkinWeightsDataflowNode::GetBoneWeightsKey() const
+FCollectionAttributeKey FTransferSkinWeightsGroomNode::GetBoneWeightsKey() const
 {
-	// Get the pin value if plugged
 	FCollectionAttributeKey Key;
 	Key.Group = (CurvesType == EGroomCollectionType::Guides) ? UE::Groom::FGroomGuidesFacade::VerticesGroup.ToString() :
 		UE::Groom::FGroomStrandsFacade::VerticesGroup.ToString();
-	Key.Attribute = BoneWeightsName;
+	Key.Attribute = UE::Groom::FGroomGuidesFacade::PointBoneWeightsAttribute.ToString();
 	return Key;
 }
 
-void FPackSkinWeightsDataflowNode::Evaluate(UE::Dataflow::FContext& Context, const FDataflowOutput* Out) const
-{
-	using namespace UE::Dataflow;
-	
-	// Get the pin value if plugged
-	FCollectionAttributeKey BoneIndicesKeyValue = GetBoneIndicesKey(Context);
-	FCollectionAttributeKey BoneWeightsKeyValue = GetBoneWeightsKey(Context);
-
-	if (Out->IsA<FManagedArrayCollection>(&Collection))
-	{
-		// Evaluate in collection
-		FManagedArrayCollection GroomCollection = GetValue<FManagedArrayCollection>(Context, &Collection);
-
-		if(CurvesType == EGroomCollectionType::Guides)
-		{
-			UE::Groom::Private::ReportSkinningData<UE::Groom::FGroomGuidesFacade>(GroomCollection, BoneIndicesKeyValue, BoneWeightsKeyValue);
-		}
-		else
-		{
-			UE::Groom::Private::ReportSkinningData<UE::Groom::FGroomStrandsFacade>(GroomCollection, BoneIndicesKeyValue, BoneWeightsKeyValue);
-		}
-
-		SetValue(Context, MoveTemp(GroomCollection), &Collection);
-	}
-}
-
-TArray<UE::Dataflow::FRenderingParameter> FPackSkinWeightsDataflowNode::GetRenderParametersImpl() const
+TArray<UE::Dataflow::FRenderingParameter> FTransferSkinWeightsGroomNode::GetRenderParametersImpl() const
 {
 	if(CurvesType == EGroomCollectionType::Guides)
 	{
@@ -366,34 +286,3 @@ TArray<UE::Dataflow::FRenderingParameter> FPackSkinWeightsDataflowNode::GetRende
 	}
 }
 
-FCollectionAttributeKey FPackSkinWeightsDataflowNode::GetBoneIndicesKey(UE::Dataflow::FContext& Context) const
-{
-	// Get the pin value if plugged
-	FCollectionAttributeKey Key = GetValue(Context, &BoneIndicesKey, BoneIndicesKey);
-
-	// If nothing set used the local value
-	if(Key.Attribute.IsEmpty() && Key.Group.IsEmpty())
-	{
-		Key.Group = (CurvesType == EGroomCollectionType::Guides) ? UE::Groom::FGroomGuidesFacade::VerticesGroup.ToString() :
-																   UE::Groom::FGroomStrandsFacade::VerticesGroup.ToString();
-		Key.Attribute = BoneIndicesName;
-	}
-	return Key;
-}
-
-FCollectionAttributeKey FPackSkinWeightsDataflowNode::GetBoneWeightsKey(UE::Dataflow::FContext& Context) const
-{
-	// Get the pin value if plugged
-	FCollectionAttributeKey Key = GetValue(Context, &BoneWeightsKey, BoneWeightsKey);
-
-	// If nothing set used the local value
-	if(Key.Attribute.IsEmpty() && Key.Group.IsEmpty())
-	{
-		Key.Group = (CurvesType == EGroomCollectionType::Guides) ? UE::Groom::FGroomGuidesFacade::VerticesGroup.ToString() :
-																   UE::Groom::FGroomStrandsFacade::VerticesGroup.ToString();
-		Key.Attribute = BoneWeightsName;
-	}
-	return Key;
-}
-
-
diff --git a/Engine/Plugins/Runtime/HairStrands/Source/HairStrandsDataflow/Private/BuildGuidesLODsNode.cpp b/Engine/Plugins/Runtime/HairStrands/Source/HairStrandsDataflow/Private/BuildGuidesLODsNode.cpp
--- a/Engine/Plugins/Runtime/HairStrands/Source/HairStrandsDataflow/Private/BuildGuidesLODsNode.cpp
+++ b/Engine/Plugins/Runtime/HairStrands/Source/HairStrandsDataflow/Private/BuildGuidesLODsNode.cpp
@@ -86,6 +86,22 @@ namespace UE::Groom::Private
 			CurveOffset = GuidesFacade.GetObjectCurveOffsets()[ObjectIndex];
 		}
 	}
+
+	FCollectionAttributeKey GetCurveParentsKey()
+    {
+    	FCollectionAttributeKey Key;
+    	Key.Group =  UE::Groom::FGroomGuidesFacade::CurvesGroup.ToString();
+    	Key.Attribute = UE::Groom::FGroomGuidesFacade::CurveParentIndicesAttribute.ToString();
+    	return Key;
+    }
+
+	FCollectionAttributeKey GetCurveLodsKey()
+	{
+		FCollectionAttributeKey Key;
+		Key.Group =  UE::Groom::FGroomGuidesFacade::CurvesGroup.ToString();
+		Key.Attribute = UE::Groom::FGroomGuidesFacade::CurveLodIndicesAttribute.ToString();
+		return Key;
+	}
 }
 
 void FBuildGuidesLODsDataflowNode::Evaluate(UE::Dataflow::FContext& Context, const FDataflowOutput* Out) const
@@ -103,5 +119,13 @@ void FBuildGuidesLODsDataflowNode::Evaluate(UE::Dataflow::FContext& Context, con
 		
 		SetValue(Context, MoveTemp(GroomCollection), &Collection);
 	}
+	else if (Out->IsA<FCollectionAttributeKey>(&CurveParentsKey))
+	{
+		SetValue(Context, UE::Groom::Private::GetCurveParentsKey(), &CurveParentsKey);
+	}
+	else if (Out->IsA<FCollectionAttributeKey>(&CurveLodsKey))
+	{
+		SetValue(Context, UE::Groom::Private::GetCurveLodsKey(), &CurveLodsKey);
+	}
 }
 
diff --git a/Engine/Plugins/Runtime/HairStrands/Source/HairStrandsDataflow/Private/GetGroomAssetNode.cpp b/Engine/Plugins/Runtime/HairStrands/Source/HairStrandsDataflow/Private/GetGroomAssetNode.cpp
--- a/Engine/Plugins/Runtime/HairStrands/Source/HairStrandsDataflow/Private/GetGroomAssetNode.cpp
+++ b/Engine/Plugins/Runtime/HairStrands/Source/HairStrandsDataflow/Private/GetGroomAssetNode.cpp
@@ -11,8 +11,7 @@
 
 namespace UE::Groom::Private
 {
-	template<typename FacadeType>
-	FORCEINLINE const UGroomAsset* BuildGroomCollection(Dataflow::FContext& Context, const UGroomAsset* NodeGroom, FManagedArrayCollection& GroomCollection, FEditableGroom& EditGroom)
+	FORCEINLINE const UGroomAsset* FillEditableAsset(Dataflow::FContext& Context, const UGroomAsset* NodeGroom, FEditableGroom& EditGroom)
 	{
 		const UGroomAsset* GroomAsset = NodeGroom;
 		if (!GroomAsset)
@@ -26,38 +25,37 @@ namespace UE::Groom::Private
 		if (GroomAsset)
 		{
 			ConvertFromGroomAsset(const_cast<UGroomAsset*>(GroomAsset), &EditGroom, false, false, false);
+		}
+		return GroomAsset;
+	}
+	
+	template<typename FacadeType>
+	FORCEINLINE void BuildGroomCollection(FManagedArrayCollection& GroomCollection, FEditableGroom& EditGroom, const FString& GroomName)
+	{
+		TArray<FVector3f> PointRestPositions;
+		TArray<int32> ObjectCurveOffsets;
+		TArray<int32> CurvePointOffsets;
+		TArray<FString> ObjectGroupNames;
 
-			TArray<FVector3f> PointRestPositions;
-			TArray<int32> ObjectCurveOffsets;
-			TArray<int32> CurvePointOffsets;
-			TArray<FString> ObjectGroupNames;
-
-			uint32 GroupIndex = 0;
-			for(FEditableGroomGroup& GroomGroup : EditGroom.Groups)
+		uint32 GroupIndex = 0;
+		for(FEditableGroomGroup& GroomGroup : EditGroom.Groups)
+		{
+			for( const typename FacadeType::FEditableType& EditType : FacadeType::GetEditableGroom(GroomGroup))
 			{
-				
-				for( const typename FacadeType::FEditableType& EditType : FacadeType::GetEditableGroom(GroomGroup))
+				for(auto& GuidePoint : EditType.ControlPoints)
 				{
-					for(auto& GuidePoint : EditType.ControlPoints)
-					{
-						PointRestPositions.Add(GuidePoint.Position);
-					}
-					CurvePointOffsets.Add(PointRestPositions.Num());
+					PointRestPositions.Add(GuidePoint.Position);
 				}
-				ObjectCurveOffsets.Add(CurvePointOffsets.Num());
-				
-				const FString GroupName = GroomAsset->GetName() + TEXT("_") + FacadeType::GroupPrefix.ToString();
-				ObjectGroupNames.Add(GroupName);
-				++GroupIndex;
+				CurvePointOffsets.Add(PointRestPositions.Num());
 			}
-			FacadeType CurvesFacade(GroomCollection);
-			CurvesFacade.InitGroomCollection(PointRestPositions, CurvePointOffsets, ObjectCurveOffsets, ObjectGroupNames);
-
-			const FString GroomName = GroomAsset->GetName();
-
-			return GroomAsset;
+			ObjectCurveOffsets.Add(CurvePointOffsets.Num());
+			
+			const FString GroupName = GroomName + TEXT("_") + FacadeType::GroupPrefix.ToString();
+			ObjectGroupNames.Add(GroupName);
+			++GroupIndex;
 		}
-		return nullptr;
+		FacadeType CurvesFacade(GroomCollection);
+		CurvesFacade.InitGroomCollection(PointRestPositions, CurvePointOffsets, ObjectCurveOffsets, ObjectGroupNames);
 	}
 }
 
@@ -67,27 +65,24 @@ void FGetGroomAssetDataflowNode::Evaluate(UE::Dataflow::FContext& Context, const
 	{
 		FManagedArrayCollection GroomCollection;
 		FEditableGroom EditGroom;
-		const UGroomAsset* LocalGroom = nullptr;
-		if(CurvesType == EGroomCollectionType::Guides)
-		{
-			LocalGroom = UE::Groom::Private::BuildGroomCollection<UE::Groom::FGroomGuidesFacade>(
-				Context, GroomAsset, GroomCollection, EditGroom);
-		}
-		else if(CurvesType == EGroomCollectionType::Strands)
+		if(const UGroomAsset* LocalGroom = UE::Groom::Private::FillEditableAsset(Context, GroomAsset, EditGroom))
 		{
-			LocalGroom = UE::Groom::Private::BuildGroomCollection<UE::Groom::FGroomStrandsFacade>(
-				Context, GroomAsset, GroomCollection, EditGroom);
-		}
+			// Fill the guides facade
+			UE::Groom::Private::BuildGroomCollection<UE::Groom::FGroomGuidesFacade>(
+					GroomCollection, EditGroom, LocalGroom->GetName());
 
-		if(LocalGroom)
-		{
+			// Fill the strands facade
+			UE::Groom::Private::BuildGroomCollection<UE::Groom::FGroomStrandsFacade>(
+					GroomCollection, EditGroom, LocalGroom->GetName());
+
+			// Add guides points samples for future resampling
 			TArray<int32> ObjectPointSamples, CurveStrandIndices;
 			const TArray<FHairGroupsPhysics>& GroupsPhysics = LocalGroom->GetHairGroupsPhysics();
 			for(const FHairGroupsPhysics& GroupPhysics : GroupsPhysics)
 			{
 				ObjectPointSamples.Add(static_cast<uint8>(GroupPhysics.StrandsParameters.StrandsSize));
 			}
-				
+			
 			UE::Groom::FGroomGuidesFacade GuidesFacade(GroomCollection);
 			if(GroomCollection.NumElements(UE::Groom::FGroomGuidesFacade::ObjectsGroup) == 0)
 			{
@@ -95,6 +90,7 @@ void FGetGroomAssetDataflowNode::Evaluate(UE::Dataflow::FContext& Context, const
 			}
 			GuidesFacade.SetObjectPointSamples(ObjectPointSamples);
 		}
+
 		SetValue(Context, MoveTemp(GroomCollection), &Collection);
 	}
 }
diff --git a/Engine/Plugins/Runtime/HairStrands/Source/HairStrandsDataflow/Private/GetGroomAttributesNodes.cpp b/Engine/Plugins/Runtime/HairStrands/Source/HairStrandsDataflow/Private/GetGroomAttributesNodes.cpp
new file mode 100644
--- /dev/null
+++ b/Engine/Plugins/Runtime/HairStrands/Source/HairStrandsDataflow/Private/GetGroomAttributesNodes.cpp
@@ -0,0 +1,47 @@
+// Copyright Epic Games, Inc. All Rights Reserved.
+
+#include "GetGroomAttributesNodes.h"
+#include "GroomCollectionFacades.h"
+#include "Dataflow/DataflowObjectInterface.h"
+
+#include UE_INLINE_GENERATED_CPP_BY_NAME(GetGroomAttributesNodes)
+
+void FGetGroomAttributesDataflowNode::Evaluate(UE::Dataflow::FContext& Context, const FDataflowOutput* Out) const
+{
+	if (Out->IsA<FCollectionAttributeKey>(&AttributeKey))
+	{
+		FCollectionAttributeKey Key;
+
+		const FString VertexGroup = (CurvesType == EGroomCollectionType::Guides) ? UE::Groom::FGroomGuidesFacade::VerticesGroup.ToString() :
+								UE::Groom::FGroomStrandsFacade::VerticesGroup.ToString();
+		const FString CurveGroup = (CurvesType == EGroomCollectionType::Guides) ? UE::Groom::FGroomGuidesFacade::CurvesGroup.ToString() :
+								UE::Groom::FGroomStrandsFacade::CurvesGroup.ToString();
+		if(AttributeType == EGroomAttributeType::KinematicWeights)
+		{ 
+			Key.Group = VertexGroup;
+			Key.Attribute = UE::Groom::FGroomGuidesFacade::PointKinematicWeightsAttribute.ToString();
+		}
+		else if (AttributeType == EGroomAttributeType::BoneIndices)
+		{
+			Key.Group = VertexGroup;
+			Key.Attribute = UE::Groom::FGroomGuidesFacade::PointBoneIndicesAttribute.ToString();
+		}
+		else if (AttributeType == EGroomAttributeType::BoneWeights)
+		{
+			Key.Group = VertexGroup;
+			Key.Attribute = UE::Groom::FGroomGuidesFacade::PointBoneWeightsAttribute.ToString();
+		}
+		else if (AttributeType == EGroomAttributeType::CurveLods)
+		{
+			Key.Group = CurveGroup;
+			Key.Attribute = UE::Groom::FGroomGuidesFacade::CurveLodIndicesAttribute.ToString();
+		}
+		else if (AttributeType == EGroomAttributeType::CurveParents)
+		{
+			Key.Group = CurveGroup;
+			Key.Attribute = UE::Groom::FGroomGuidesFacade::CurveParentIndicesAttribute.ToString();
+		}
+		SetValue(Context, MoveTemp(Key), &AttributeKey);
+	}
+}
+
diff --git a/Engine/Plugins/Runtime/HairStrands/Source/HairStrandsDataflow/Private/GetGuidesAttributesNodes.cpp b/Engine/Plugins/Runtime/HairStrands/Source/HairStrandsDataflow/Private/GetGuidesAttributesNodes.cpp
deleted file mode 100644
diff --git a/Engine/Plugins/Runtime/HairStrands/Source/HairStrandsDataflow/Private/GroomAssetTerminalNode.cpp b/Engine/Plugins/Runtime/HairStrands/Source/HairStrandsDataflow/Private/GroomAssetTerminalNode.cpp
--- a/Engine/Plugins/Runtime/HairStrands/Source/HairStrandsDataflow/Private/GroomAssetTerminalNode.cpp
+++ b/Engine/Plugins/Runtime/HairStrands/Source/HairStrandsDataflow/Private/GroomAssetTerminalNode.cpp
@@ -5,6 +5,7 @@
 #include "AssetCompilingManager.h"
 #include "GroomCollectionFacades.h"
 #include "GroomEdit.h"
+#include "Dataflow/DataflowConnectionTypes.h"
 #include "Dataflow/DataflowObjectInterface.h"
 
 #include UE_INLINE_GENERATED_CPP_BY_NAME(GroomAssetTerminalNode)
@@ -48,12 +49,12 @@ namespace UE::Groom::Private
 	}
 
 	FORCEINLINE void CopyCollectionAttributes(const FManagedArrayCollection* InputCollection, FManagedArrayCollection* OutputCollection,
-		const TArray<TTuple<FName, FName>> AttributesToCopy = TArray<TTuple<FName, FName>>())
+		const TArray<FCollectionAttributeKey>& AttributesToCopy = TArray<FCollectionAttributeKey>())
 	{
-		for(const TTuple<FName, FName>& AttributeToCopy : AttributesToCopy)
+		for(const FCollectionAttributeKey& AttributeToCopy : AttributesToCopy)
 		{
-			const FName& AttributeName = AttributeToCopy.Value;
-			const FName& GroupName = AttributeToCopy.Key;
+			const FName AttributeName(AttributeToCopy.Attribute);
+			const FName GroupName(AttributeToCopy.Group);
 
 			if (InputCollection->HasGroup(GroupName))
 			{
@@ -168,6 +169,37 @@ namespace UE::Groom::Private
 				Cast<USkeletalMesh>(ObjectSkeletalMeshes[GroupIndex]), ObjectMeshLODs[GroupIndex]);
 		}
 	}
+	template<typename FacadeType>
+	FORCEINLINE void TransferCollectionAttributes(const FManagedArrayCollection& InCollection, FManagedArrayCollection* OutCollection, const FacadeType& GroomFacade,
+		const TArray<FCollectionAttributeKey>& ExternalAttributes, const TArray<FCollectionAttributeKey>& InternalAttributes)
+	{
+		// Register all the internal attributes to be copied
+		TArray<FCollectionAttributeKey> AttributesToCopy;
+
+		auto AddAttributeKeys = [&AttributesToCopy](const TArray<FCollectionAttributeKey>& AttributeKeys)
+		{
+			for (const FCollectionAttributeKey& AttributeKey : AttributeKeys)
+			{
+				if (AttributeKey.Group == FacadeType::CurvesGroup || AttributeKey.Group == FacadeType::ObjectsGroup ||
+					AttributeKey.Group == FacadeType::PointsGroup || AttributeKey.Group == FacadeType::EdgesGroup)
+				{
+					AttributesToCopy.Add(AttributeKey);
+				}
+			}
+		};
+
+		AddAttributeKeys(ExternalAttributes);
+		AddAttributeKeys(InternalAttributes);
+
+		// Copy attributes from input collection
+		UE::Groom::Private::CopyCollectionAttributes(&InCollection, OutCollection, AttributesToCopy);
+				
+		// Skip default vertex attributes as not defined by the user
+		const TArray<FName> AttributesToSkip = {FacadeType::VertexLinearColorsAttribute};
+				
+		// Transfer vertices weight maps onto the points to be stored onto the rest collection
+		UE::Groom::Private::TransferVerticesAttributes<FacadeType>(InCollection, OutCollection, GroomFacade.GetNumPoints(), AttributesToSkip);
+	}
 }
 
 void FGroomAssetTerminalDataflowNode::SetAssetValue(TObjectPtr<UObject> Asset, UE::Dataflow::FContext& Context) const
@@ -182,27 +214,13 @@ void FGroomAssetTerminalDataflowNode::SetAssetValue(TObjectPtr<UObject> Asset, U
 
 		if(StrandsFacade.IsValid())
 		{
-			// Skip default vertex attributes as not defined by the user
-			const TArray<FName> AttributesToSkip = {UE::Groom::FGroomStrandsFacade::VertexLinearColorsAttribute};
-			
-			// Transfer vertices weight maps onto the points to be stored onto the rest collection
-			UE::Groom::Private::TransferVerticesAttributes<UE::Groom::FGroomStrandsFacade>(InCollection, OutCollection, StrandsFacade.GetNumPoints(), AttributesToSkip);
+			UE::Groom::Private::TransferCollectionAttributes<UE::Groom::FGroomStrandsFacade>(InCollection, OutCollection, StrandsFacade, AttributeKeys, {});
 		}
 		if(GuidesFacade.IsValid())
 		{
-			// Only copy the attributes not stored already in the asset
-			const TArray<TTuple<FName, FName>> AttributesToCopy = {
-				{UE::Groom::FGroomGuidesFacade::CurvesGroup, UE::Groom::FGroomGuidesFacade::CurveParentIndicesAttribute},
-				{UE::Groom::FGroomGuidesFacade::CurvesGroup, UE::Groom::FGroomGuidesFacade::CurveLodIndicesAttribute}};
-
-			// Copy attributes from input collection
-			UE::Groom::Private::CopyCollectionAttributes(&InCollection, OutCollection, AttributesToCopy);
-
-			// Skip default vertex attributes as not defined by the user
-			const TArray<FName> AttributesToSkip = {UE::Groom::FGroomGuidesFacade::VertexLinearColorsAttribute};
-			
-			// Transfer vertices weight maps onto the points to be stored onto the rest collection (VertexKinematicWeights should be included here)
-			UE::Groom::Private::TransferVerticesAttributes<UE::Groom::FGroomGuidesFacade>(InCollection, OutCollection, GuidesFacade.GetNumPoints(), AttributesToSkip);
+			UE::Groom::Private::TransferCollectionAttributes<UE::Groom::FGroomGuidesFacade>(InCollection, OutCollection, GuidesFacade, AttributeKeys,
+				{{UE::Groom::FGroomGuidesFacade::CurveParentIndicesAttribute.ToString(), UE::Groom::FGroomGuidesFacade::CurvesGroup.ToString()},
+				{UE::Groom::FGroomGuidesFacade::CurveLodIndicesAttribute.ToString(), UE::Groom::FGroomGuidesFacade::CurvesGroup.ToString()}});
 
 			// Build an editable groom asset for the strands
 			FEditableGroom EditGroom;
@@ -249,5 +267,81 @@ void FGroomAssetTerminalDataflowNode::Evaluate(UE::Dataflow::FContext& Context,
 	SetValue(Context, InCollection, &Collection);
 }
 
+UE::Dataflow::TConnectionReference<FCollectionAttributeKey> FGroomAssetTerminalDataflowNode::GetConnectionReference(int32 Index) const
+{
+	return { &AttributeKeys[Index], Index, &AttributeKeys };
+}
+
+TArray<UE::Dataflow::FPin> FGroomAssetTerminalDataflowNode::AddPins()
+{
+	const int32 Index = AttributeKeys.AddDefaulted();
+	const FDataflowInput& Input = RegisterInputArrayConnection(GetConnectionReference(Index));
+	return { { UE::Dataflow::FPin::EDirection::INPUT, Input.GetType(), Input.GetName() } };
+}
+
+TArray<UE::Dataflow::FPin> FGroomAssetTerminalDataflowNode::GetPinsToRemove() const
+{
+	const int32 Index = AttributeKeys.Num() - 1;
+	check(AttributeKeys.IsValidIndex(Index));
+	if (const FDataflowInput* const Input = FindInput(GetConnectionReference(Index)))
+	{
+		return { { UE::Dataflow::FPin::EDirection::INPUT, Input->GetType(), Input->GetName() } };
+	}
+	return Super::GetPinsToRemove();
+}
+
+void FGroomAssetTerminalDataflowNode::OnPinRemoved(const UE::Dataflow::FPin& Pin)
+{
+	const int32 Index = AttributeKeys.Num() - 1;
+	check(AttributeKeys.IsValidIndex(Index));
+#if DO_CHECK
+	const FDataflowInput* const Input = FindInput(GetConnectionReference(Index));
+	check(Input);
+	check(Input->GetName() == Pin.Name);
+	check(Input->GetType() == Pin.Type);
+#endif
+	AttributeKeys.SetNum(Index);
+
+	return Super::OnPinRemoved(Pin);
+}
+
+void FGroomAssetTerminalDataflowNode::PostSerialize(const FArchive& Ar)
+{
+	// because we add pins we need to make sure we restore them when loading
+	// to make sure they can get properly reconnected
+
+	if (Ar.IsLoading())
+	{
+		check(AttributeKeys.Num() >= 0);
+		// register new elements from the array as inputs
+		for (int32 Index = 0; Index < AttributeKeys.Num(); ++Index)
+		{
+			FindOrRegisterInputArrayConnection(GetConnectionReference(Index));
+		}
+		if (Ar.IsTransacting())
+		{
+			// if we have more inputs than materials then we need to unregister the inputs 
+			const int32 NumAttributeInputs = (GetNumInputs() - NumOtherInputs);
+			const int32 NumInputs = AttributeKeys.Num();
+			if (NumAttributeInputs > NumInputs)
+			{
+				// Inputs have been removed.
+				// Temporarily expand Collections so we can get connection references.
+				AttributeKeys.SetNum(NumAttributeInputs);
+				for (int32 Index = NumInputs; Index < AttributeKeys.Num(); ++Index)
+				{
+					UnregisterInputConnection(GetConnectionReference(Index));
+				}
+				AttributeKeys.SetNum(NumInputs);
+			}
+		}
+		else
+		{
+			ensureAlways(AttributeKeys.Num() + NumOtherInputs == GetNumInputs());
+		}
+	}
+}
+
+
 
 
diff --git a/Engine/Plugins/Runtime/HairStrands/Source/HairStrandsDataflow/Private/GroomCollectionFacades.cpp b/Engine/Plugins/Runtime/HairStrands/Source/HairStrandsDataflow/Private/GroomCollectionFacades.cpp
--- a/Engine/Plugins/Runtime/HairStrands/Source/HairStrandsDataflow/Private/GroomCollectionFacades.cpp
+++ b/Engine/Plugins/Runtime/HairStrands/Source/HairStrandsDataflow/Private/GroomCollectionFacades.cpp
@@ -42,7 +42,7 @@ namespace UE::Groom
 	// Guides Attributes
 	const FName FGroomGuidesFacade::ObjectMeshLODsAttribute("ObjectMeshLODs");
 	const FName FGroomGuidesFacade::ObjectSkeletalMeshesAttribute("ObjectSkeletalMeshes");
-	const FName FGroomGuidesFacade::VertexKinematicWeightsAttribute("VertexKinematicWeights");
+	const FName FGroomGuidesFacade::PointKinematicWeightsAttribute("PointKinematicWeights");
 	const FName FGroomGuidesFacade::PointBoneIndicesAttribute("PointBoneIndices");
 	const FName FGroomGuidesFacade::PointBoneWeightsAttribute("PointBoneWeights");
 	const FName FGroomGuidesFacade::ObjectPointSamplesAttribute("ObjectPointSamples");
@@ -218,9 +218,9 @@ namespace UE::Groom
 	
 	FGroomGuidesFacade::FGroomGuidesFacade(FManagedArrayCollection& InCollection) :
 		FGroomCollectionFacade<FGroomGuidesFacade>(InCollection),
-		VertexKinematicWeights(InCollection, VertexKinematicWeightsAttribute, VerticesGroup),
-		PointBoneIndices(InCollection, PointBoneIndicesAttribute, PointsGroup),
-		PointBoneWeights(InCollection, PointBoneWeightsAttribute, PointsGroup),
+		PointKinematicWeights(InCollection, PointKinematicWeightsAttribute, VerticesGroup),
+		PointBoneIndices(InCollection, PointBoneIndicesAttribute, VerticesGroup),
+		PointBoneWeights(InCollection, PointBoneWeightsAttribute, VerticesGroup),
 		ObjectPointSamples(InCollection, ObjectPointSamplesAttribute, ObjectsGroup),
 		CurveStrandIndices(InCollection, CurveStrandIndicesAttribute, CurvesGroup),
 		CurveParentIndices(InCollection, CurveParentIndicesAttribute, CurvesGroup),
@@ -231,9 +231,9 @@ namespace UE::Groom
 
 	FGroomGuidesFacade::FGroomGuidesFacade(const FManagedArrayCollection& InCollection) :
 		FGroomCollectionFacade<FGroomGuidesFacade>(InCollection),
-		VertexKinematicWeights(InCollection, VertexKinematicWeightsAttribute, VerticesGroup),
-		PointBoneIndices(InCollection, PointBoneIndicesAttribute, PointsGroup),
-		PointBoneWeights(InCollection, PointBoneWeightsAttribute, PointsGroup),
+		PointKinematicWeights(InCollection, PointKinematicWeightsAttribute, VerticesGroup),
+		PointBoneIndices(InCollection, PointBoneIndicesAttribute, VerticesGroup),
+		PointBoneWeights(InCollection, PointBoneWeightsAttribute, VerticesGroup),
 		ObjectPointSamples(InCollection, ObjectPointSamplesAttribute, ObjectsGroup),
 		CurveStrandIndices(InCollection, CurveStrandIndicesAttribute, CurvesGroup),
 		CurveParentIndices(InCollection, CurveParentIndicesAttribute, CurvesGroup),
@@ -243,14 +243,14 @@ namespace UE::Groom
 
 	bool FGroomGuidesFacade::IsFacadeValid() const
 	{
-		return VertexKinematicWeights.IsValid() && PointBoneIndices.IsValid() && PointBoneWeights.IsValid() && ObjectPointSamples.IsValid() &&
+		return PointKinematicWeights.IsValid() && PointBoneIndices.IsValid() && PointBoneWeights.IsValid() && ObjectPointSamples.IsValid() &&
 			CurveStrandIndices.IsValid() && CurveLodIndices.IsValid() && CurveParentIndices.IsValid();
 	}
 
 	void FGroomGuidesFacade::DefineFacadeSchema()
 	{
 		check(!IsConst());
-		VertexKinematicWeights.Add();
+		PointKinematicWeights.Add();
 		PointBoneWeights.Add();
 		PointBoneIndices.Add();
 		ObjectPointSamples.Add();
@@ -262,7 +262,7 @@ namespace UE::Groom
 	void FGroomGuidesFacade::InitFacadeCollection()
 	{
 		// Reset attributes
-		VertexKinematicWeights.Modify().Fill(0.0f);
+		PointKinematicWeights.Modify().Fill(0.0f);
 		PointBoneIndices.Modify().Fill(FIntVector4::ZeroValue);
 		PointBoneWeights.Modify().Fill(FVector4f::Zero());
 		ObjectPointSamples.Modify().Fill(4);
diff --git a/Engine/Plugins/Runtime/HairStrands/Source/HairStrandsDataflow/Private/HairStrandsDataflowModule.cpp b/Engine/Plugins/Runtime/HairStrands/Source/HairStrandsDataflow/Private/HairStrandsDataflowModule.cpp
--- a/Engine/Plugins/Runtime/HairStrands/Source/HairStrandsDataflow/Private/HairStrandsDataflowModule.cpp
+++ b/Engine/Plugins/Runtime/HairStrands/Source/HairStrandsDataflow/Private/HairStrandsDataflowModule.cpp
@@ -7,7 +7,7 @@
 #include "GetGroomAssetNode.h"
 #include "GroomCollectionFacades.h"
 #include "GroomAssetTerminalNode.h"
-#include "GetGuidesAttributesNodes.h"
+#include "GetGroomAttributesNodes.h"
 #include "ResampleGuidesPointsNode.h"
 #include "BuildGroomSkinningNodes.h"
 #include "GenerateGuidesCurvesNode.h"
@@ -52,14 +52,12 @@ void FHairStrandsDataflowModule::StartupModule()
 	DATAFLOW_NODE_REGISTER_CREATION_FACTORY(FGetGroomAssetDataflowNode);
 	DATAFLOW_NODE_REGISTER_CREATION_FACTORY(FGroomAssetTerminalDataflowNode);
 	DATAFLOW_NODE_REGISTER_CREATION_FACTORY(FResampleGuidesPointsDataflowNode);
-	DATAFLOW_NODE_REGISTER_CREATION_FACTORY(FBuildGuidesSkinningDataflowNode);
+	DATAFLOW_NODE_REGISTER_CREATION_FACTORY(FTransferSkinWeightsGroomNode);
 	DATAFLOW_NODE_REGISTER_CREATION_FACTORY(FGenerateGuidesCurvesDataflowNode);
-	DATAFLOW_NODE_REGISTER_CREATION_FACTORY(FGetGuidesKinematicWeightsDataflowNode);
+	DATAFLOW_NODE_REGISTER_CREATION_FACTORY(FGetGroomAttributesDataflowNode);
 	DATAFLOW_NODE_REGISTER_CREATION_FACTORY(FSmoothGuidesCurvesDataflowNode);
 	DATAFLOW_NODE_REGISTER_CREATION_FACTORY(FAttachGuidesRootsDataflowNode);
 	DATAFLOW_NODE_REGISTER_CREATION_FACTORY(FBuildGuidesLODsDataflowNode);
-	DATAFLOW_NODE_REGISTER_CREATION_FACTORY(FUnpackSkinWeightsDataflowNode);
-	DATAFLOW_NODE_REGISTER_CREATION_FACTORY(FPackSkinWeightsDataflowNode);
 	
 	UE::Groom::RegisterRenderingCallbacks();
 	
diff --git a/Engine/Plugins/Runtime/HairStrands/Source/HairStrandsDataflow/Public/AttachGuidesRootsNode.h b/Engine/Plugins/Runtime/HairStrands/Source/HairStrandsDataflow/Public/AttachGuidesRootsNode.h
--- a/Engine/Plugins/Runtime/HairStrands/Source/HairStrandsDataflow/Public/AttachGuidesRootsNode.h
+++ b/Engine/Plugins/Runtime/HairStrands/Source/HairStrandsDataflow/Public/AttachGuidesRootsNode.h
@@ -5,6 +5,7 @@
 #include "CoreMinimal.h"
 #include "GroomAsset.h"
 #include "ChaosLog.h"
+#include "Dataflow/DataflowConnectionTypes.h"
 #include "Dataflow/DataflowCore.h"
 
 #include "AttachGuidesRootsNode.generated.h"
@@ -25,6 +26,7 @@ public:
 	{
 		RegisterInputConnection(&Collection);
 		RegisterOutputConnection(&Collection, &Collection);
+		RegisterOutputConnection(&KinematicWeightsKey);
 	}
 
 	//~ Begin FDataflowNode interface
@@ -38,5 +40,9 @@ public:
 	/** Group index on which the roots will be attached. -1 will attach all the groups */
 	UPROPERTY(EditAnywhere, Category="Groups", meta = (DisplayName = "Group Index"))
 	int32 GroupIndex = INDEX_NONE;
+	
+	/** Point Kinematic weights key to be used in other nodes if necessary */
+	UPROPERTY(meta = (DisplayName = "Kinematic Weights", DataflowOutput))
+	FCollectionAttributeKey KinematicWeightsKey;
 };
 
diff --git a/Engine/Plugins/Runtime/HairStrands/Source/HairStrandsDataflow/Public/BuildGroomSkinningNodes.h b/Engine/Plugins/Runtime/HairStrands/Source/HairStrandsDataflow/Public/BuildGroomSkinningNodes.h
--- a/Engine/Plugins/Runtime/HairStrands/Source/HairStrandsDataflow/Public/BuildGroomSkinningNodes.h
+++ b/Engine/Plugins/Runtime/HairStrands/Source/HairStrandsDataflow/Public/BuildGroomSkinningNodes.h
@@ -11,21 +11,23 @@
 
 #include "BuildGroomSkinningNodes.generated.h"
 
-/** Build the guides skinning by transfering the indices weights from a skelmesh */
+/** Build the guides skinning by transferring the indices weights from a skelmesh */
 USTRUCT(meta = (Experimental, DataflowGroom))
-struct FBuildGuidesSkinningDataflowNode : public FDataflowNode
+struct FTransferSkinWeightsGroomNode : public FDataflowNode
 {
 	GENERATED_BODY()
 	
-	DATAFLOW_NODE_DEFINE_INTERNAL(FBuildGuidesSkinningDataflowNode, "TransferSkinWeights", "Groom", "")
+	DATAFLOW_NODE_DEFINE_INTERNAL(FTransferSkinWeightsGroomNode, "TransferSkinWeights", "Groom", "")
 
 public:
 	
-	FBuildGuidesSkinningDataflowNode(const UE::Dataflow::FNodeParameters& InParam, FGuid InGuid = FGuid::NewGuid())
+	FTransferSkinWeightsGroomNode(const UE::Dataflow::FNodeParameters& InParam, FGuid InGuid = FGuid::NewGuid())
 	: FDataflowNode(InParam, InGuid)
 	{
 		RegisterInputConnection(&Collection);
 		RegisterOutputConnection(&Collection, &Collection);
+		RegisterOutputConnection(&BoneIndicesKey);
+		RegisterOutputConnection(&BoneWeightsKey);
 	}
 
 	//~ Begin FDataflowNode interface
@@ -56,56 +58,15 @@ public:
 	/** Type of curves to use to fill the groom collection (guides/strands) */
 	UPROPERTY(EditAnywhere, Category = "Groom Groups", meta = (DisplayName = "Curves Type"))
 	EGroomCollectionType CurvesType = EGroomCollectionType::Guides;
-};
-
-/** Extract the points skinning attributes to fill the vertices ones */
-USTRUCT(meta = (Experimental, DataflowGroom))
-struct FUnpackSkinWeightsDataflowNode : public FDataflowNode
-{
-	GENERATED_BODY()
-	
-	DATAFLOW_NODE_DEFINE_INTERNAL(FUnpackSkinWeightsDataflowNode, "UnpackSkinWeights", "Groom", "")
-
-public:
 	
-	FUnpackSkinWeightsDataflowNode(const UE::Dataflow::FNodeParameters& InParam, FGuid InGuid = FGuid::NewGuid())
-	: FDataflowNode(InParam, InGuid)
-	{
-		RegisterInputConnection(&Collection);
-		RegisterOutputConnection(&Collection, &Collection);
-		RegisterOutputConnection(&BoneIndicesKey);
-		RegisterOutputConnection(&BoneWeightsKey);
-	}
-
-	//~ Begin FDataflowNode interface
-	virtual void Evaluate(UE::Dataflow::FContext& Context, const FDataflowOutput* Out) const override;
-	virtual TArray<UE::Dataflow::FRenderingParameter> GetRenderParametersImpl() const override;
-	//~ End FDataflowNode interface
-
-	/** Managed array collection to be used to store datas */
-	UPROPERTY(meta = (DataflowInput, DataflowOutput, DisplayName = "Collection", DataflowPassthrough  = "Collection"))
-	FManagedArrayCollection Collection;
-
-	/** The name to be set as a bones indices attribute. */
-	UPROPERTY(EditAnywhere, Category = "Vertex Attributes", meta = (DisplayName = "Bone Indices"))
-	FString BoneIndicesName;
-
-	/** The name to be set as a bones weights attribute. */
-	UPROPERTY(EditAnywhere, Category = "Vertex Attributes", meta = (DisplayName = "Bone Weights"))
-	FString BoneWeightsName;
-
 	/** Bone indices key to be used in other nodes if necessary */
-	UPROPERTY(meta = (DisplayName = "Bone Indices Key", DataflowOutput))
+	UPROPERTY(meta = (DisplayName = "Bone Indices", DataflowOutput))
 	FCollectionAttributeKey BoneIndicesKey;
 
 	/** Bone weights key to be used in other nodes if necessary */
-	UPROPERTY(meta = (DisplayName = "Bone Weights Key", DataflowOutput))
+	UPROPERTY(meta = (DisplayName = "Bone Weights", DataflowOutput))
 	FCollectionAttributeKey BoneWeightsKey;
 
-	/** Type of curves to use to fill the groom collection (guides/strands) */
-	UPROPERTY(EditAnywhere, Category = "Vertex Attributes", meta = (DisplayName = "Curves Type"))
-	EGroomCollectionType CurvesType = EGroomCollectionType::Guides;
-
 private :
 	/** Get the bone indices key */
 	FCollectionAttributeKey GetBoneIndicesKey() const;
@@ -114,59 +75,4 @@ private :
 	FCollectionAttributeKey GetBoneWeightsKey() const;
 };
 
-/** Report the vertices skinning attributes to fill the points ones */
-USTRUCT(meta = (Experimental, DataflowGroom))
-struct FPackSkinWeightsDataflowNode : public FDataflowNode
-{
-	GENERATED_BODY()
-	
-	DATAFLOW_NODE_DEFINE_INTERNAL(FPackSkinWeightsDataflowNode, "PackSkinWeights", "Groom", "")
-
-public:
-	
-	FPackSkinWeightsDataflowNode(const UE::Dataflow::FNodeParameters& InParam, FGuid InGuid = FGuid::NewGuid())
-	: FDataflowNode(InParam, InGuid)
-	{
-		RegisterInputConnection(&Collection);
-		RegisterInputConnection(&BoneIndicesKey);
-		RegisterInputConnection(&BoneWeightsKey);
-		RegisterOutputConnection(&Collection, &Collection);
-	}
-
-	//~ Begin FDataflowNode interface
-	virtual void Evaluate(UE::Dataflow::FContext& Context, const FDataflowOutput* Out) const override;
-	virtual TArray<UE::Dataflow::FRenderingParameter> GetRenderParametersImpl() const override;
-	//~ End FDataflowNode interface
-
-	/** Managed array collection to be used to store datas */
-	UPROPERTY(meta = (DataflowInput, DataflowOutput, DisplayName = "Collection", DataflowPassthrough  = "Collection"))
-	FManagedArrayCollection Collection;
-
-	/** The name to be set as a weight map attribute. */
-	UPROPERTY(EditAnywhere, Category = "Vertex Attributes", meta = (DisplayName = "Bone Indices"))
-	FString BoneIndicesName;
-
-	/** The name to be set as a weight map attribute. */
-	UPROPERTY(EditAnywhere, Category = "Vertex Attributes", meta = (DisplayName = "Bone Weights"))
-	FString BoneWeightsName;
-
-	/** Bone indices key to be used in other nodes if necessary */
-	UPROPERTY(meta = (DisplayName = "Bone Indices Key", DataflowInput))
-	FCollectionAttributeKey BoneIndicesKey;
-
-	/** Bone weights key to be used in other nodes if necessary */
-	UPROPERTY(meta = (DisplayName = "Bone Weights Key", DataflowInput))
-	FCollectionAttributeKey BoneWeightsKey;
-
-	/** Type of curves to use to fill the groom collection (guides/strands) */
-	UPROPERTY(EditAnywhere, Category = "Vertex Attributes", meta = (DisplayName = "Curves Type"))
-	EGroomCollectionType CurvesType = EGroomCollectionType::Guides;
-
-private :
-	/** Get the bone indices key */
-	FCollectionAttributeKey GetBoneIndicesKey(UE::Dataflow::FContext& Context) const;
-
-	/** Get the bone weights key */
-	FCollectionAttributeKey GetBoneWeightsKey(UE::Dataflow::FContext& Context) const;
-};
 
diff --git a/Engine/Plugins/Runtime/HairStrands/Source/HairStrandsDataflow/Public/BuildGuidesLODsNode.h b/Engine/Plugins/Runtime/HairStrands/Source/HairStrandsDataflow/Public/BuildGuidesLODsNode.h
--- a/Engine/Plugins/Runtime/HairStrands/Source/HairStrandsDataflow/Public/BuildGuidesLODsNode.h
+++ b/Engine/Plugins/Runtime/HairStrands/Source/HairStrandsDataflow/Public/BuildGuidesLODsNode.h
@@ -5,6 +5,7 @@
 #include "CoreMinimal.h"
 #include "GroomAsset.h"
 #include "ChaosLog.h"
+#include "Dataflow/DataflowConnectionTypes.h"
 #include "Dataflow/DataflowCore.h"
 #include "Dataflow/DataflowEngine.h"
 
@@ -27,6 +28,8 @@ public:
 	{
 		RegisterInputConnection(&Collection);
 		RegisterOutputConnection(&Collection, &Collection);
+		RegisterOutputConnection(&CurveParentsKey);
+		RegisterOutputConnection(&CurveLodsKey);
 	}
 
 	//~ Begin FDataflowNode interface
@@ -36,5 +39,13 @@ public:
 	/** Managed array collection to be used to store data */
 	UPROPERTY(meta = (DataflowInput, DataflowOutput, DisplayName = "Collection", DataflowPassthrough  = "Collection"))
 	FManagedArrayCollection Collection;
+	
+	/** Curve parent indices key to be used in other nodes if necessary */
+	UPROPERTY(meta = (DisplayName = "Curve Parents", DataflowOutput))
+	FCollectionAttributeKey CurveParentsKey;
+
+	/** Curve lods indices key to be used in other nodes if necessary */
+	UPROPERTY(meta = (DisplayName = "Curve Lods", DataflowOutput))
+	FCollectionAttributeKey CurveLodsKey;
 };
 
diff --git a/Engine/Plugins/Runtime/HairStrands/Source/HairStrandsDataflow/Public/GetGroomAssetNode.h b/Engine/Plugins/Runtime/HairStrands/Source/HairStrandsDataflow/Public/GetGroomAssetNode.h
--- a/Engine/Plugins/Runtime/HairStrands/Source/HairStrandsDataflow/Public/GetGroomAssetNode.h
+++ b/Engine/Plugins/Runtime/HairStrands/Source/HairStrandsDataflow/Public/GetGroomAssetNode.h
@@ -37,8 +37,8 @@ public:
 	TObjectPtr<const UGroomAsset> GroomAsset = nullptr;
 
 	/** Type of curves to use to fill the groom collection (guides/strands) */
-	UPROPERTY(EditAnywhere, Category = "Groom", meta = (DisplayName = "CurvesType"))
-	EGroomCollectionType CurvesType = EGroomCollectionType::Guides;
+	UPROPERTY(EditAnywhere, Category = "Groom", meta = (DisplayName = "Display Curves"))
+	EGroomCollectionType CurvesType = EGroomCollectionType::Strands;
 
 	/** Managed array collection used to store the guides */
 	UPROPERTY(meta = (DataflowOutput, DisplayName = "Collection"))
diff --git a/Engine/Plugins/Runtime/HairStrands/Source/HairStrandsDataflow/Public/GetGroomAttributesNodes.h b/Engine/Plugins/Runtime/HairStrands/Source/HairStrandsDataflow/Public/GetGroomAttributesNodes.h
new file mode 100644
--- /dev/null
+++ b/Engine/Plugins/Runtime/HairStrands/Source/HairStrandsDataflow/Public/GetGroomAttributesNodes.h
@@ -0,0 +1,56 @@
+// Copyright Epic Games, Inc. All Rights Reserved.
+
+#pragma once 
+
+#include "CoreMinimal.h"
+#include "ChaosLog.h"
+#include "GroomCollectionFacades.h"
+#include "Dataflow/DataflowCore.h"
+#include "Dataflow/DataflowConnectionTypes.h"
+
+#include "GetGroomAttributesNodes.generated.h"
+
+/** Enum to pick groom attribute type */
+UENUM()
+enum class EGroomAttributeType : uint8
+{
+	KinematicWeights UMETA(DisplayName = "Kinematic Weights"),
+	BoneIndices UMETA(DisplayName = "Bone Indices"),
+	BoneWeights UMETA(DisplayName = "Bone Weights"),
+	CurveParents UMETA(DisplayName = "Curve Parents"),
+	CurveLods UMETA(DisplayName = "Curve Lods")
+};
+
+/** Get the kinematic weights attributes names */
+USTRUCT(meta = (Experimental, DataflowGroom))
+struct FGetGroomAttributesDataflowNode : public FDataflowNode
+{
+	GENERATED_BODY()
+	
+	DATAFLOW_NODE_DEFINE_INTERNAL(FGetGroomAttributesDataflowNode, "GetGroomAttributes", "Groom", "")
+
+public:
+	
+	FGetGroomAttributesDataflowNode(const UE::Dataflow::FNodeParameters& InParam, FGuid InGuid = FGuid::NewGuid())
+	: FDataflowNode(InParam, InGuid)
+	{
+		RegisterOutputConnection(&AttributeKey);
+	}
+
+	//~ Begin FDataflowNode interface
+	virtual void Evaluate(UE::Dataflow::FContext& Context, const FDataflowOutput* Out) const override;
+	//~ End FDataflowNode interface
+
+	/** Attribute key to build */
+	UPROPERTY(meta = (DisplayName = "Attribute Key", DataflowOutput))
+	FCollectionAttributeKey AttributeKey;
+
+	/** Type of curves to use (guides/strands) */
+	UPROPERTY(EditAnywhere, Category = "Groom", meta = (DisplayName = "Curves Type"))
+	EGroomCollectionType CurvesType = EGroomCollectionType::Guides;
+
+	/** Type of attribute to use */
+	UPROPERTY(EditAnywhere, Category = "Groom", meta = (DisplayName = "Attribute Type"))
+	EGroomAttributeType AttributeType = EGroomAttributeType::KinematicWeights;
+};
+
diff --git a/Engine/Plugins/Runtime/HairStrands/Source/HairStrandsDataflow/Public/GetGuidesAttributesNodes.h b/Engine/Plugins/Runtime/HairStrands/Source/HairStrandsDataflow/Public/GetGuidesAttributesNodes.h
deleted file mode 100644
diff --git a/Engine/Plugins/Runtime/HairStrands/Source/HairStrandsDataflow/Public/GroomAssetTerminalNode.h b/Engine/Plugins/Runtime/HairStrands/Source/HairStrandsDataflow/Public/GroomAssetTerminalNode.h
--- a/Engine/Plugins/Runtime/HairStrands/Source/HairStrandsDataflow/Public/GroomAssetTerminalNode.h
+++ b/Engine/Plugins/Runtime/HairStrands/Source/HairStrandsDataflow/Public/GroomAssetTerminalNode.h
@@ -7,6 +7,7 @@
 #include "ChaosLog.h"
 #include "Dataflow/DataflowCore.h"
 #include "Dataflow/DataflowEngine.h"
+#include "Dataflow/DataflowConnectionTypes.h"
 
 #include "GroomAssetTerminalNode.generated.h"
 
@@ -28,9 +29,26 @@ public:
 	//~ Begin FDataflowTerminalNode interface
 	virtual void SetAssetValue(TObjectPtr<UObject> Asset, UE::Dataflow::FContext& Context) const override;
 	virtual void Evaluate(UE::Dataflow::FContext& Context, const FDataflowOutput* Out) const override;
+
+	virtual TArray<UE::Dataflow::FPin> AddPins() override;
+	virtual bool CanAddPin() const override { return true; }
+	virtual bool CanRemovePin() const override { return !AttributeKeys.IsEmpty(); }
+	virtual TArray<UE::Dataflow::FPin> GetPinsToRemove() const override;
+	virtual void OnPinRemoved(const UE::Dataflow::FPin& Pin) override;
+	virtual void PostSerialize(const FArchive& Ar) override;
 	//~ End FDataflowTerminalNode interface
 
 	UPROPERTY(meta = (DataflowInput, DataflowOutput, DisplayName = "Collection", DataflowPassthrough  = "Collection"))
 	FManagedArrayCollection Collection;
+	
+	/** List of attribute keys that will be used to save matching attributes in the collection. */
+	UPROPERTY()
+	TArray<FCollectionAttributeKey> AttributeKeys;
+
+private :
+	/** Get the connection reference matching an attribute key index */
+	UE::Dataflow::TConnectionReference<FCollectionAttributeKey> GetConnectionReference(int32 Index) const;
+
+	static constexpr int32 NumOtherInputs = 1;
 };
 
diff --git a/Engine/Plugins/Runtime/HairStrands/Source/HairStrandsDataflow/Public/GroomCollectionFacades.h b/Engine/Plugins/Runtime/HairStrands/Source/HairStrandsDataflow/Public/GroomCollectionFacades.h
--- a/Engine/Plugins/Runtime/HairStrands/Source/HairStrandsDataflow/Public/GroomCollectionFacades.h
+++ b/Engine/Plugins/Runtime/HairStrands/Source/HairStrandsDataflow/Public/GroomCollectionFacades.h
@@ -214,7 +214,7 @@ namespace UE::Groom
 
 		static const FName ObjectMeshLODsAttribute;
 		static const FName ObjectSkeletalMeshesAttribute;
-		static const FName VertexKinematicWeightsAttribute;
+		static const FName PointKinematicWeightsAttribute;
 		static const FName PointBoneIndicesAttribute;
 		static const FName PointBoneWeightsAttribute;
 		static const FName ObjectPointSamplesAttribute;
@@ -237,8 +237,8 @@ namespace UE::Groom
 			return GroomGroup.Guides;
 		}
 		
-		/** Get the vertex kinematic weights */
-		const TArray<float>& GetVertexKinematicWeights() const { return VertexKinematicWeights.Get().GetConstArray(); }
+		/** Get the point kinematic weights */
+		const TArray<float>& GetPointKinematicWeights() const { return PointKinematicWeights.Get().GetConstArray(); }
 		
 		/** Get the point bone indices */
 		const FIntVector4& GetPointBoneIndices(const int32 PointIndex) const { return PointBoneIndices.Get()[PointIndex]; }
@@ -259,7 +259,7 @@ namespace UE::Groom
 		const TArray<int32>& GetCurveLodIndices() const { return CurveLodIndices.Get().GetConstArray(); }
 
 		/** Set the vertex kinematic weights */
-		void SetVertexKinematicWeights(const TArray<float>& InKinematicWeights) { VertexKinematicWeights.Modify() = InKinematicWeights; }
+		void SetPointKinematicWeights(const TArray<float>& InKinematicWeights) { PointKinematicWeights.Modify() = InKinematicWeights; }
 
 		/** Set the point bone indices */
 		void SetPointBoneIndices(const int32 PointIndex, const FIntVector4& InBoneIndices) { return PointBoneIndices.ModifyAt(PointIndex,InBoneIndices); }
@@ -285,7 +285,7 @@ namespace UE::Groom
 	protected :
 		
 		/** Max distance from the kinematic target */
-		TManagedArrayAccessor<float> VertexKinematicWeights;
+		TManagedArrayAccessor<float> PointKinematicWeights;
 
 		/** Point bone indices */
 		TManagedArrayAccessor<FIntVector4> PointBoneIndices;
