Fix hair barycentric value returning NaN in certain configuration when transferring positions.

diff --git a/Engine/Plugins/Runtime/HairStrands/Source/HairStrandsCore/Private/GroomBindingBuilder.cpp b/Engine/Plugins/Runtime/HairStrands/Source/HairStrandsCore/Private/GroomBindingBuilder.cpp
--- a/Engine/Plugins/Runtime/HairStrands/Source/HairStrandsCore/Private/GroomBindingBuilder.cpp
+++ b/Engine/Plugins/Runtime/HairStrands/Source/HairStrandsCore/Private/GroomBindingBuilder.cpp
@@ -2164,7 +2164,7 @@ namespace GroomBinding_Transfer
 		float VC = D1 * D4 - D3 * D2;
 		if (VC <= 0.f && D1 >= 0.f && D3 <= 0.f)
 		{
-			float V = D1 / (D1 - D3);
+			float V = FMath::Abs(D1 - D3) > 0.0001f? D1 / (D1 - D3) : 0.f;
 
 			FTrianglePoint Out;
 			Out.P = A + V * AB;
