Fix bad RBF transfer when far from origin

diff --git a/Engine/Plugins/Runtime/HairStrands/Shaders/Private/NiagaraDataInterfaceHairStrandsTemplate.ush b/Engine/Plugins/Runtime/HairStrands/Shaders/Private/NiagaraDataInterfaceHairStrandsTemplate.ush
--- a/Engine/Plugins/Runtime/HairStrands/Shaders/Private/NiagaraDataInterfaceHairStrandsTemplate.ush
+++ b/Engine/Plugins/Runtime/HairStrands/Shaders/Private/NiagaraDataInterfaceHairStrandsTemplate.ush
@@ -30,9 +30,11 @@ uint				{ParameterName}_ResetSimulation;
 uint				{ParameterName}_RestUpdate;
 uint				{ParameterName}_LocalSimulation;
 int					{ParameterName}_SampleCount;
+uint				{ParameterName}_RBFLocalSpace;
 int4				{ParameterName}_BoundingBoxOffsets;
 StructuredBuffer<float4> {ParameterName}_RestSamplePositionsBuffer;
 StructuredBuffer<float4> {ParameterName}_MeshSampleWeightsBuffer;
+StructuredBuffer<float4> {ParameterName}_DeformedSamplePositionsBuffer;
 Buffer<float>		{ParameterName}_ParamsScaleBuffer;
 float3				{ParameterName}_BoneLinearVelocity;
 float3				{ParameterName}_BoneAngularVelocity;
@@ -857,7 +859,9 @@ float3 ApplyRBF_{ParameterName}(in float3 RestControlPoint)
 		RestControlPoint, 
 		{ParameterName}_SampleCount, 
 		{ParameterName}_RestSamplePositionsBuffer,
-		{ParameterName}_MeshSampleWeightsBuffer);
+		{ParameterName}_MeshSampleWeightsBuffer,
+		{ParameterName}_DeformedSamplePositionsBuffer,
+		{ParameterName}_RBFLocalSpace);
 }
 
  // Compute the global position
diff --git a/Engine/Plugins/Runtime/HairStrands/Source/HairStrandsCore/Private/GroomBindingAsset.cpp b/Engine/Plugins/Runtime/HairStrands/Source/HairStrandsCore/Private/GroomBindingAsset.cpp
--- a/Engine/Plugins/Runtime/HairStrands/Source/HairStrandsCore/Private/GroomBindingAsset.cpp
+++ b/Engine/Plugins/Runtime/HairStrands/Source/HairStrandsCore/Private/GroomBindingAsset.cpp
@@ -1009,12 +1009,13 @@ static FString BuildDerivedDataKeySuffix(const UGroomBindingAsset& BindingAsset,
 	FString PointKey  = FString::FromInt(BindingAsset.GetNumInterpolationPoints());
 	FString SectionKey = FString::FromInt(BindingAsset.GetMatchingSection());
 	FString TargetBindingAttributeKey = BindingAsset.GetTargetBindingAttribute().ToString();
+	FString RBLocalSpaceKey = FString::FromInt(UE::Groom::IsRBFLocalSpaceEnabled());
 
-	uint32 KeyLength  = BindingType.Len() + SourceKey.Len() + TargetKey.Len() + GroomKey.Len() + PointKey.Len() + SectionKey.Len() + SourceLODKey.Len() + TargetLODKey.Len() + TargetBindingAttributeKey.Len();
+	uint32 KeyLength  = BindingType.Len() + SourceKey.Len() + TargetKey.Len() + GroomKey.Len() + PointKey.Len() + SectionKey.Len() + SourceLODKey.Len() + TargetLODKey.Len() + TargetBindingAttributeKey.Len() + RBLocalSpaceKey.Len();
 
 	FString KeySuffix;
 	KeySuffix.Reserve(KeyLength);
-	KeySuffix = BindingType + SourceKey + TargetKey + GroomKey + PointKey + SectionKey + SourceLODKey + TargetLODKey + TargetBindingAttributeKey;
+	KeySuffix = BindingType + SourceKey + TargetKey + GroomKey + PointKey + SectionKey + SourceLODKey + TargetLODKey + TargetBindingAttributeKey + RBLocalSpaceKey;
 	return KeySuffix;
 }
 
diff --git a/Engine/Plugins/Runtime/HairStrands/Source/HairStrandsCore/Private/GroomBindingBuilder.cpp b/Engine/Plugins/Runtime/HairStrands/Source/HairStrandsCore/Private/GroomBindingBuilder.cpp
--- a/Engine/Plugins/Runtime/HairStrands/Source/HairStrandsCore/Private/GroomBindingBuilder.cpp
+++ b/Engine/Plugins/Runtime/HairStrands/Source/HairStrandsCore/Private/GroomBindingBuilder.cpp
@@ -56,6 +56,20 @@ DEFINE_LOG_CATEGORY_STATIC(LogGroomBindingBuilder, Log, All);
 static int32 GHairStrandsBindingBuilderWarningEnable = 1;
 static FAutoConsoleVariableRef CVarHairStrandsBindingBuilderWarningEnable(TEXT("r.HairStrands.Log.BindingBuilderWarning"), GHairStrandsBindingBuilderWarningEnable, TEXT("Enable/disable warning during groom binding builder"));
 
+static TAutoConsoleVariable<int32> CVarHairStrandsBFLocalSpace(
+	TEXT("r.HairStrands.RBFLocalSpace"),
+	1,
+	TEXT("Enable/disable rbf computation in local space."),
+	ECVF_ReadOnly | ECVF_RenderThreadSafe);
+
+namespace UE::Groom
+{
+	bool IsRBFLocalSpaceEnabled()
+	{
+		return CVarHairStrandsBFLocalSpace.GetValueOnAnyThread() > 0;
+	}
+}
+
 ///////////////////////////////////////////////////////////////////////////////////////////////////
 
 FGroomBindingBuilder::FInput::FInput(UGroomBindingAsset* BindingAsset, const ITargetPlatform* TargetPlatform, int32 InSourceMeshLOD, int32 InTargetMeshMinLOD)
@@ -78,7 +92,7 @@ FGroomBindingBuilder::FInput::FInput(UGroomBindingAsset* BindingAsset, const ITa
 FString FGroomBindingBuilder::GetVersion()
 {
 	// Important to update the version when groom building changes
-	return TEXT("6a");
+	return TEXT("6c");
 }
 
 ///////////////////////////////////////////////////////////////////////////////////////////////////
@@ -856,7 +870,7 @@ namespace GroomBinding_RBFWeighting
 {
 
 	FWeightsBuilder::FWeightsBuilder(const uint32 NumRows, const uint32 NumColumns,
-		const FVector3f* SourcePositions, const FVector3f* TargetPositions)
+		const FVector3f* SourcePositions, const FVector3f* TargetPositions, const bool bLocalSpace)
 	{
 		const uint32 PolyRows 	 = FGroomRBFDeformer::GetEntryCount(NumRows);
 		const uint32 PolyColumns = FGroomRBFDeformer::GetEntryCount(NumColumns);
@@ -881,9 +895,11 @@ namespace GroomBinding_RBFWeighting
 				PolyColumns,
 				SourcePositions,
 				TargetPositions,
-				&LocalEntries
+				&LocalEntries,
+				bLocalSpace
 			] (uint32 RowIndex)
 			{
+				const FVector3f PositionOffset = (bLocalSpace && (NumRows > 0))? SourcePositions[0] : FVector3f::Zero();
 				int32 EntryIndex = RowIndex * PolyColumns;
 				for (uint32 j = 0; j < NumColumns; ++j)
 				{
@@ -891,23 +907,23 @@ namespace GroomBinding_RBFWeighting
 					LocalEntries[EntryIndex++] = FMath::Sqrt(FunctionScale * FunctionScale + 1.0);
 				}
 				LocalEntries[EntryIndex++] = 1.0;
-				LocalEntries[EntryIndex++] = SourcePositions[RowIndex].X;
-				LocalEntries[EntryIndex++] = SourcePositions[RowIndex].Y;
-				LocalEntries[EntryIndex++] = SourcePositions[RowIndex].Z;
+				LocalEntries[EntryIndex++] = (SourcePositions[RowIndex].X-PositionOffset.X);
+				LocalEntries[EntryIndex++] = (SourcePositions[RowIndex].Y-PositionOffset.Y);
+				LocalEntries[EntryIndex++] = (SourcePositions[RowIndex].Z-PositionOffset.Z);
 
 				EntryIndex = NumRows * PolyColumns + RowIndex;
 				LocalEntries[EntryIndex] = 1.0;
 
 				EntryIndex += PolyColumns;
-				LocalEntries[EntryIndex] = SourcePositions[RowIndex].X;
+				LocalEntries[EntryIndex] = (SourcePositions[RowIndex].X-PositionOffset.X);
 
 				EntryIndex += PolyColumns;
-				LocalEntries[EntryIndex] = SourcePositions[RowIndex].Y;
+				LocalEntries[EntryIndex] = (SourcePositions[RowIndex].Y-PositionOffset.Y);
 
 				EntryIndex += PolyColumns;
-				LocalEntries[EntryIndex] = SourcePositions[RowIndex].Z;
+				LocalEntries[EntryIndex] = (SourcePositions[RowIndex].Z-PositionOffset.Z);
 			});
-		const float REGUL_VALUE = 1e-4;
+		static constexpr float REGUL_VALUE = 1e-4;
 		int32 EntryIndex = NumRows * PolyColumns + NumColumns;
 		LocalEntries[EntryIndex] = REGUL_VALUE;
 
@@ -1071,7 +1087,7 @@ namespace GroomBinding_RBFWeighting
 					FPointsSampler PointsSampler(ValidPoints, PositionsPointer, MaxSamples);
 					const uint32 SampleCount = PointsSampler.SamplePositions.Num();
 
-					FWeightsBuilder InterpolationWeights(SampleCount, SampleCount, PointsSampler.SamplePositions.GetData(), PointsSampler.SamplePositions.GetData());
+					FWeightsBuilder InterpolationWeights(SampleCount, SampleCount, PointsSampler.SamplePositions.GetData(), PointsSampler.SamplePositions.GetData(), UE::Groom::IsRBFLocalSpaceEnabled());
 
 					// Guides
 					UpdateInterpolationWeights(InterpolationWeights, PointsSampler, MeshLODData, Out.SimRootDatas[MeshLODIndex]);
@@ -1093,7 +1109,7 @@ namespace GroomBinding_RBFWeighting
 				FPointsSampler PointsSampler(ValidPoints, PositionsPointer, MaxSamples);
 				const uint32 SampleCount = PointsSampler.SamplePositions.Num();
 
-				FWeightsBuilder InterpolationWeights(SampleCount, SampleCount, PointsSampler.SamplePositions.GetData(), PointsSampler.SamplePositions.GetData());
+				FWeightsBuilder InterpolationWeights(SampleCount, SampleCount, PointsSampler.SamplePositions.GetData(), PointsSampler.SamplePositions.GetData(), UE::Groom::IsRBFLocalSpaceEnabled());
 
 				// Guides
 				UpdateInterpolationWeights(InterpolationWeights, PointsSampler, MeshLODData, Out.SimRootDatas[MeshLODIndex]);
diff --git a/Engine/Plugins/Runtime/HairStrands/Source/HairStrandsCore/Private/GroomRBFDeformer.cpp b/Engine/Plugins/Runtime/HairStrands/Source/HairStrandsCore/Private/GroomRBFDeformer.cpp
--- a/Engine/Plugins/Runtime/HairStrands/Source/HairStrandsCore/Private/GroomRBFDeformer.cpp
+++ b/Engine/Plugins/Runtime/HairStrands/Source/HairStrandsCore/Private/GroomRBFDeformer.cpp
@@ -48,21 +48,25 @@ void InitMeshSamples(
 void UpdateMeshSamples(
 	uint32 MaxSampleCount,
 	const TArray<float>& InterpolationWeightsBuffer,
-	const TArray<FVector4f>& SampleRestPositionsBuffer,
-	const TArray<FVector3f>& SampleDeformedPositionsBuffer,
+	const TArray<FVector4f>& RestSamplePositionsBuffer,
+	const TArray<FVector3f>& DeformedSamplePositionsBuffer,
 	TArray<FVector3f>& OutSampleDeformationsBuffer
 )
 {
 	const uint32 EntryCount = FGroomRBFDeformer::GetEntryCount(MaxSampleCount);
 	OutSampleDeformationsBuffer.SetNum(EntryCount);
 
+	const bool bIsRBFLocalSpaceEnabled = !RestSamplePositionsBuffer.IsEmpty() && !DeformedSamplePositionsBuffer.IsEmpty() && UE::Groom::IsRBFLocalSpaceEnabled();
+	const FVector4f RestSampleOffset = bIsRBFLocalSpaceEnabled ? RestSamplePositionsBuffer[0] : FVector4f(0, 0, 0, 0);
+	const FVector3f DeformedSampleOffset = bIsRBFLocalSpaceEnabled ? DeformedSamplePositionsBuffer[0] : FVector3f(0, 0, 0);
+
 	for (uint32 SampleIndex = 0; SampleIndex < EntryCount; ++SampleIndex)
 	{
 		uint32 WeightsOffset = SampleIndex * EntryCount;
 		FVector3f SampleDeformation(FVector3f::ZeroVector);
 		for (uint32 i = 0; i < MaxSampleCount; ++i, ++WeightsOffset)
 		{
-			SampleDeformation += InterpolationWeightsBuffer[WeightsOffset] * (SampleDeformedPositionsBuffer[i] - SampleRestPositionsBuffer[i]);
+			SampleDeformation += InterpolationWeightsBuffer[WeightsOffset] * ((DeformedSamplePositionsBuffer[i]- DeformedSampleOffset) - (RestSamplePositionsBuffer[i]-RestSampleOffset));
 		}
 
 		OutSampleDeformationsBuffer[SampleIndex] = SampleDeformation;
@@ -74,10 +78,15 @@ FVector3f DisplacePosition(
 	const FVector3f& RestControlPoint,
 	uint32 SampleCount,
 	const TArray<FVector4f>& RestSamplePositionsBuffer,
+	const TArray<FVector3f>& DeformedSamplePositionsBuffer,
 	const TArray<FVector3f>& MeshSampleWeightsBuffer
 )
 {
-	FVector3f ControlPoint = RestControlPoint;
+	const bool bIsRBFLocalSpaceEnabled = !RestSamplePositionsBuffer.IsEmpty() && !DeformedSamplePositionsBuffer.IsEmpty() && UE::Groom::IsRBFLocalSpaceEnabled();
+	const FVector3f RestSampleOffset = bIsRBFLocalSpaceEnabled ? FVector3f(RestSamplePositionsBuffer[0]) : FVector3f(0, 0, 0);
+	const FVector3f DeformedSampleOffset = bIsRBFLocalSpaceEnabled ? DeformedSamplePositionsBuffer[0] : FVector3f(0, 0, 0);
+
+	FVector3f ControlPoint = RestControlPoint-RestSampleOffset;
 
 	// Apply rbf interpolation from the samples set
 	for (uint32 i = 0; i < SampleCount; ++i)
@@ -87,10 +96,10 @@ FVector3f DisplacePosition(
 		ControlPoint += FunctionValue * MeshSampleWeightsBuffer[i];
 	}
 	ControlPoint += MeshSampleWeightsBuffer[SampleCount];
-	ControlPoint += MeshSampleWeightsBuffer[SampleCount + 1] * RestControlPoint.X;
-	ControlPoint += MeshSampleWeightsBuffer[SampleCount + 2] * RestControlPoint.Y;
-	ControlPoint += MeshSampleWeightsBuffer[SampleCount + 3] * RestControlPoint.Z;
-	return ControlPoint;
+	ControlPoint += MeshSampleWeightsBuffer[SampleCount + 1] * (RestControlPoint.X - RestSampleOffset.X);
+	ControlPoint += MeshSampleWeightsBuffer[SampleCount + 2] * (RestControlPoint.Y - RestSampleOffset.Y);
+	ControlPoint += MeshSampleWeightsBuffer[SampleCount + 3] * (RestControlPoint.Z - RestSampleOffset.Z);
+	return ControlPoint + DeformedSampleOffset;
 }
 
 void DeformStrands(	
@@ -105,6 +114,7 @@ void DeformStrands(
 	const bool bHasRBF,
 	const TArray<FVector3f>& RestPosePositionBuffer,
 	const TArray<FVector4f>& RestSamplePositionsBuffer,
+	const TArray<FVector3f>& DeformedSamplePositionsBuffer,
 	const TArray<FVector3f>& MeshSampleWeightsBuffer,
 	TArray<FVector3f>& OutDeformedPositionBuffer
 )
@@ -114,7 +124,7 @@ void DeformStrands(
 	ParallelFor(VertexCount, [&](uint32 VertexIndex)
 	{
 		const FVector3f& ControlPoint = RestPosePositionBuffer[VertexIndex];
-		const FVector3f DisplacedPosition = DisplacePosition(ControlPoint, SampleCount, RestSamplePositionsBuffer, MeshSampleWeightsBuffer);
+		const FVector3f DisplacedPosition = DisplacePosition(ControlPoint, SampleCount, RestSamplePositionsBuffer, DeformedSamplePositionsBuffer, MeshSampleWeightsBuffer);
 		OutDeformedPositionBuffer[VertexIndex] = bHasRBF ? DisplacedPosition : ControlPoint;
 	});
 
@@ -226,7 +236,6 @@ TArray<FVector3f> GetDeformedHairStrandsPositions(
 
 	// Deform the strands vertices with the deformed mesh samples
 	const TArray<FVector3f>& RestPosePositionBuffer = OutPositions;
-	const TArray<FVector4f>& RestSamplePositionsBuffer = RestLODData.RestSamplePositionsBuffer;
 	const TArray<FVector3f>& MeshSampleWeightsBuffer = OutSampleDeformationsBuffer;
 	TArray<FVector3f> OutDeformedPositionBuffer;
 
@@ -241,7 +250,8 @@ TArray<FVector3f> GetDeformedHairStrandsPositions(
 		MaxSampleCount, 
 		bHasRBF,
 		RestPosePositionBuffer, 
-		RestSamplePositionsBuffer, 
+		SampleRestPositionsBuffer,
+		SampleDeformedPositionsBuffer,
 		MeshSampleWeightsBuffer, 
 		OutDeformedPositionBuffer);
 
@@ -533,7 +543,6 @@ void DeformStaticMeshPositions(
 	UpdateMeshSamples(MaxSampleCount, InterpolationWeightsBuffer, SampleRestPositionsBuffer, SampleDeformedPositionsBuffer, OutSampleDeformationsBuffer);
 
 	// Deform the strands vertices with the deformed mesh samples
-	const TArray<FVector4f>& RestSamplePositionsBuffer = RestLODData.RestSamplePositionsBuffer;
 	const TArray<FVector3f>& MeshSampleWeightsBuffer = OutSampleDeformationsBuffer;
 
 	// Raw deformation with RBF
@@ -549,7 +558,7 @@ void DeformStaticMeshPositions(
 		{
 			FVertexID VertexID(VertexIndex);
 			FVector3f& Point = VertexPositions[VertexID];
-			Point = DisplacePosition(Point, MaxSampleCount, RestSamplePositionsBuffer, MeshSampleWeightsBuffer);
+			Point = DisplacePosition(Point, MaxSampleCount, SampleRestPositionsBuffer, SampleDeformedPositionsBuffer, MeshSampleWeightsBuffer);
 		});
 		MeshDescriptions.Add(MeshDescription);
 	}
diff --git a/Engine/Plugins/Runtime/HairStrands/Source/HairStrandsCore/Private/HairStrandsInterpolation.cpp b/Engine/Plugins/Runtime/HairStrands/Source/HairStrandsCore/Private/HairStrandsInterpolation.cpp
--- a/Engine/Plugins/Runtime/HairStrands/Source/HairStrandsCore/Private/HairStrandsInterpolation.cpp
+++ b/Engine/Plugins/Runtime/HairStrands/Source/HairStrandsCore/Private/HairStrandsInterpolation.cpp
@@ -351,8 +351,10 @@ class FDeformGuideCS : public FGlobalShader
 		SHADER_PARAMETER_RDG_BUFFER_UAV(RWByteAddressBuffer, OutSimDeformedPositionBuffer)
 
 		SHADER_PARAMETER(uint32, SampleCount)
+		SHADER_PARAMETER(uint32, RBFLocalSpace)
 		SHADER_PARAMETER_RDG_BUFFER_SRV(StructuredBuffer, RestSamplePositionsBuffer)
 		SHADER_PARAMETER_RDG_BUFFER_SRV(StructuredBuffer, MeshSampleWeightsBuffer)
+		SHADER_PARAMETER_RDG_BUFFER_SRV(StructuredBuffer, DeformedSamplePositionsBuffer)
 
 		SHADER_PARAMETER_SRV(Buffer<float4>, BoneDeformedPositionBuffer)
 	END_SHADER_PARAMETER_STRUCT()
@@ -448,7 +450,10 @@ void AddDeformSimHairStrandsPass(
 					InternalDeformationType = InternalDeformationType_RBF;
 					Parameters->MeshSampleWeightsBuffer = RegisterAsSRV(GraphBuilder, DeformedLODDatas.GetMeshSampleWeightsBuffer(FHairStrandsLODDeformedRootResource::Current));
 					Parameters->RestSamplePositionsBuffer = RegisterAsSRV(GraphBuilder, RestLODDatas.RestSamplePositionsBuffer);
+					Parameters->DeformedSamplePositionsBuffer = RegisterAsSRV(GraphBuilder, DeformedLODDatas.GetDeformedSamplePositionsBuffer(FHairStrandsLODDeformedRootResource::Current));
 					Parameters->SampleCount = RestLODDatas.SampleCount;
+					Parameters->RBFLocalSpace = UE::Groom::IsRBFLocalSpaceEnabled();
+;
 				}
 			}
 		}
diff --git a/Engine/Plugins/Runtime/HairStrands/Source/HairStrandsCore/Private/HairStrandsMeshProjection.cpp b/Engine/Plugins/Runtime/HairStrands/Source/HairStrandsCore/Private/HairStrandsMeshProjection.cpp
--- a/Engine/Plugins/Runtime/HairStrands/Source/HairStrandsCore/Private/HairStrandsMeshProjection.cpp
+++ b/Engine/Plugins/Runtime/HairStrands/Source/HairStrandsCore/Private/HairStrandsMeshProjection.cpp
@@ -11,6 +11,7 @@
 #include "RenderGraphUtils.h"
 #include "GroomResources.h"
 #include "SystemTextures.h"
+#include "GroomBindingBuilder.h"
 
 static int32 GHairProjectionMaxTrianglePerProjectionIteration = 8;
 static FAutoConsoleVariableRef CVarHairProjectionMaxTrianglePerProjectionIteration(TEXT("r.HairStrands.Projection.MaxTrianglePerIteration"), GHairProjectionMaxTrianglePerProjectionIteration, TEXT("Change the number of triangles which are iterated over during one projection iteration step. In kilo triangle (e.g., 8 == 8000 triangles). Default is 8."));
@@ -403,9 +404,11 @@ private:
 	BEGIN_SHADER_PARAMETER_STRUCT(FParameters, )
 		SHADER_PARAMETER(uint32, VertexCount)
 		SHADER_PARAMETER(uint32, MaxSampleCount)
+		SHADER_PARAMETER(uint32, RBFLocalSpace)
 
 		SHADER_PARAMETER_RDG_BUFFER_SRV(StructuredBuffer, RestSamplePositionsBuffer)
 		SHADER_PARAMETER_RDG_BUFFER_SRV(StructuredBuffer, MeshSampleWeightsBuffer)
+		SHADER_PARAMETER_RDG_BUFFER_SRV(StructuredBuffer, DeformedSamplePositionsBuffer)
 
 		SHADER_PARAMETER_SRV(Buffer, RestPositionBuffer)
 		SHADER_PARAMETER_RDG_BUFFER_UAV(RWBuffer, OutDeformedPositionBuffer)
@@ -447,14 +450,17 @@ void InternalAddHairRBFInterpolationPass(
 	FHairStrandsLODRestRootResource& RestLODData = *RestRootResources->GetLOD(MeshLODIndex);
 	FHairStrandsLODDeformedRootResource& DeformedLODData = *DeformedRootResources->GetLOD(MeshLODIndex);
 	FHairMeshesInterpolateCS::FParameters* Parameters = GraphBuilder.AllocParameters<FHairMeshesInterpolateCS::FParameters>();
+	
 	Parameters->VertexCount = VertexCount;
 	Parameters->MaxSampleCount = RestLODData.SampleCount;
+	Parameters->RBFLocalSpace = UE::Groom::IsRBFLocalSpaceEnabled();
 
 	Parameters->RestPositionBuffer			= RestResources->RestPositionBuffer.ShaderResourceViewRHI;
 	Parameters->OutDeformedPositionBuffer	= DeformedPositionBuffer_Curr.UAV;
 
 	Parameters->RestSamplePositionsBuffer	= RegisterAsSRV(GraphBuilder, RestLODData.RestSamplePositionsBuffer);
 	Parameters->MeshSampleWeightsBuffer		= RegisterAsSRV(GraphBuilder, DeformedLODData.GetMeshSampleWeightsBuffer(FHairStrandsLODDeformedRootResource::Current));
+	Parameters->DeformedSamplePositionsBuffer	= RegisterAsSRV(GraphBuilder, DeformedLODData.GetDeformedSamplePositionsBuffer(FHairStrandsLODDeformedRootResource::Current));
 
 	const FIntVector DispatchGroupCount = FComputeShaderUtils::GetGroupCount(VertexCount, 128);
 	check(DispatchGroupCount.X <= GRHIMaxDispatchThreadGroupsPerDimension.X);
@@ -647,6 +653,7 @@ private:
 
 	BEGIN_SHADER_PARAMETER_STRUCT(FParameters, )
 		SHADER_PARAMETER(uint32, MaxSampleCount)
+		SHADER_PARAMETER(uint32, RBFLocalSpace)
 
 		SHADER_PARAMETER_RDG_BUFFER_SRV(StructuredBuffer, InterpolationWeightsBuffer)
 		SHADER_PARAMETER_RDG_BUFFER_SRV(StructuredBuffer, SampleRestPositionsBuffer)
@@ -693,6 +700,7 @@ void AddHairStrandUpdateMeshSamplesPass(
 		Parameters->SampleRestPositionsBuffer		= RegisterAsSRV(GraphBuilder, RestLODData.RestSamplePositionsBuffer);
 		Parameters->SampleDeformedPositionsBuffer	= RegisterAsSRV(GraphBuilder, DeformedLODData.GetDeformedSamplePositionsBuffer(FHairStrandsLODDeformedRootResource::Current));
 		Parameters->OutSampleDeformationsBuffer		= OutWeightsBuffer.UAV;
+		Parameters->RBFLocalSpace = UE::Groom::IsRBFLocalSpaceEnabled();
 
 		const FIntVector DispatchGroupCount = FComputeShaderUtils::GetGroupCount(RestLODData.SampleCount+4, 128);
 		check(DispatchGroupCount.X <= GRHIMaxDispatchThreadGroupsPerDimension.X);
diff --git a/Engine/Plugins/Runtime/HairStrands/Source/HairStrandsCore/Private/Niagara/NiagaraDataInterfaceHairStrands.cpp b/Engine/Plugins/Runtime/HairStrands/Source/HairStrandsCore/Private/Niagara/NiagaraDataInterfaceHairStrands.cpp
--- a/Engine/Plugins/Runtime/HairStrands/Source/HairStrandsCore/Private/Niagara/NiagaraDataInterfaceHairStrands.cpp
+++ b/Engine/Plugins/Runtime/HairStrands/Source/HairStrandsCore/Private/Niagara/NiagaraDataInterfaceHairStrands.cpp
@@ -23,6 +23,7 @@
 
 #include "GroomComponent.h"
 #include "GroomAsset.h"
+#include "GroomBindingBuilder.h"
 #include "SystemTextures.h"
 
 #include UE_INLINE_GENERATED_CPP_BY_NAME(NiagaraDataInterfaceHairStrands)
@@ -58,9 +59,11 @@ BEGIN_SHADER_PARAMETER_STRUCT(FShaderParameters,)
 	SHADER_PARAMETER(uint32,							RestUpdate)
 	SHADER_PARAMETER(uint32,							LocalSimulation)
 	SHADER_PARAMETER(int,								SampleCount)
+	SHADER_PARAMETER(int,								RBFLocalSpace)
 	SHADER_PARAMETER(FIntVector4,						BoundingBoxOffsets)
 	SHADER_PARAMETER_RDG_BUFFER_SRV(StructuredBuffer<float4>, RestSamplePositionsBuffer)
 	SHADER_PARAMETER_RDG_BUFFER_SRV(StructuredBuffer<float4>, MeshSampleWeightsBuffer)
+	SHADER_PARAMETER_RDG_BUFFER_SRV(StructuredBuffer<float4>, DeformedSamplePositionsBuffer)
 	SHADER_PARAMETER_RDG_BUFFER_SRV(Buffer<float>,		ParamsScaleBuffer)
 	SHADER_PARAMETER(FVector3f,							BoneLinearVelocity)
 	SHADER_PARAMETER(FVector3f,							BoneAngularVelocity)
@@ -295,6 +298,7 @@ struct FHairGroupInstanceRDG
 	// RBF
 	FRDGBufferSRVRef RestSamplePositionsBuffer = nullptr;
 	FRDGBufferSRVRef MeshSampleWeightsBuffer = nullptr;
+	FRDGBufferSRVRef DeformedSamplePositionsBuffer = nullptr;
 
 	bool IsValid() const 		{ return bIsValid && GeometryType != EHairGeometryType::NoneGeometry; }
 	bool IsDeformedValid() const{ return bIsDeformedValid && GeometryType != EHairGeometryType::NoneGeometry; }
@@ -355,6 +359,7 @@ static FHairGroupInstanceRDG Convert(FRDGBuilder& GraphBuilder, const TRefCountP
 			{
 				Out.DeformedTrianglePositionBuffer	= RegisterAsSRV(GraphBuilder, RootResource->GetDeformedUniqueTrianglePositionBuffer(FHairStrandsLODDeformedRootResource::Current));
 				Out.MeshSampleWeightsBuffer			= RegisterAsSRV(GraphBuilder, RootResource->GetMeshSampleWeightsBuffer((FHairStrandsLODDeformedRootResource::Current)));
+				Out.DeformedSamplePositionsBuffer	= RegisterAsSRV(GraphBuilder, RootResource->GetDeformedSamplePositionsBuffer((FHairStrandsLODDeformedRootResource::Current)));
 				Out.bIsRootValid					= Out.BindingType == EHairBindingType::Skinning;
 			}
 		}
@@ -384,7 +389,7 @@ static FHairGroupInstanceRDG Convert(FRDGBuilder& GraphBuilder, const TRefCountP
 	if (!Out.RootToUniqueTriangleIndexBuffer)	{ Out.RootToUniqueTriangleIndexBuffer 	= DummyVertexBuffer4byte; };
 	if (!Out.RestSamplePositionsBuffer)			{ Out.RestSamplePositionsBuffer 		= DummyStructuredBuffer; };
 	if (!Out.MeshSampleWeightsBuffer)			{ Out.MeshSampleWeightsBuffer 			= DummyStructuredBuffer; };
-
+	if (!Out.DeformedSamplePositionsBuffer)		{ Out.DeformedSamplePositionsBuffer 	= DummyStructuredBuffer; };
 	return Out;
 }
 
@@ -3564,6 +3569,7 @@ void UNiagaraDataInterfaceHairStrands::SetShaderParameters(const FNiagaraDataInt
 		ShaderParameters->DeformedRootOffset = FVector3f::ZeroVector;
 		ShaderParameters->RestPositionOffset = InstanceRDG.RestPositionOffsetValue;
 		ShaderParameters->SampleCount = InstanceRDG.SampleCount;
+		ShaderParameters->RBFLocalSpace = UE::Groom::IsRBFLocalSpaceEnabled();
 
 		ShaderParameters->DeformedPositionBuffer = InstanceRDG.DeformedPositionBufferUAV;
 		ShaderParameters->CurvesOffsetsBuffer = InstanceRDG.CurvesOffsetsBuffer;
@@ -3573,6 +3579,7 @@ void UNiagaraDataInterfaceHairStrands::SetShaderParameters(const FNiagaraDataInt
 		ShaderParameters->DeformedTrianglePositionBuffer = InstanceRDG.DeformedTrianglePositionBuffer;
 		ShaderParameters->RestSamplePositionsBuffer = InstanceRDG.RestSamplePositionsBuffer;
 		ShaderParameters->MeshSampleWeightsBuffer = InstanceRDG.MeshSampleWeightsBuffer;
+		ShaderParameters->DeformedSamplePositionsBuffer = InstanceRDG.DeformedSamplePositionsBuffer;
 		ShaderParameters->RootBarycentricCoordinatesBuffer = InstanceRDG.RootBarycentricCoordinatesBuffer;
 		ShaderParameters->RootToUniqueTriangleIndexBuffer = InstanceRDG.RootToUniqueTriangleIndexBuffer;
 
@@ -3607,6 +3614,7 @@ void UNiagaraDataInterfaceHairStrands::SetShaderParameters(const FNiagaraDataInt
 		ShaderParameters->DeformedRootOffset = FVector3f::ZeroVector;
 		ShaderParameters->RestPositionOffset = FVector3f::ZeroVector;
 		ShaderParameters->SampleCount = 0;
+		ShaderParameters->RBFLocalSpace = 0;
 
 		ShaderParameters->DeformedPositionBuffer = InstanceRDG.DeformedPositionBufferUAV;
 		ShaderParameters->CurvesOffsetsBuffer = InstanceRDG.CurvesOffsetsBuffer;
@@ -3616,6 +3624,7 @@ void UNiagaraDataInterfaceHairStrands::SetShaderParameters(const FNiagaraDataInt
 		ShaderParameters->DeformedTrianglePositionBuffer = InstanceRDG.DeformedTrianglePositionBuffer;
 		ShaderParameters->RestSamplePositionsBuffer = InstanceRDG.RestSamplePositionsBuffer;
 		ShaderParameters->MeshSampleWeightsBuffer = InstanceRDG.MeshSampleWeightsBuffer;
+		ShaderParameters->DeformedSamplePositionsBuffer = InstanceRDG.DeformedSamplePositionsBuffer;
 		ShaderParameters->RootBarycentricCoordinatesBuffer = InstanceRDG.RootBarycentricCoordinatesBuffer;
 		ShaderParameters->RootToUniqueTriangleIndexBuffer = InstanceRDG.RootToUniqueTriangleIndexBuffer;
 
diff --git a/Engine/Plugins/Runtime/HairStrands/Source/HairStrandsCore/Public/GroomBindingBuilder.h b/Engine/Plugins/Runtime/HairStrands/Source/HairStrandsCore/Public/GroomBindingBuilder.h
--- a/Engine/Plugins/Runtime/HairStrands/Source/HairStrandsCore/Public/GroomBindingBuilder.h
+++ b/Engine/Plugins/Runtime/HairStrands/Source/HairStrandsCore/Public/GroomBindingBuilder.h
@@ -16,6 +16,11 @@
 class ITargetPlatform;
 class UGroomBindingAsset;
 
+namespace UE::Groom
+{
+	UE_API bool IsRBFLocalSpaceEnabled();
+}
+
 struct FGroomBindingBuilder
 {
 	struct FInput
@@ -91,7 +96,7 @@ namespace GroomBinding_RBFWeighting
 	struct FWeightsBuilder
 	{
 		UE_API FWeightsBuilder(const uint32 NumRows, const uint32 NumColumns,
-			const FVector3f* SourcePositions, const FVector3f* TargetPositions);
+			const FVector3f* SourcePositions, const FVector3f* TargetPositions, const bool bLocalSpace = false);
 
 		/** Compute the weights by inverting the matrix*/
 		UE_API void ComputeWeights(const uint32 NumRows, const uint32 NumColumns);
diff --git a/Engine/Shaders/Private/HairStrands/HairStrandsBindingCommon.ush b/Engine/Shaders/Private/HairStrands/HairStrandsBindingCommon.ush
--- a/Engine/Shaders/Private/HairStrands/HairStrandsBindingCommon.ush
+++ b/Engine/Shaders/Private/HairStrands/HairStrandsBindingCommon.ush
@@ -172,9 +172,15 @@ float3 ApplyRBF(
 	float3 RestControlPoint,
 	int InMaxSampleCount,
 	StructuredBuffer<float4> InRestSamplePositionsBuffer,
-	StructuredBuffer<float4> InMeshSampleWeightsBuffer)
+	StructuredBuffer<float4> InMeshSampleWeightsBuffer,
+	StructuredBuffer<float4> InDeformedSamplePositionsBuffer,
+	uint RBFLocalSpace)
 {
-	float3 ControlPoint = RestControlPoint;
+	const bool bEnableLocalSpace = (InMaxSampleCount > 0) && (RBFLocalSpace == 1);
+	const float4 RestSampleOffset = bEnableLocalSpace ? InRestSamplePositionsBuffer[0] : float4(0,0,0,0);
+	const float4 DeformedSampleOffset = bEnableLocalSpace ? InDeformedSamplePositionsBuffer[0] : float4(0,0,0,0);
+	
+	float3 ControlPoint = RestControlPoint-RestSampleOffset;
 	// Apply rbf interpolation from the samples set
 	for (int i = 0; i < InMaxSampleCount; ++i)
 	{
@@ -183,8 +189,10 @@ float3 ApplyRBF(
 		ControlPoint += FunctionValue * InMeshSampleWeightsBuffer[i].xyz;
 	}
 	ControlPoint += InMeshSampleWeightsBuffer[InMaxSampleCount].xyz;
-	ControlPoint += InMeshSampleWeightsBuffer[InMaxSampleCount + 1].xyz * RestControlPoint.x;
-	ControlPoint += InMeshSampleWeightsBuffer[InMaxSampleCount + 2].xyz * RestControlPoint.y;
-	ControlPoint += InMeshSampleWeightsBuffer[InMaxSampleCount + 3].xyz * RestControlPoint.z;
+	ControlPoint += InMeshSampleWeightsBuffer[InMaxSampleCount + 1].xyz * (RestControlPoint.x-RestSampleOffset.x);
+	ControlPoint += InMeshSampleWeightsBuffer[InMaxSampleCount + 2].xyz * (RestControlPoint.y-RestSampleOffset.y);
+	ControlPoint += InMeshSampleWeightsBuffer[InMaxSampleCount + 3].xyz * (RestControlPoint.z-RestSampleOffset.z);
+
+	ControlPoint += DeformedSampleOffset;
 	return ControlPoint;
 }
diff --git a/Engine/Shaders/Private/HairStrands/HairStrandsGuideDeform.usf b/Engine/Shaders/Private/HairStrands/HairStrandsGuideDeform.usf
--- a/Engine/Shaders/Private/HairStrands/HairStrandsGuideDeform.usf
+++ b/Engine/Shaders/Private/HairStrands/HairStrandsGuideDeform.usf
@@ -54,7 +54,9 @@ Buffer<float4>  BoneDeformedPositionBuffer;
 
 #if PERMUTATION_DEFORMATION == DEFORMATION_OFFSET_GUIDE_GLOBAL
 uint SampleCount;
+uint RBFLocalSpace;
 StructuredBuffer<float4>  RestSamplePositionsBuffer;
+StructuredBuffer<float4>  DeformedSamplePositionsBuffer;
 StructuredBuffer<float4>  MeshSampleWeightsBuffer;
 #endif
 
@@ -84,7 +86,7 @@ float3 DisplacePosition(uint VertexIndex, float3 Pos)
 
 #if PERMUTATION_DEFORMATION == DEFORMATION_OFFSET_GUIDE_GLOBAL
 	const float3 RestControlPoint = Pos + SimRestOffset;
-	const float3 ControlPoint = ApplyRBF(RestControlPoint, SampleCount, RestSamplePositionsBuffer, MeshSampleWeightsBuffer);
+	const float3 ControlPoint = ApplyRBF(RestControlPoint, SampleCount, RestSamplePositionsBuffer, MeshSampleWeightsBuffer, DeformedSamplePositionsBuffer, RBFLocalSpace);
 	return ControlPoint - OutSimHairPositionOffset;
 #endif
 
diff --git a/Engine/Shaders/Private/HairStrands/HairStrandsMesh.usf b/Engine/Shaders/Private/HairStrands/HairStrandsMesh.usf
--- a/Engine/Shaders/Private/HairStrands/HairStrandsMesh.usf
+++ b/Engine/Shaders/Private/HairStrands/HairStrandsMesh.usf
@@ -11,9 +11,11 @@
 
 uint VertexCount;
 uint MaxSampleCount;
+uint RBFLocalSpace;
 
 StructuredBuffer<float4> RestSamplePositionsBuffer;
 StructuredBuffer<float4> MeshSampleWeightsBuffer;
+StructuredBuffer<float4> DeformedSamplePositionsBuffer;
 
 Buffer<float4> RestPositionBuffer;
 RWBuffer<float4> OutDeformedPositionBuffer;
@@ -26,7 +28,7 @@ void MainHairMeshesCS(uint3 DispatchThreadId : SV_DispatchThreadID)
 		return;
 
 	const float3 RestPosition = RestPositionBuffer[VertexIndex].xyz;
-	const float3 DeformedPosition = ApplyRBF(RestPosition, MaxSampleCount, RestSamplePositionsBuffer, MeshSampleWeightsBuffer);
+	const float3 DeformedPosition = ApplyRBF(RestPosition, MaxSampleCount, RestSamplePositionsBuffer, MeshSampleWeightsBuffer, DeformedSamplePositionsBuffer, RBFLocalSpace);
 	OutDeformedPositionBuffer[VertexIndex] = float4(DeformedPosition,1);
 }
 
@@ -97,6 +99,7 @@ void MainCS(uint3 DispatchThreadId : SV_DispatchThreadID)
 #include "../Common.ush"
 
 uint MaxSampleCount;
+uint RBFLocalSpace;
 StructuredBuffer<float> InterpolationWeightsBuffer;
 
 StructuredBuffer<float4> SampleRestPositionsBuffer;
@@ -113,12 +116,17 @@ void MainCS(uint3 DispatchThreadId : SV_DispatchThreadID)
 		return;
 	}
 
+	const bool bEnableLocalSpace = (MaxSampleCount > 0) && (RBFLocalSpace == 1);
+	const float4 RestSampleOffset = bEnableLocalSpace ? SampleRestPositionsBuffer[0] : float4(0,0,0,0);
+	const float4 DeformedSampleOffset = bEnableLocalSpace ? SampleDeformedPositionsBuffer[0] : float4(0,0,0,0);
+
 	uint WeightsBaseOffset = SampleIndex * HAIR_RBF_ENTRY_COUNT(MaxSampleCount);
 	float3 SampleDeformation = float3(0,0,0);
 	for(uint i = 0; i < MaxSampleCount; ++i)
 	{
 		[REDACTED_EMAIL] - Incrementing '++WeightsBaseOffset' instead of using primary iterator '+i' results in visual artifacts with Metal backend
-		SampleDeformation += InterpolationWeightsBuffer[WeightsBaseOffset + i] * (SampleDeformedPositionsBuffer[i] - SampleRestPositionsBuffer[i]).xyz;
+		SampleDeformation += InterpolationWeightsBuffer[WeightsBaseOffset + i] * ((SampleDeformedPositionsBuffer[i]-DeformedSampleOffset) -
+			(SampleRestPositionsBuffer[i]-RestSampleOffset)).xyz;
 	}
 
 	OutSampleDeformationsBuffer[SampleIndex] = float4(SampleDeformation,1);
