Fix GPU crash due to hair card RT geometry being build on non-updated data.

Grooms are not rendered/updated in reflection capture updates which causes RT geometry to be built/update on non-update data, leading to GPU crash. Skipping RT cards/meshes update.

diff --git a/Engine/Plugins/Runtime/HairStrands/Source/HairStrandsCore/Private/GroomManager.cpp b/Engine/Plugins/Runtime/HairStrands/Source/HairStrandsCore/Private/GroomManager.cpp
--- a/Engine/Plugins/Runtime/HairStrands/Source/HairStrandsCore/Private/GroomManager.cpp
+++ b/Engine/Plugins/Runtime/HairStrands/Source/HairStrandsCore/Private/GroomManager.cpp
@@ -2174,9 +2174,11 @@ static void RunHairStrandsInterpolation_Cards(
 		}
 	}
 
+	const bool bIsReflectionCapture = View && View->bIsReflectionCapture;
+
 	// Cards only - Build/update RT geometry
 	#if RHI_RAYTRACING
-	if (IsRayTracingEnabled())
+	if (IsRayTracingEnabled() && !bIsReflectionCapture)
 	{
 		for (uint32 InstanceIndex : CardInstances)
 		{
@@ -2214,7 +2216,7 @@ static void RunHairStrandsInterpolation_Cards(
 
 	// Meshes only - Build/update RT geometry
 	#if RHI_RAYTRACING
-	if (IsRayTracingEnabled())
+	if (IsRayTracingEnabled() && !bIsReflectionCapture)
 	{
 		for (uint32 InstanceIndex : MeshInstances)
 		{
