Fix attribute and solver groom deformer offsets

diff --git a/Engine/Plugins/Runtime/HairStrands/Shaders/Private/Deformers/DeformerGroomAttributeRead.ush b/Engine/Plugins/Runtime/HairStrands/Shaders/Private/Deformers/DeformerGroomAttributeRead.ush
--- a/Engine/Plugins/Runtime/HairStrands/Shaders/Private/Deformers/DeformerGroomAttributeRead.ush
+++ b/Engine/Plugins/Runtime/HairStrands/Shaders/Private/Deformers/DeformerGroomAttributeRead.ush
@@ -3,6 +3,8 @@
 #include "/Plugin/ComputeFramework/Private/ComputeKernelCommon.ush"
 
 uint {DataInterfaceName}_NumValues; 
+uint {DataInterfaceName}_BaseValueIndex; 
+uint {DataInterfaceName}_TotalValueCount;
 
 Buffer<float> {DataInterfaceName}_ValueTypedData;
 
@@ -10,38 +12,44 @@ Buffer<float> {DataInterfaceName}_ValueTypedData;
 
 uint ReadNumValues_{DataInterfaceName}()
 {
-	return {DataInterfaceName}_NumValues;
+	return {DataInterfaceName}_TotalValueCount;
 }
 
 uint ReadValueTypedUint_{DataInterfaceName}(uint ValueIndex)
 {
+	ValueIndex -= {DataInterfaceName}_BaseValueIndex;
 	return (ValueIndex < {DataInterfaceName}_NumValues) ? (uint){DataInterfaceName}_ValueTypedData[ValueIndex] : 0;
 }
 
 bool ReadValueTypedBool_{DataInterfaceName}(uint ValueIndex)
 {
+	ValueIndex -= {DataInterfaceName}_BaseValueIndex;
 	return (ValueIndex < {DataInterfaceName}_NumValues) ? (({DataInterfaceName}_ValueTypedData[ValueIndex] == 1.0f) ? true : false) : false;
 }
 
 int ReadValueTypedInt_{DataInterfaceName}(uint ValueIndex)
 {
+	ValueIndex -= {DataInterfaceName}_BaseValueIndex;
 	return (ValueIndex < {DataInterfaceName}_NumValues) ? (int){DataInterfaceName}_ValueTypedData[ValueIndex] : 0;
 }
 
 int2 ReadValueTypedInt2_{DataInterfaceName}(uint ValueIndex)
 {
+	ValueIndex -= {DataInterfaceName}_BaseValueIndex;
 	const int ValueOffset = 2 * ValueIndex;
 	return (ValueIndex < {DataInterfaceName}_NumValues) ? int2({DataInterfaceName}_ValueTypedData[ValueOffset], {DataInterfaceName}_ValueTypedData[ValueOffset+1]) : int2(0,0);
 }
 
 int3 ReadValueTypedInt3_{DataInterfaceName}(uint ValueIndex)
 {
+	ValueIndex -= {DataInterfaceName}_BaseValueIndex;
 	const int ValueOffset = 3 * ValueIndex;
 	return (ValueIndex < {DataInterfaceName}_NumValues) ? int3({DataInterfaceName}_ValueTypedData[ValueOffset], {DataInterfaceName}_ValueTypedData[ValueOffset+1], {DataInterfaceName}_ValueTypedData[ValueOffset+2]) : int3(0,0,0);
 }
 
 int4 ReadValueTypedInt4_{DataInterfaceName}(uint ValueIndex)
 {
+	ValueIndex -= {DataInterfaceName}_BaseValueIndex;
 	const int ValueOffset = 4 * ValueIndex;
 	return (ValueIndex < {DataInterfaceName}_NumValues) ? int4({DataInterfaceName}_ValueTypedData[ValueOffset], {DataInterfaceName}_ValueTypedData[ValueOffset+1],
 					{DataInterfaceName}_ValueTypedData[ValueOffset+2], {DataInterfaceName}_ValueTypedData[ValueOffset+3]) : int4(0,0,0,0);
@@ -49,11 +57,13 @@ int4 ReadValueTypedInt4_{DataInterfaceName}(uint ValueIndex)
 
 float ReadValueTypedFloat_{DataInterfaceName}(uint ValueIndex)
 {
+	ValueIndex -= {DataInterfaceName}_BaseValueIndex;
 	return (ValueIndex < {DataInterfaceName}_NumValues) ? {DataInterfaceName}_ValueTypedData[ValueIndex] : 0.0f;
 }
 
 float2 ReadValueTypedFloat2_{DataInterfaceName}(uint ValueIndex)
 {
+	ValueIndex -= {DataInterfaceName}_BaseValueIndex;
 	const int ValueOffset = 2 * ValueIndex;
 	return (ValueIndex < {DataInterfaceName}_NumValues) ? float2({DataInterfaceName}_ValueTypedData[ValueOffset], {DataInterfaceName}_ValueTypedData[ValueOffset+1]) : float2(0.0f,0.0f);
 
@@ -61,12 +71,14 @@ float2 ReadValueTypedFloat2_{DataInterfaceName}(uint ValueIndex)
 
 float3 ReadValueTypedFloat3_{DataInterfaceName}(uint ValueIndex)
 {
+	ValueIndex -= {DataInterfaceName}_BaseValueIndex;
 	const int ValueOffset = 3 * ValueIndex;
 	return (ValueIndex < {DataInterfaceName}_NumValues) ? float3({DataInterfaceName}_ValueTypedData[ValueOffset], {DataInterfaceName}_ValueTypedData[ValueOffset+1], {DataInterfaceName}_ValueTypedData[ValueOffset+2]) : float3(0.0f,0.0f,0.0f);
 }
 
 float4 ReadValueTypedFloat4_{DataInterfaceName}(uint ValueIndex)
 {
+	ValueIndex -= {DataInterfaceName}_BaseValueIndex;
 	const int ValueOffset = 4 * ValueIndex;
 	return (ValueIndex < {DataInterfaceName}_NumValues) ? float4({DataInterfaceName}_ValueTypedData[ValueOffset], {DataInterfaceName}_ValueTypedData[ValueOffset+1],
 					{DataInterfaceName}_ValueTypedData[ValueOffset+2], {DataInterfaceName}_ValueTypedData[ValueOffset+3]) : float4(0.0f,0.0f,0.0f,0.0f);
@@ -74,6 +86,7 @@ float4 ReadValueTypedFloat4_{DataInterfaceName}(uint ValueIndex)
 
 float4x4 ReadValueTypedFloat4x4_{DataInterfaceName}(uint ValueIndex)
 {
+	ValueIndex -= {DataInterfaceName}_BaseValueIndex;
 	const int ValueOffset = 16 * ValueIndex;
 	return (ValueIndex < {DataInterfaceName}_NumValues) ? float4x4({DataInterfaceName}_ValueTypedData[ValueOffset], {DataInterfaceName}_ValueTypedData[ValueOffset+1],
 						{DataInterfaceName}_ValueTypedData[ValueOffset+2], {DataInterfaceName}_ValueTypedData[ValueOffset+3],
@@ -88,6 +101,7 @@ float4x4 ReadValueTypedFloat4x4_{DataInterfaceName}(uint ValueIndex)
 
 float3x4 ReadValueTypedFloat3x4_{DataInterfaceName}(uint ValueIndex)
 {
+	ValueIndex -= {DataInterfaceName}_BaseValueIndex;
 	const int ValueOffset = 12 * ValueIndex;
 	return (ValueIndex < {DataInterfaceName}_NumValues) ? float3x4({DataInterfaceName}_ValueTypedData[ValueOffset], {DataInterfaceName}_ValueTypedData[ValueOffset+1],
 						{DataInterfaceName}_ValueTypedData[ValueOffset+2], {DataInterfaceName}_ValueTypedData[ValueOffset+3],
diff --git a/Engine/Plugins/Runtime/HairStrands/Source/HairStrandsDeformer/Private/DeformerGroomAttributeRead.cpp b/Engine/Plugins/Runtime/HairStrands/Source/HairStrandsDeformer/Private/DeformerGroomAttributeRead.cpp
--- a/Engine/Plugins/Runtime/HairStrands/Source/HairStrandsDeformer/Private/DeformerGroomAttributeRead.cpp
+++ b/Engine/Plugins/Runtime/HairStrands/Source/HairStrandsDeformer/Private/DeformerGroomAttributeRead.cpp
@@ -217,6 +217,8 @@ void UOptimusGroomAttributeReadDataInterface::GetSupportedInputs(TArray<FShaderF
 
 BEGIN_SHADER_PARAMETER_STRUCT(FOptimusGroomAttributeReadParameters, )
 	SHADER_PARAMETER(uint32, NumValues)
+	SHADER_PARAMETER(uint32, BaseValueIndex)
+	SHADER_PARAMETER(uint32, TotalValueCount)
 	SHADER_PARAMETER_RDG_BUFFER_SRV(Buffer<float>, ValueTypedData)
 END_SHADER_PARAMETER_STRUCT()
 
@@ -277,7 +279,7 @@ FOptimusGroomAttributeReadProviderProxy::FOptimusGroomAttributeReadProviderProxy
 	TArray<const UGroomComponent*> GroomComponents;
 	UE::Groom::Private::GatherGroomComponents(MeshComponent, GroomComponents);
 	UE::Groom::Private::GroomComponentsToInstances(GroomComponents, GroupInstances);
-	UE::Groom::Private::GetGroomInvocationElementGroups(GroomComponents,UOptimusGroomExecDataInterface::GetExecutionDomainName(GroomAttributeGroup), GroupElements, INDEX_NONE, true);
+	UE::Groom::Private::GetGroomInvocationElementGroups(GroomComponents,UOptimusGroomExecDataInterface::GetExecutionDomainName(GroomAttributeGroup), GroupElements, true);
 }
 
 bool FOptimusGroomAttributeReadProviderProxy::IsValid(FValidationData const& InValidationData) const
@@ -645,6 +647,7 @@ void FOptimusGroomAttributeReadProviderProxy::GatherDispatchData(FDispatchData c
 	const TStridedView<FParameters> ParameterArray = MakeStridedParameterView<FParameters>(InDispatchData);
 
 	int32 InvocationIndex = 0;
+	int32 BaseValueIndex = 0;
 	for (const TPair<UGroomAsset*,UE::Groom::Private::FGroupElements>& AssetPair : GroupElements)
 	{
 		int32 InstanceIndex = 0;
@@ -654,15 +657,31 @@ void FOptimusGroomAttributeReadProviderProxy::GatherDispatchData(FDispatchData c
 			if (const FHairGroupInstance* GroupInstance = AssetPair.Value.GroupInstances[InstanceIndex])
 			{
 				Parameters.NumValues = AssetPair.Value.GroupOffsets[GroupIndex+1]-AssetPair.Value.GroupOffsets[GroupIndex];
+				Parameters.BaseValueIndex = BaseValueIndex;
 				Parameters.ValueTypedData = AttributeValuesResources[InvocationIndex];
+
+				BaseValueIndex += Parameters.NumValues;
 			}
 			else
 			{
 				Parameters.NumValues = 0;
+				Parameters.BaseValueIndex = BaseValueIndex;
 				Parameters.ValueTypedData = FallbackStructuredSRV;
 			}
 			++InstanceIndex;
 			++InvocationIndex;
 		}
 	}
+
+	InvocationIndex = 0;
+	for (const TPair<UGroomAsset*, UE::Groom::Private::FGroupElements>& AssetPair : GroupElements)
+	{
+		for (int32 GroupIndex : AssetPair.Value.GroupIndices)
+		{
+			FParameters& Parameters = ParameterArray[InvocationIndex];
+			Parameters.TotalValueCount = BaseValueIndex;
+			++InvocationIndex;
+		}
+	}
+
 }
diff --git a/Engine/Plugins/Runtime/HairStrands/Source/HairStrandsDeformer/Private/DeformerGroomDomainsSource.cpp b/Engine/Plugins/Runtime/HairStrands/Source/HairStrandsDeformer/Private/DeformerGroomDomainsSource.cpp
--- a/Engine/Plugins/Runtime/HairStrands/Source/HairStrandsDeformer/Private/DeformerGroomDomainsSource.cpp
+++ b/Engine/Plugins/Runtime/HairStrands/Source/HairStrandsDeformer/Private/DeformerGroomDomainsSource.cpp
@@ -147,7 +147,7 @@ bool UOptimusGroomSolverComponentSource::GetComponentElementCountsForExecutionDo
 		}
 		else if (DomainName == FSolverExecutionDomains::Points)
 		{
-			TotalCount = UE::Groom::Private::GetGroomInvocationElementCounts(GroomComponents, UOptimusGroomAssetComponentSource::FGuidesExecutionDomains::Points, LocalElementCounts, GroupSize);
+			TotalCount = UE::Groom::Private::GetGroomInvocationElementCounts(GroomComponents, UOptimusGroomAssetComponentSource::FGuidesExecutionDomains::Points, LocalElementCounts);
 		}
 		else if (DomainName == FDynamicExecutionDomains::Curves)
 		{
diff --git a/Engine/Plugins/Runtime/HairStrands/Source/HairStrandsDeformer/Private/DeformerGroomInterfaceUtils.h b/Engine/Plugins/Runtime/HairStrands/Source/HairStrandsDeformer/Private/DeformerGroomInterfaceUtils.h
--- a/Engine/Plugins/Runtime/HairStrands/Source/HairStrandsDeformer/Private/DeformerGroomInterfaceUtils.h
+++ b/Engine/Plugins/Runtime/HairStrands/Source/HairStrandsDeformer/Private/DeformerGroomInterfaceUtils.h
@@ -258,12 +258,11 @@ namespace UE::Groom::Private
 		return true;
 	}
 
-	FORCEINLINE void GetGroomGroupElementOffsets(const UGroomComponent* GroomComponent, const FName DomainName, TArray<int32>& GroupOffsets, const int32 GroupSize = INDEX_NONE, const bool bSourceElements = false)
+	static void GetGroomGroupElementOffsets(const UGroomComponent* GroomComponent, const FName DomainName, TArray<int32>& GroupOffsets, const bool bSourceElements = false)
 	{
 		if(GroomComponent) 
 		{
 			const uint32 NumGroups = GroomComponent->GetGroupCount();
-			const float GroupAlignment = static_cast<float>(GroupSize);
 			int32 MeshLOD = INDEX_NONE;
 			
 			GroupOffsets.Init(0, NumGroups+1);
@@ -327,14 +326,12 @@ namespace UE::Groom::Private
 						}
 					}
 				}
-				NumElements = (GroupSize == INDEX_NONE) ? NumElements :
-				   FMath::CeilToInt32(NumElements / GroupAlignment) * GroupAlignment;
 			}
 			GroupOffsets[NumGroups] = NumElements;
 		}
 	}
 
-	FORCEINLINE int32 GetGroomInvocationElementCounts(const TArray<const UGroomComponent*>& GroomComponents, const FName DomainName, TArray<int32>& InvocationCounts, int32 GroupSize = INDEX_NONE, const bool bSourceElements = false)
+	static int32 GetGroomInvocationElementCounts(const TArray<const UGroomComponent*>& GroomComponents, const FName DomainName, TArray<int32>& InvocationCounts, const bool bSourceElements = false)
 	{
 		TArray<int32> GroupOffsets;
 		int32 TotalCount = 0;
@@ -343,7 +340,7 @@ namespace UE::Groom::Private
 		{
 			if(GroomComponent)
 			{
-				GetGroomGroupElementOffsets(GroomComponent, DomainName, GroupOffsets, GroupSize, bSourceElements);
+				GetGroomGroupElementOffsets(GroomComponent, DomainName, GroupOffsets, bSourceElements);
 
 				const uint32 NumGroups = GroomComponent->GetGroupCount();
 				for (uint32 GroupIndex = 0; GroupIndex < NumGroups; ++GroupIndex)
@@ -366,7 +363,7 @@ namespace UE::Groom::Private
 		TArray<const FHairGroupInstance*> GroupInstances;
 	};
 	
-	FORCEINLINE void GetGroomInvocationElementGroups(const TArray<const UGroomComponent*>& GroomComponents, const FName DomainName, TArray<TPair<UGroomAsset*, FGroupElements>>& InvocationGroups, int32 GroupSize = INDEX_NONE, const bool bSourceElements = false)
+	FORCEINLINE void GetGroomInvocationElementGroups(const TArray<const UGroomComponent*>& GroomComponents, const FName DomainName, TArray<TPair<UGroomAsset*, FGroupElements>>& InvocationGroups,  const bool bSourceElements = false)
 	{
 		InvocationGroups.Reset();
 		for(const UGroomComponent* GroomComponent : GroomComponents)
@@ -374,7 +371,7 @@ namespace UE::Groom::Private
 			if(GroomComponent)
 			{
 				FGroupElements GroupData;
-				GetGroomGroupElementOffsets(GroomComponent, DomainName, GroupData.GroupOffsets, GroupSize, bSourceElements);
+				GetGroomGroupElementOffsets(GroomComponent, DomainName, GroupData.GroupOffsets, bSourceElements);
 				
 				const uint32 NumGroups = GroomComponent->GetGroupCount();
 				for (uint32 GroupIndex = 0; GroupIndex < NumGroups; ++GroupIndex)
diff --git a/Engine/Plugins/Runtime/HairStrands/Source/HairStrandsDeformer/Private/DeformerGroomSolverRead.cpp b/Engine/Plugins/Runtime/HairStrands/Source/HairStrandsDeformer/Private/DeformerGroomSolverRead.cpp
--- a/Engine/Plugins/Runtime/HairStrands/Source/HairStrandsDeformer/Private/DeformerGroomSolverRead.cpp
+++ b/Engine/Plugins/Runtime/HairStrands/Source/HairStrandsDeformer/Private/DeformerGroomSolverRead.cpp
@@ -1,4 +1,4 @@
-﻿// Copyright Epic Games, Inc. All Rights Reserved.
+// Copyright Epic Games, Inc. All Rights Reserved.
 
 #include "DeformerGroomSolverRead.h"
 
@@ -273,7 +273,7 @@ FComputeDataProviderRenderProxy* UOptimusGroomSolverReadDataProvider::GetRenderP
 				{
 					bResetSimulationCount = 0;
 				}
-			}
+			} 
 			bResetSimulationTrigger = (bResetSimulationCount < MaxResetCount);
 			if(bResetSimulationTrigger)
 			{
@@ -295,9 +295,6 @@ FOptimusGroomSolverReadDataProviderProxy::FOptimusGroomSolverReadDataProviderPro
 	UE::Groom::Private::GetGroomInvocationElementCounts(GroomComponents,
 		UOptimusGroomAssetComponentSource::FGuidesExecutionDomains::Points, PointsCounts);
 
-	UE::Groom::Private::GetGroomInvocationElementCounts(GroomComponents,
-		UOptimusGroomAssetComponentSource::FGuidesExecutionDomains::Points, AlignedPoints, UOptimusGroomSolverComponentSource::GroupSize);
-
 	UE::Groom::Private::GetGroomInvocationElementCounts(GroomComponents,
 		UOptimusGroomAssetComponentSource::FGuidesExecutionDomains::Curves, CurvesCounts);
 	
@@ -327,7 +324,7 @@ void FOptimusGroomSolverReadDataProviderProxy::GatherPermutations(FPermutationDa
 {
 }
 
-FORCEINLINE void FillElementsData(const TArray<int32>& ElementsCounts, const TArray<int32>& AlignedCounts,
+FORCEINLINE void FillElementsData(const TArray<int32>& ElementsCounts, 
 	TArray<uint32>& ObjectOffsets, TArray<uint32>& ObjectCounts, TArray<int32>& ElementObjects, int32& NumElements)
 {
 	const uint32 NumObjects = ElementsCounts.Num();
@@ -339,7 +336,7 @@ FORCEINLINE void FillElementsData(const TArray<int32>& ElementsCounts, const TAr
 	{
 		ObjectCounts[ObjectIndex] = static_cast<uint32>(ElementsCounts[ObjectIndex]);
 		ObjectOffsets[ObjectIndex] = ElementOffset;
-		ElementOffset += AlignedCounts[ObjectIndex];
+		ElementOffset += ElementsCounts[ObjectIndex];
 	}
 	
 	ObjectOffsets[NumObjects] = ElementOffset;
@@ -362,7 +359,7 @@ void FOptimusGroomSolverReadDataProviderProxy::CreateInternalBuffers(FRDGBuilder
 	TArray<uint32> ObjectNumPoints;
 	TArray<int32> PointObjectIndices;
 
-	FillElementsData(PointsCounts, AlignedPoints, ObjectPointOffsets,  ObjectNumPoints,
+	FillElementsData(PointsCounts, ObjectPointOffsets,  ObjectNumPoints,
 		PointObjectIndices, NumPoints);
 
 	FString BufferName = FString("Hair.Deformer.ObjectPointOffsets");
@@ -394,7 +391,7 @@ void FOptimusGroomSolverReadDataProviderProxy::CreateInternalBuffers(FRDGBuilder
 	TArray<uint32> ObjectNumCurves;
 	TArray<int32> CurveObjectIndices;
 	
-	FillElementsData(CurvesCounts, CurvesCounts, ObjectCurveOffsets,  ObjectNumCurves,
+	FillElementsData(CurvesCounts, ObjectCurveOffsets,  ObjectNumCurves,
 	CurveObjectIndices, NumCurves);
 
 	BufferName = FString("Hair.Deformer.ObjectCurveOffsets");
diff --git a/Engine/Plugins/Runtime/HairStrands/Source/HairStrandsDeformer/Private/DeformerGroomSolverRead.h b/Engine/Plugins/Runtime/HairStrands/Source/HairStrandsDeformer/Private/DeformerGroomSolverRead.h
--- a/Engine/Plugins/Runtime/HairStrands/Source/HairStrandsDeformer/Private/DeformerGroomSolverRead.h
+++ b/Engine/Plugins/Runtime/HairStrands/Source/HairStrandsDeformer/Private/DeformerGroomSolverRead.h
@@ -1,4 +1,4 @@
-﻿// Copyright Epic Games, Inc. All Rights Reserved.
+// Copyright Epic Games, Inc. All Rights Reserved.
 
 #pragma once
 
@@ -128,9 +128,6 @@ private:
 	/** Invocation curves count */
 	TArray<int32> CurvesCounts;
 
-	/** Invocation aligned points */
-	TArray<int32> AlignedPoints;
-
 	/** Total number of points */
 	int32 NumPoints = 0;
 
diff --git a/Engine/Plugins/Runtime/HairStrands/Source/HairStrandsSolver/Private/GroomSolverComponent.cpp b/Engine/Plugins/Runtime/HairStrands/Source/HairStrandsSolver/Private/GroomSolverComponent.cpp
--- a/Engine/Plugins/Runtime/HairStrands/Source/HairStrandsSolver/Private/GroomSolverComponent.cpp
+++ b/Engine/Plugins/Runtime/HairStrands/Source/HairStrandsSolver/Private/GroomSolverComponent.cpp
@@ -165,7 +165,7 @@ void UGroomSolverComponent::SelectDynamicCurves()
 						SolverSettings.ObjectDistanceLods.Add(ObjectLod);
 						
 						CurveOffset += NumTotalCurves;
-						PointOffset += FMath::CeilToInt32(NumTotalPoints / static_cast<float>(GroupSize)) * GroupSize;
+						PointOffset += NumTotalPoints;
 					}
 				
 					// Keep it for now for debugging since it is WIP
