Removed UFUNCTION from DEFINE_GROOM_MEMBER_ACCESSOR_BP macros.
UHT doesn't parse cpp, so no need to have these.

diff --git a/Engine/Plugins/Runtime/HairStrands/Source/HairStrandsCore/Private/GroomAsset.cpp b/Engine/Plugins/Runtime/HairStrands/Source/HairStrandsCore/Private/GroomAsset.cpp
--- a/Engine/Plugins/Runtime/HairStrands/Source/HairStrandsCore/Private/GroomAsset.cpp
+++ b/Engine/Plugins/Runtime/HairStrands/Source/HairStrandsCore/Private/GroomAsset.cpp
@@ -3928,16 +3928,12 @@ FName UGroomAsset::GetAssetPathName(int32 LODIndex)
 
 #define DEFINE_GROOM_MEMBER_ACCESSOR_BP(Type, Name)\
 	GROOMASSET_DEFINE_MEMBER_NAME(Name)\
-	UFUNCTION(BlueprintGetter)\
 	GROOMASSET_DEFINE_GET(Type, Name, , , const)\
-	UFUNCTION(BlueprintSetter)\
 	GROOMASSET_DEFINE_SET(Type, Name, , )
 
 #define DEFINE_GROOM_MEMBER_ACCESSOR_BP_ARRAY(Type, Name)\
 	GROOMASSET_DEFINE_MEMBER_NAME(Name)\
-	UFUNCTION(BlueprintGetter)\
 	GROOMASSET_DEFINE_GET(Type, Name, &, , )\
-	UFUNCTION(BlueprintSetter)\
 	GROOMASSET_DEFINE_SET(Type, Name, &, const)\
 	GROOMASSET_DEFINE_GET(Type, Name, &, const, const)
 
diff --git a/Engine/Plugins/Runtime/HairStrands/Source/HairStrandsCore/Private/GroomBindingAsset.cpp b/Engine/Plugins/Runtime/HairStrands/Source/HairStrandsCore/Private/GroomBindingAsset.cpp
--- a/Engine/Plugins/Runtime/HairStrands/Source/HairStrandsCore/Private/GroomBindingAsset.cpp
+++ b/Engine/Plugins/Runtime/HairStrands/Source/HairStrandsCore/Private/GroomBindingAsset.cpp
@@ -1702,7 +1702,6 @@ FName UGroomBindingAsset::GetAssetPathName(int32 LODIndex)
 	}
 
 #define DEFINE_GROOMBINDING_MEMBER_GETTER(Type, Access, Name, GetConst, SetConst)\
-	UFUNCTION(BlueprintGetter)\
 	GetConst Type Access UGroomBindingAsset::Get##Name() const\
 	{\
 		WaitUntilAsyncPropertyReleased(EGroomBindingAsyncProperties::Name, EGroomBindingAsyncPropertyLockType::ReadOnly); \
@@ -1712,7 +1711,6 @@ FName UGroomBindingAsset::GetAssetPathName(int32 LODIndex)
 	}
 
 #define DEFINE_GROOMBINDING_MEMBER_SETTER(Type, Access, Name, GetConst, SetConst)\
-	UFUNCTION(BlueprintSetter)\
 	void UGroomBindingAsset::Set##Name(SetConst Type Access In##Name)\
 	{\
 		WaitUntilAsyncPropertyReleased(EGroomBindingAsyncProperties::Name, EGroomBindingAsyncPropertyLockType::WriteOnly); \
