Updated headers using UnrealCodeFixup to make sure dllstorage is on methods/staticvars instead of types. Part 4/n

diff --git a/Engine/Plugins/Online/OnlineSubsystemUtils/Source/OnlineBlueprintSupport/Classes/K2Node_InAppPurchase2.h b/Engine/Plugins/Online/OnlineSubsystemUtils/Source/OnlineBlueprintSupport/Classes/K2Node_InAppPurchase2.h
--- a/Engine/Plugins/Online/OnlineSubsystemUtils/Source/OnlineBlueprintSupport/Classes/K2Node_InAppPurchase2.h
+++ b/Engine/Plugins/Online/OnlineSubsystemUtils/Source/OnlineBlueprintSupport/Classes/K2Node_InAppPurchase2.h
@@ -5,8 +5,8 @@
 #include "K2Node_BaseAsyncTask.h"
 #include "K2Node_InAppPurchase2.generated.h"
 
-UCLASS()
-class ONLINEBLUEPRINTSUPPORT_API UK2Node_InAppPurchase2 : public UK2Node_BaseAsyncTask
+UCLASS(MinimalAPI)
+class UK2Node_InAppPurchase2 : public UK2Node_BaseAsyncTask
 {
 	GENERATED_UCLASS_BODY()
 };
diff --git a/Engine/Plugins/Online/OnlineSubsystemUtils/Source/OnlineBlueprintSupport/Classes/K2Node_InAppPurchaseCheckout.h b/Engine/Plugins/Online/OnlineSubsystemUtils/Source/OnlineBlueprintSupport/Classes/K2Node_InAppPurchaseCheckout.h
--- a/Engine/Plugins/Online/OnlineSubsystemUtils/Source/OnlineBlueprintSupport/Classes/K2Node_InAppPurchaseCheckout.h
+++ b/Engine/Plugins/Online/OnlineSubsystemUtils/Source/OnlineBlueprintSupport/Classes/K2Node_InAppPurchaseCheckout.h
@@ -5,8 +5,8 @@
 #include "K2Node_BaseAsyncTask.h"
 #include "K2Node_InAppPurchaseCheckout.generated.h"
 
-UCLASS()
-class ONLINEBLUEPRINTSUPPORT_API UK2Node_InAppPurchaseCheckout : public UK2Node_BaseAsyncTask
+UCLASS(MinimalAPI)
+class UK2Node_InAppPurchaseCheckout : public UK2Node_BaseAsyncTask
 {
 	GENERATED_UCLASS_BODY()
 };
diff --git a/Engine/Plugins/Online/OnlineSubsystemUtils/Source/OnlineBlueprintSupport/Classes/K2Node_InAppPurchaseFinalize.h b/Engine/Plugins/Online/OnlineSubsystemUtils/Source/OnlineBlueprintSupport/Classes/K2Node_InAppPurchaseFinalize.h
--- a/Engine/Plugins/Online/OnlineSubsystemUtils/Source/OnlineBlueprintSupport/Classes/K2Node_InAppPurchaseFinalize.h
+++ b/Engine/Plugins/Online/OnlineSubsystemUtils/Source/OnlineBlueprintSupport/Classes/K2Node_InAppPurchaseFinalize.h
@@ -5,8 +5,8 @@
 #include "K2Node_BaseAsyncTask.h"
 #include "K2Node_InAppPurchaseFinalize.generated.h"
 
-UCLASS()
-class ONLINEBLUEPRINTSUPPORT_API UK2Node_InAppPurchaseFinalize : public UK2Node_BaseAsyncTask
+UCLASS(MinimalAPI)
+class UK2Node_InAppPurchaseFinalize : public UK2Node_BaseAsyncTask
 {
 	GENERATED_UCLASS_BODY()
 };
diff --git a/Engine/Plugins/Online/OnlineSubsystemUtils/Source/OnlineBlueprintSupport/Classes/K2Node_InAppPurchaseGetKnownReceipts.h b/Engine/Plugins/Online/OnlineSubsystemUtils/Source/OnlineBlueprintSupport/Classes/K2Node_InAppPurchaseGetKnownReceipts.h
--- a/Engine/Plugins/Online/OnlineSubsystemUtils/Source/OnlineBlueprintSupport/Classes/K2Node_InAppPurchaseGetKnownReceipts.h
+++ b/Engine/Plugins/Online/OnlineSubsystemUtils/Source/OnlineBlueprintSupport/Classes/K2Node_InAppPurchaseGetKnownReceipts.h
@@ -5,8 +5,8 @@
 #include "K2Node_BaseAsyncTask.h"
 #include "K2Node_InAppPurchaseGetKnownReceipts.generated.h"
 
-UCLASS()
-class ONLINEBLUEPRINTSUPPORT_API UK2Node_InAppPurchaseGetKnownReceipts : public UK2Node_BaseAsyncTask
+UCLASS(MinimalAPI)
+class UK2Node_InAppPurchaseGetKnownReceipts : public UK2Node_BaseAsyncTask
 {
 	GENERATED_UCLASS_BODY()
 };
diff --git a/Engine/Plugins/Online/OnlineSubsystemUtils/Source/OnlineBlueprintSupport/Classes/K2Node_InAppPurchaseQuery2.h b/Engine/Plugins/Online/OnlineSubsystemUtils/Source/OnlineBlueprintSupport/Classes/K2Node_InAppPurchaseQuery2.h
--- a/Engine/Plugins/Online/OnlineSubsystemUtils/Source/OnlineBlueprintSupport/Classes/K2Node_InAppPurchaseQuery2.h
+++ b/Engine/Plugins/Online/OnlineSubsystemUtils/Source/OnlineBlueprintSupport/Classes/K2Node_InAppPurchaseQuery2.h
@@ -5,8 +5,8 @@
 #include "K2Node_BaseAsyncTask.h"
 #include "K2Node_InAppPurchaseQuery2.generated.h"
 
-UCLASS()
-class ONLINEBLUEPRINTSUPPORT_API UK2Node_InAppPurchaseQuery2 : public UK2Node_BaseAsyncTask
+UCLASS(MinimalAPI)
+class UK2Node_InAppPurchaseQuery2 : public UK2Node_BaseAsyncTask
 {
 	GENERATED_UCLASS_BODY()
 };
diff --git a/Engine/Plugins/Online/OnlineSubsystemUtils/Source/OnlineBlueprintSupport/Classes/K2Node_InAppPurchaseQueryOwned2.h b/Engine/Plugins/Online/OnlineSubsystemUtils/Source/OnlineBlueprintSupport/Classes/K2Node_InAppPurchaseQueryOwned2.h
--- a/Engine/Plugins/Online/OnlineSubsystemUtils/Source/OnlineBlueprintSupport/Classes/K2Node_InAppPurchaseQueryOwned2.h
+++ b/Engine/Plugins/Online/OnlineSubsystemUtils/Source/OnlineBlueprintSupport/Classes/K2Node_InAppPurchaseQueryOwned2.h
@@ -5,8 +5,8 @@
 #include "K2Node_BaseAsyncTask.h"
 #include "K2Node_InAppPurchaseQueryOwned2.generated.h"
 
-UCLASS()
-class ONLINEBLUEPRINTSUPPORT_API UK2Node_InAppPurchaseQueryOwned2 : public UK2Node_BaseAsyncTask
+UCLASS(MinimalAPI)
+class UK2Node_InAppPurchaseQueryOwned2 : public UK2Node_BaseAsyncTask
 {
 	GENERATED_UCLASS_BODY()
 };
diff --git a/Engine/Plugins/Online/OnlineSubsystemUtils/Source/OnlineBlueprintSupport/Classes/K2Node_InAppPurchaseQueryOwnedProducts.h b/Engine/Plugins/Online/OnlineSubsystemUtils/Source/OnlineBlueprintSupport/Classes/K2Node_InAppPurchaseQueryOwnedProducts.h
--- a/Engine/Plugins/Online/OnlineSubsystemUtils/Source/OnlineBlueprintSupport/Classes/K2Node_InAppPurchaseQueryOwnedProducts.h
+++ b/Engine/Plugins/Online/OnlineSubsystemUtils/Source/OnlineBlueprintSupport/Classes/K2Node_InAppPurchaseQueryOwnedProducts.h
@@ -5,8 +5,8 @@
 #include "K2Node_BaseAsyncTask.h"
 #include "K2Node_InAppPurchaseQueryOwnedProducts.generated.h"
 
-UCLASS()
-class ONLINEBLUEPRINTSUPPORT_API UK2Node_InAppPurchaseQueryOwnedProducts : public UK2Node_BaseAsyncTask
+UCLASS(MinimalAPI)
+class UK2Node_InAppPurchaseQueryOwnedProducts : public UK2Node_BaseAsyncTask
 {
 	GENERATED_UCLASS_BODY()
 };
diff --git a/Engine/Plugins/Online/OnlineSubsystemUtils/Source/OnlineBlueprintSupport/Classes/K2Node_InAppPurchaseRestore2.h b/Engine/Plugins/Online/OnlineSubsystemUtils/Source/OnlineBlueprintSupport/Classes/K2Node_InAppPurchaseRestore2.h
--- a/Engine/Plugins/Online/OnlineSubsystemUtils/Source/OnlineBlueprintSupport/Classes/K2Node_InAppPurchaseRestore2.h
+++ b/Engine/Plugins/Online/OnlineSubsystemUtils/Source/OnlineBlueprintSupport/Classes/K2Node_InAppPurchaseRestore2.h
@@ -5,8 +5,8 @@
 #include "K2Node_BaseAsyncTask.h"
 #include "K2Node_InAppPurchaseRestore2.generated.h"
 
-UCLASS()
-class ONLINEBLUEPRINTSUPPORT_API UK2Node_InAppPurchaseRestore2 : public UK2Node_BaseAsyncTask
+UCLASS(MinimalAPI)
+class UK2Node_InAppPurchaseRestore2 : public UK2Node_BaseAsyncTask
 {
 	GENERATED_UCLASS_BODY()
 };
diff --git a/Engine/Plugins/Online/OnlineSubsystemUtils/Source/OnlineBlueprintSupport/Classes/K2Node_InAppPurchaseRestoreOwnedProducts.h b/Engine/Plugins/Online/OnlineSubsystemUtils/Source/OnlineBlueprintSupport/Classes/K2Node_InAppPurchaseRestoreOwnedProducts.h
--- a/Engine/Plugins/Online/OnlineSubsystemUtils/Source/OnlineBlueprintSupport/Classes/K2Node_InAppPurchaseRestoreOwnedProducts.h
+++ b/Engine/Plugins/Online/OnlineSubsystemUtils/Source/OnlineBlueprintSupport/Classes/K2Node_InAppPurchaseRestoreOwnedProducts.h
@@ -5,8 +5,8 @@
 #include "K2Node_BaseAsyncTask.h"
 #include "K2Node_InAppPurchaseRestoreOwnedProducts.generated.h"
 
-UCLASS()
-class ONLINEBLUEPRINTSUPPORT_API UK2Node_InAppPurchaseRestoreOwnedProducts : public UK2Node_BaseAsyncTask
+UCLASS(MinimalAPI)
+class UK2Node_InAppPurchaseRestoreOwnedProducts : public UK2Node_BaseAsyncTask
 {
 	GENERATED_UCLASS_BODY()
 };
diff --git a/Engine/Plugins/Online/OnlineSubsystemUtils/Source/OnlineBlueprintSupport/Classes/K2Node_InAppPurchaseUnprocessed2.h b/Engine/Plugins/Online/OnlineSubsystemUtils/Source/OnlineBlueprintSupport/Classes/K2Node_InAppPurchaseUnprocessed2.h
--- a/Engine/Plugins/Online/OnlineSubsystemUtils/Source/OnlineBlueprintSupport/Classes/K2Node_InAppPurchaseUnprocessed2.h
+++ b/Engine/Plugins/Online/OnlineSubsystemUtils/Source/OnlineBlueprintSupport/Classes/K2Node_InAppPurchaseUnprocessed2.h
@@ -5,8 +5,8 @@
 #include "K2Node_BaseAsyncTask.h"
 #include "K2Node_InAppPurchaseUnprocessed2.generated.h"
 
-UCLASS()
-class ONLINEBLUEPRINTSUPPORT_API UK2Node_InAppPurchaseUnprocessed2 : public UK2Node_BaseAsyncTask
+UCLASS(MinimalAPI)
+class UK2Node_InAppPurchaseUnprocessed2 : public UK2Node_BaseAsyncTask
 {
 	GENERATED_UCLASS_BODY()
 };
diff --git a/Engine/Plugins/Online/OnlineSubsystemUtils/Source/OnlineBlueprintSupport/Classes/K2Node_LeaderboardFlush.h b/Engine/Plugins/Online/OnlineSubsystemUtils/Source/OnlineBlueprintSupport/Classes/K2Node_LeaderboardFlush.h
--- a/Engine/Plugins/Online/OnlineSubsystemUtils/Source/OnlineBlueprintSupport/Classes/K2Node_LeaderboardFlush.h
+++ b/Engine/Plugins/Online/OnlineSubsystemUtils/Source/OnlineBlueprintSupport/Classes/K2Node_LeaderboardFlush.h
@@ -5,17 +5,21 @@
 #include "K2Node_BaseAsyncTask.h"
 #include "K2Node_LeaderboardFlush.generated.h"
 
-UCLASS()
-class ONLINEBLUEPRINTSUPPORT_API UK2Node_LeaderboardFlush : public UK2Node_BaseAsyncTask
+#define UE_API ONLINEBLUEPRINTSUPPORT_API
+
+UCLASS(MinimalAPI)
+class UK2Node_LeaderboardFlush : public UK2Node_BaseAsyncTask
 {
 	GENERATED_UCLASS_BODY()
 
 	//~ Begin UEdGraphNode Interface
-	virtual FText GetTooltipText() const override;
-	virtual FText GetNodeTitle(ENodeTitleType::Type TitleType) const override;
+	UE_API virtual FText GetTooltipText() const override;
+	UE_API virtual FText GetNodeTitle(ENodeTitleType::Type TitleType) const override;
 	//~ End UEdGraphNode Interface
 
 	//~ Begin UK2Node Interface
-	virtual FText GetMenuCategory() const override;
+	UE_API virtual FText GetMenuCategory() const override;
 	//~ End UK2Node Interface
 };
+
+#undef UE_API
diff --git a/Engine/Plugins/Online/OnlineSubsystemUtils/Source/OnlineBlueprintSupport/Classes/K2Node_LeaderboardQuery.h b/Engine/Plugins/Online/OnlineSubsystemUtils/Source/OnlineBlueprintSupport/Classes/K2Node_LeaderboardQuery.h
--- a/Engine/Plugins/Online/OnlineSubsystemUtils/Source/OnlineBlueprintSupport/Classes/K2Node_LeaderboardQuery.h
+++ b/Engine/Plugins/Online/OnlineSubsystemUtils/Source/OnlineBlueprintSupport/Classes/K2Node_LeaderboardQuery.h
@@ -5,8 +5,8 @@
 #include "K2Node_BaseAsyncTask.h"
 #include "K2Node_LeaderboardQuery.generated.h"
 
-UCLASS()
-class ONLINEBLUEPRINTSUPPORT_API UK2Node_LeaderboardQuery : public UK2Node_BaseAsyncTask
+UCLASS(MinimalAPI)
+class UK2Node_LeaderboardQuery : public UK2Node_BaseAsyncTask
 {
 	GENERATED_UCLASS_BODY()
 };
diff --git a/Engine/Plugins/PCG/Source/PCGEditor/Public/DataVisualizations/PCGBaseTextureDataVisualization.h b/Engine/Plugins/PCG/Source/PCGEditor/Public/DataVisualizations/PCGBaseTextureDataVisualization.h
--- a/Engine/Plugins/PCG/Source/PCGEditor/Public/DataVisualizations/PCGBaseTextureDataVisualization.h
+++ b/Engine/Plugins/PCG/Source/PCGEditor/Public/DataVisualizations/PCGBaseTextureDataVisualization.h
@@ -5,13 +5,17 @@
 #include "PCGSpatialDataVisualization.h"
 #include "PCGEditorModule.h"
 
+#define UE_API PCGEDITOR_API
+
 class UPCGData;
 
-class PCGEDITOR_API FPCGBaseTextureDataVisualization : public IPCGSpatialDataVisualization
+class FPCGBaseTextureDataVisualization : public IPCGSpatialDataVisualization
 {
 public:
 	// ~Begin IPCGDataVisualization interface
-	virtual TArray<TSharedPtr<FStreamableHandle>> LoadRequiredResources(const UPCGData* Data) const override;
-	virtual FPCGSetupSceneFunc GetViewportSetupFunc(const UPCGData* Data) const override;
+	UE_API virtual TArray<TSharedPtr<FStreamableHandle>> LoadRequiredResources(const UPCGData* Data) const override;
+	UE_API virtual FPCGSetupSceneFunc GetViewportSetupFunc(const UPCGData* Data) const override;
 	// ~End IPCGDataVisualization interface
 };
+
+#undef UE_API
diff --git a/Engine/Plugins/PCG/Source/PCGEditor/Public/DataVisualizations/PCGLandscapeDataVisualization.h b/Engine/Plugins/PCG/Source/PCGEditor/Public/DataVisualizations/PCGLandscapeDataVisualization.h
--- a/Engine/Plugins/PCG/Source/PCGEditor/Public/DataVisualizations/PCGLandscapeDataVisualization.h
+++ b/Engine/Plugins/PCG/Source/PCGEditor/Public/DataVisualizations/PCGLandscapeDataVisualization.h
@@ -4,15 +4,19 @@
 
 #include "PCGSpatialDataVisualization.h"
 
+#define UE_API PCGEDITOR_API
+
 class UPCGData;
 class UPCGPointData;
 struct FPCGContext;
 
 /** Overrides collapse behavior to target a given number of points generated, to avoid too coarse debug (points way bigger than data) or enormous amount of points (points way smaller than data). */
-class PCGEDITOR_API IPCGLandscapeDataVisualization : public IPCGSpatialDataVisualization
+class IPCGLandscapeDataVisualization : public IPCGSpatialDataVisualization
 {
 public:
 	// ~Begin IPCGSpatialDataVisualization interface
-	virtual const UPCGBasePointData* CollapseToDebugBasePointData(FPCGContext* Context, const UPCGData* Data) const override;
+	UE_API virtual const UPCGBasePointData* CollapseToDebugBasePointData(FPCGContext* Context, const UPCGData* Data) const override;
 	// ~End IPCGSpatialDataVisualization interface
 };
+
+#undef UE_API
diff --git a/Engine/Plugins/PCG/Source/PCGEditor/Public/DataVisualizations/PCGParamDataVisualization.h b/Engine/Plugins/PCG/Source/PCGEditor/Public/DataVisualizations/PCGParamDataVisualization.h
--- a/Engine/Plugins/PCG/Source/PCGEditor/Public/DataVisualizations/PCGParamDataVisualization.h
+++ b/Engine/Plugins/PCG/Source/PCGEditor/Public/DataVisualizations/PCGParamDataVisualization.h
@@ -4,11 +4,15 @@
 
 #include "PCGDataVisualization.h"
 
+#define UE_API PCGEDITOR_API
+
 class UPCGData;
 
-class PCGEDITOR_API IPCGParamDataVisualization : public IPCGDataVisualization
+class IPCGParamDataVisualization : public IPCGDataVisualization
 {
 public:
 	virtual void ExecuteDebugDisplay(FPCGContext* Context, const UPCGSettingsInterface* SettingsInterface, const UPCGData* Data, AActor* TargetActor) const { /* Do nothing. */ }
-	virtual FPCGTableVisualizerInfo GetTableVisualizerInfoWithDomain(const UPCGData* Data, const FPCGMetadataDomainID& DomainID) const override;
+	UE_API virtual FPCGTableVisualizerInfo GetTableVisualizerInfoWithDomain(const UPCGData* Data, const FPCGMetadataDomainID& DomainID) const override;
 };
+
+#undef UE_API
diff --git a/Engine/Plugins/PCG/Source/PCGEditor/Public/DataVisualizations/PCGSpatialDataVisualization.h b/Engine/Plugins/PCG/Source/PCGEditor/Public/DataVisualizations/PCGSpatialDataVisualization.h
--- a/Engine/Plugins/PCG/Source/PCGEditor/Public/DataVisualizations/PCGSpatialDataVisualization.h
+++ b/Engine/Plugins/PCG/Source/PCGEditor/Public/DataVisualizations/PCGSpatialDataVisualization.h
@@ -5,6 +5,8 @@
 #include "PCGDataVisualization.h"
 #include "PCGEditorModule.h"
 
+#define UE_API PCGEDITOR_API
+
 class AActor;
 class UPCGData;
 class UPCGBasePointData;
@@ -13,26 +15,26 @@ struct FPCGContext;
 struct FPCGDebugVisualizationSettings;
 
 /** Default implementation for spatial data. Collapses to a PointData representation. */
-class PCGEDITOR_API IPCGSpatialDataVisualization : public IPCGDataVisualization
+class IPCGSpatialDataVisualization : public IPCGDataVisualization
 {
 public:
 	// ~Begin IPCGDataVisualization interface
-	virtual void ExecuteDebugDisplay(FPCGContext* Context, const UPCGSettingsInterface* SettingsInterface, const UPCGData* Data, AActor* TargetActor) const override;
-	virtual FPCGTableVisualizerInfo GetTableVisualizerInfoWithDomain(const UPCGData* Data, const FPCGMetadataDomainID& DomainID) const override;
+	UE_API virtual void ExecuteDebugDisplay(FPCGContext* Context, const UPCGSettingsInterface* SettingsInterface, const UPCGData* Data, AActor* TargetActor) const override;
+	UE_API virtual FPCGTableVisualizerInfo GetTableVisualizerInfoWithDomain(const UPCGData* Data, const FPCGMetadataDomainID& DomainID) const override;
 	
 	// Show the sampled points by default
 	virtual FPCGMetadataDomainID GetDefaultDomainForInspection(const UPCGData* Data) const override { return PCGMetadataDomainID::Elements; }
-	virtual FString GetDomainDisplayNameForInspection(const UPCGData* Data, const FPCGMetadataDomainID& DomainID) const override;
-	virtual TArray<FPCGMetadataDomainID> GetAllSupportedDomainsForInspection(const UPCGData* Data) const override;
-	virtual FPCGSetupSceneFunc GetViewportSetupFunc(const UPCGData* Data) const override;
+	UE_API virtual FString GetDomainDisplayNameForInspection(const UPCGData* Data, const FPCGMetadataDomainID& DomainID) const override;
+	UE_API virtual TArray<FPCGMetadataDomainID> GetAllSupportedDomainsForInspection(const UPCGData* Data) const override;
+	UE_API virtual FPCGSetupSceneFunc GetViewportSetupFunc(const UPCGData* Data) const override;
 	// ~End IPCGDataVisualization interface
 
 	UE_DEPRECATED(5.6, "Implement CollapseToDebugBasePointData instead")
-	virtual const UPCGPointData* CollapseToDebugPointData(FPCGContext* Context, const UPCGData* Data) const;
+	UE_API virtual const UPCGPointData* CollapseToDebugPointData(FPCGContext* Context, const UPCGData* Data) const;
 
-	virtual const UPCGBasePointData* CollapseToDebugBasePointData(FPCGContext* Context, const UPCGData* Data) const;
+	UE_API virtual const UPCGBasePointData* CollapseToDebugBasePointData(FPCGContext* Context, const UPCGData* Data) const;
 
-	virtual void ExecuteDebugDisplayHelper(
+	UE_API virtual void ExecuteDebugDisplayHelper(
 		const UPCGData* Data,
 		const FPCGDebugVisualizationSettings& DebugSettings,
 		FPCGContext* Context,
@@ -40,3 +42,5 @@ public:
 		const FPCGCrc& Crc,
 		const TFunction<void(class UInstancedStaticMeshComponent*)>& OnISMCCreatedCallback) const;
 };
+
+#undef UE_API
diff --git a/Engine/Plugins/PCG/Source/PCGEditor/Public/DataVisualizations/PCGSplineDataVisualization.h b/Engine/Plugins/PCG/Source/PCGEditor/Public/DataVisualizations/PCGSplineDataVisualization.h
--- a/Engine/Plugins/PCG/Source/PCGEditor/Public/DataVisualizations/PCGSplineDataVisualization.h
+++ b/Engine/Plugins/PCG/Source/PCGEditor/Public/DataVisualizations/PCGSplineDataVisualization.h
@@ -5,19 +5,23 @@
 #include "PCGSpatialDataVisualization.h"
 #include "PCGEditorModule.h"
 
+#define UE_API PCGEDITOR_API
+
 class UPCGData;
 struct FPCGContext;
 
-class PCGEDITOR_API IPCGSplineDataVisualization : public IPCGSpatialDataVisualization
+class IPCGSplineDataVisualization : public IPCGSpatialDataVisualization
 {
 public:
 	// ~Begin IPCGDataVisualization interface
-	virtual FPCGTableVisualizerInfo GetTableVisualizerInfoWithDomain(const UPCGData* Data, const FPCGMetadataDomainID& DomainID) const override;
-	virtual FPCGSetupSceneFunc GetViewportSetupFunc(const UPCGData* Data) const override;
+	UE_API virtual FPCGTableVisualizerInfo GetTableVisualizerInfoWithDomain(const UPCGData* Data, const FPCGMetadataDomainID& DomainID) const override;
+	UE_API virtual FPCGSetupSceneFunc GetViewportSetupFunc(const UPCGData* Data) const override;
 	// ~End IPCGDataVisualization interface
 
 	// ~Begin IPCGSpatialDataVisualization interface
 	/** Overrides collapse behavior to show the spline control points. */
-	virtual const UPCGBasePointData* CollapseToDebugBasePointData(FPCGContext* Context, const UPCGData* Data) const override;
+	UE_API virtual const UPCGBasePointData* CollapseToDebugBasePointData(FPCGContext* Context, const UPCGData* Data) const override;
 	// ~End IPCGSpatialDataVisualization interface
 };
+
+#undef UE_API
diff --git a/Engine/Plugins/PCG/Source/PCGEditor/Public/DataVisualizations/PCGVolumeDataVisualization.h b/Engine/Plugins/PCG/Source/PCGEditor/Public/DataVisualizations/PCGVolumeDataVisualization.h
--- a/Engine/Plugins/PCG/Source/PCGEditor/Public/DataVisualizations/PCGVolumeDataVisualization.h
+++ b/Engine/Plugins/PCG/Source/PCGEditor/Public/DataVisualizations/PCGVolumeDataVisualization.h
@@ -4,20 +4,22 @@
 
 #include "PCGSpatialDataVisualization.h"
 
+#define UE_API PCGEDITOR_API
+
 class UPCGData;
 struct FPCGContext;
 
 /** Overrides collapse behavior to target a given number of points generated, to avoid too coarse debug (points way bigger than data) or enormous amount of points (points way smaller than data). */
-class PCGEDITOR_API FPCGPrimitiveDataVisualization : public IPCGSpatialDataVisualization
+class FPCGPrimitiveDataVisualization : public IPCGSpatialDataVisualization
 {
 public:
 	// ~Begin IPCGSpatialDataVisualization interface
-	virtual const UPCGBasePointData* CollapseToDebugBasePointData(FPCGContext* Context, const UPCGData* Data) const override;
-	FPCGSetupSceneFunc GetViewportSetupFunc(const UPCGData* Data) const override;
+	UE_API virtual const UPCGBasePointData* CollapseToDebugBasePointData(FPCGContext* Context, const UPCGData* Data) const override;
+	UE_API FPCGSetupSceneFunc GetViewportSetupFunc(const UPCGData* Data) const override;
 	// ~End IPCGSpatialDataVisualization interface
 
 protected:
-	virtual void GetComponentsAndTransforms(const UPCGData* InData, TArray<TObjectPtr<UPrimitiveComponent>>& OutComponents, TArray<FTransform>& OutComponentTransforms) const;
+	UE_API virtual void GetComponentsAndTransforms(const UPCGData* InData, TArray<TObjectPtr<UPrimitiveComponent>>& OutComponents, TArray<FTransform>& OutComponentTransforms) const;
 };
 
 /** Specialization for volumes that don't always rely on their primitive component */
@@ -39,4 +41,6 @@ class FPCGCollisionWrapperDataVisualization : public FPCGPrimitiveDataVisualizat
 {
 protected:
 	virtual void GetComponentsAndTransforms(const UPCGData* InData, TArray<TObjectPtr<UPrimitiveComponent>>& OutComponents, TArray<FTransform>& OutComponentTransforms) const override;
-};
\ No newline at end of file
+};
+
+#undef UE_API
diff --git a/Engine/Plugins/PCG/Source/PCGEditor/Public/PCGEditorSettings.h b/Engine/Plugins/PCG/Source/PCGEditor/Public/PCGEditorSettings.h
--- a/Engine/Plugins/PCG/Source/PCGEditor/Public/PCGEditorSettings.h
+++ b/Engine/Plugins/PCG/Source/PCGEditor/Public/PCGEditorSettings.h
@@ -7,6 +7,8 @@
 
 #include "PCGEditorSettings.generated.h"
 
+#define UE_API PCGEDITOR_API
+
 class UPCGSettings;
 class UPCGBuilderSettings;
 
@@ -21,13 +23,13 @@ enum class EPCGEditorDoubleClickAction : uint8
 	JumpToSourceFile UMETA(DisplayName = "View Source File (C++)")
 };
 
-UCLASS(config=Editor, meta=(DisplayName="PCG"))
-class PCGEDITOR_API UPCGEditorProjectSettings : public UDeveloperSettings
+UCLASS(MinimalAPI, config=Editor, meta=(DisplayName="PCG"))
+class UPCGEditorProjectSettings : public UDeveloperSettings
 {
 	GENERATED_BODY()
 
 public:
-	UPCGEditorProjectSettings(const FObjectInitializer& ObjectInitializer);
+	UE_API UPCGEditorProjectSettings(const FObjectInitializer& ObjectInitializer);
 		
 	/** Default Builder Settings to use when running the PCGWorldPartitionBuilder */
 	UPROPERTY(EditAnywhere, config, Category = Builder)
@@ -41,16 +43,16 @@ public:
 	bool bAutoRefreshGPUStaticMeshSpawners = true;
 };
 
-UCLASS(config=EditorPerProjectUserSettings)
-class PCGEDITOR_API UPCGEditorSettings : public UObject
+UCLASS(MinimalAPI, config=EditorPerProjectUserSettings)
+class UPCGEditorSettings : public UObject
 {
 	GENERATED_BODY()
 
 public:
-	UPCGEditorSettings(const FObjectInitializer& ObjectInitializer);
+	UE_API UPCGEditorSettings(const FObjectInitializer& ObjectInitializer);
 
-	FLinearColor GetColor(UPCGSettings* InSettings) const;
-	FLinearColor GetPinColor(const FEdGraphPinType& InPinType) const;
+	UE_API FLinearColor GetColor(UPCGSettings* InSettings) const;
+	UE_API FLinearColor GetPinColor(const FEdGraphPinType& InPinType) const;
 
 	/** Default node color */
 	UPROPERTY(EditAnywhere, config, Category = Node, meta = (HideAlphaChannel))
@@ -285,3 +287,5 @@ public:
 	UPROPERTY()
 	bool bEnableNavigateToNativeNodes = true;
 };
+
+#undef UE_API
diff --git a/Engine/Plugins/PCG/Source/PCGEditor/Public/PCGLevelToAsset.h b/Engine/Plugins/PCG/Source/PCGEditor/Public/PCGLevelToAsset.h
--- a/Engine/Plugins/PCG/Source/PCGEditor/Public/PCGLevelToAsset.h
+++ b/Engine/Plugins/PCG/Source/PCGEditor/Public/PCGLevelToAsset.h
@@ -8,38 +8,42 @@
 
 #include "PCGLevelToAsset.generated.h"
 
+#define UE_API PCGEDITOR_API
+
 class UPackage;
 
-UCLASS(BlueprintType, Blueprintable, meta = (ShowWorldContextPin))
-class PCGEDITOR_API UPCGLevelToAsset : public UPCGAssetExporter
+UCLASS(MinimalAPI, BlueprintType, Blueprintable, meta = (ShowWorldContextPin))
+class UPCGLevelToAsset : public UPCGAssetExporter
 {
 	GENERATED_BODY()
 
 public:
 	/** Creates/updates a PCG Asset per given world. Allows exporter subclassing by passing in a Subclass. */
-	static void CreateOrUpdatePCGAssets(const TArray<FAssetData>& WorldAssets, const FPCGAssetExporterParameters& Parameters = FPCGAssetExporterParameters(), TSubclassOf<UPCGLevelToAsset> Subclass = {});
+	static UE_API void CreateOrUpdatePCGAssets(const TArray<FAssetData>& WorldAssets, const FPCGAssetExporterParameters& Parameters = FPCGAssetExporterParameters(), TSubclassOf<UPCGLevelToAsset> Subclass = {});
 
 	/** Creates/Updates a PCG Asset for a specific world. Allows exporter subclassing by passing in a Subclass. Will return null if it fails, or the package that was modified on success. */
-	static UPackage* CreateOrUpdatePCGAsset(TSoftObjectPtr<UWorld> World, const FPCGAssetExporterParameters& Parameters = FPCGAssetExporterParameters(), TSubclassOf<UPCGLevelToAsset> Subclass = {});
+	static UE_API UPackage* CreateOrUpdatePCGAsset(TSoftObjectPtr<UWorld> World, const FPCGAssetExporterParameters& Parameters = FPCGAssetExporterParameters(), TSubclassOf<UPCGLevelToAsset> Subclass = {});
 
 	/** Creates/Updates a PCG Asset for a specific world. Allows exporter subclassing (and settings creation by extension). Will return null if it fails, or the package that was modified on success. */
-	static UPackage* CreateOrUpdatePCGAsset(UWorld* World, const FPCGAssetExporterParameters& Parameters = FPCGAssetExporterParameters(), TSubclassOf<UPCGLevelToAsset> Subclass = {});
+	static UE_API UPackage* CreateOrUpdatePCGAsset(UWorld* World, const FPCGAssetExporterParameters& Parameters = FPCGAssetExporterParameters(), TSubclassOf<UPCGLevelToAsset> Subclass = {});
 
 	/** Parses the world and fills in the provided data asset. Implement this in BP to drive the generation in a custom manner. */
 	UFUNCTION(BlueprintNativeEvent, meta = (DisplayName = "Export World", ForceAsFunction))
-	bool BP_ExportWorld(UWorld* World, const FString& PackageName, UPCGDataAsset* Asset);
+	UE_API bool BP_ExportWorld(UWorld* World, const FString& PackageName, UPCGDataAsset* Asset);
 
 	UFUNCTION(BlueprintCallable, Category="PCG|IO")
-	void SetWorld(UWorld* World);
+	UE_API void SetWorld(UWorld* World);
 
 	UFUNCTION(BlueprintCallable, Category = "PCG|IO")
-	UWorld* GetWorld() const;
+	UE_API UWorld* GetWorld() const;
 
 protected:
 	//~Being UPCGAssetExporter interface
-	virtual bool ExportAsset(const FString& PackageName, UPCGDataAsset* Asset) override;
-	virtual UPackage* UpdateAsset(const FAssetData& PCGAsset) override;
+	UE_API virtual bool ExportAsset(const FString& PackageName, UPCGDataAsset* Asset) override;
+	UE_API virtual UPackage* UpdateAsset(const FAssetData& PCGAsset) override;
 	//~End UPCGAssetExporter interface
 
 	UWorld* WorldToExport = nullptr;
-};
\ No newline at end of file
+};
+
+#undef UE_API
diff --git a/Engine/Plugins/Runtime/AndroidDeviceProfileSelector/Source/AndroidDeviceProfileSelector/Public/AndroidDeviceProfileSelector.h b/Engine/Plugins/Runtime/AndroidDeviceProfileSelector/Source/AndroidDeviceProfileSelector/Public/AndroidDeviceProfileSelector.h
--- a/Engine/Plugins/Runtime/AndroidDeviceProfileSelector/Source/AndroidDeviceProfileSelector/Public/AndroidDeviceProfileSelector.h
+++ b/Engine/Plugins/Runtime/AndroidDeviceProfileSelector/Source/AndroidDeviceProfileSelector/Public/AndroidDeviceProfileSelector.h
@@ -5,6 +5,8 @@
 #include "Containers/Map.h"
 #include "UObject/NameTypes.h"
 
+#define UE_API ANDROIDDEVICEPROFILESELECTOR_API
+
 namespace FAndroidProfileSelectorSourceProperties
 {
 	static FName SRC_GPUFamily(TEXT("SRC_GPUFamily"));
@@ -29,15 +31,17 @@ namespace FAndroidProfileSelectorSourceProperties
 	static FName SRC_InsetsBottom(TEXT("SRC_InsetsBottom"));
 };
 
-class ANDROIDDEVICEPROFILESELECTOR_API FAndroidDeviceProfileSelector
+class FAndroidDeviceProfileSelector
 {
 	// Container of various device properties used for device profile matching.
-	static TMap<FName, FString> SelectorProperties;
-	static void VerifySelectorParams();
+	static UE_API TMap<FName, FString> SelectorProperties;
+	static UE_API void VerifySelectorParams();
 public:
 
-	static FString FindMatchingProfile(const FString& FallbackProfileName);
-	static int32 GetNumProfiles();
+	static UE_API FString FindMatchingProfile(const FString& FallbackProfileName);
+	static UE_API int32 GetNumProfiles();
 	static void SetSelectorProperties(const TMap<FName, FString>& Params) { SelectorProperties = Params; VerifySelectorParams(); }
 	static const TMap<FName, FString>& GetSelectorProperties() { return SelectorProperties; }
 };
+
+#undef UE_API
diff --git a/Engine/Plugins/Runtime/AppleImageUtils/Source/AppleImageUtilsBlueprintSupport/Classes/AppleImageUtilsBlueprintSupport.h b/Engine/Plugins/Runtime/AppleImageUtils/Source/AppleImageUtilsBlueprintSupport/Classes/AppleImageUtilsBlueprintSupport.h
--- a/Engine/Plugins/Runtime/AppleImageUtils/Source/AppleImageUtilsBlueprintSupport/Classes/AppleImageUtilsBlueprintSupport.h
+++ b/Engine/Plugins/Runtime/AppleImageUtils/Source/AppleImageUtilsBlueprintSupport/Classes/AppleImageUtilsBlueprintSupport.h
@@ -6,29 +6,29 @@
 
 #include "AppleImageUtilsBlueprintSupport.generated.h"
 
-UCLASS()
-class APPLEIMAGEUTILSBLUEPRINTSUPPORT_API UK2Node_ConvertToJPEG :
+UCLASS(MinimalAPI)
+class UK2Node_ConvertToJPEG :
 	public UK2Node_BaseAsyncTask
 {
 	GENERATED_UCLASS_BODY()
 };
 
-UCLASS()
-class APPLEIMAGEUTILSBLUEPRINTSUPPORT_API UK2Node_ConvertToHEIF :
+UCLASS(MinimalAPI)
+class UK2Node_ConvertToHEIF :
 	public UK2Node_BaseAsyncTask
 {
 	GENERATED_UCLASS_BODY()
 };
 
-UCLASS()
-class APPLEIMAGEUTILSBLUEPRINTSUPPORT_API UK2Node_ConvertToTIFF :
+UCLASS(MinimalAPI)
+class UK2Node_ConvertToTIFF :
 	public UK2Node_BaseAsyncTask
 {
 	GENERATED_UCLASS_BODY()
 };
 
-UCLASS()
-class APPLEIMAGEUTILSBLUEPRINTSUPPORT_API UK2Node_ConvertToPNG :
+UCLASS(MinimalAPI)
+class UK2Node_ConvertToPNG :
 	public UK2Node_BaseAsyncTask
 {
 	GENERATED_UCLASS_BODY()
diff --git a/Engine/Plugins/Runtime/AudioModulation/Source/AudioModulationEditor/Public/AudioModulationClassTemplates.h b/Engine/Plugins/Runtime/AudioModulation/Source/AudioModulationEditor/Public/AudioModulationClassTemplates.h
--- a/Engine/Plugins/Runtime/AudioModulation/Source/AudioModulationEditor/Public/AudioModulationClassTemplates.h
+++ b/Engine/Plugins/Runtime/AudioModulation/Source/AudioModulationEditor/Public/AudioModulationClassTemplates.h
@@ -9,14 +9,14 @@
 #include "AudioModulationClassTemplates.generated.h"
 
 
-UCLASS(Abstract)
-class AUDIOMODULATIONEDITOR_API USoundModulationClassTemplate : public UPluginClassTemplate
+UCLASS(MinimalAPI, Abstract)
+class USoundModulationClassTemplate : public UPluginClassTemplate
 {
 	GENERATED_UCLASS_BODY()
 };
 
-UCLASS()
-class AUDIOMODULATIONEDITOR_API USoundModulationGeneratorClassTemplate : public USoundModulationClassTemplate
+UCLASS(MinimalAPI)
+class USoundModulationGeneratorClassTemplate : public USoundModulationClassTemplate
 {
 	GENERATED_UCLASS_BODY()
 };
diff --git a/Engine/Plugins/Runtime/AudioSynesthesia/Source/AudioSynesthesiaEditor/Public/AudioSynesthesiaEditorModule.h b/Engine/Plugins/Runtime/AudioSynesthesia/Source/AudioSynesthesiaEditor/Public/AudioSynesthesiaEditorModule.h
--- a/Engine/Plugins/Runtime/AudioSynesthesia/Source/AudioSynesthesiaEditor/Public/AudioSynesthesiaEditorModule.h
+++ b/Engine/Plugins/Runtime/AudioSynesthesia/Source/AudioSynesthesiaEditor/Public/AudioSynesthesiaEditorModule.h
@@ -10,7 +10,7 @@ AUDIOSYNESTHESIAEDITOR_API DECLARE_LOG_CATEGORY_EXTERN(LogAudioSynesthesiaEditor
 extern const FName AudioSynesthesiaEditorAppIdentifier;
 
 /** audio analyzer editor module interface */
-class AUDIOSYNESTHESIAEDITOR_API IAudioSynesthesiaEditorModule :	public IModuleInterface
+class IAudioSynesthesiaEditorModule :	public IModuleInterface
 {
 public:
 	/** Registers audio analyzer asset actions. */
diff --git a/Engine/Plugins/Runtime/ComputeFramework/Source/ComputeFrameworkEditor/Private/ComputeFramework/ComputeKernelAssetActions.h b/Engine/Plugins/Runtime/ComputeFramework/Source/ComputeFrameworkEditor/Private/ComputeFramework/ComputeKernelAssetActions.h
--- a/Engine/Plugins/Runtime/ComputeFramework/Source/ComputeFrameworkEditor/Private/ComputeFramework/ComputeKernelAssetActions.h
+++ b/Engine/Plugins/Runtime/ComputeFramework/Source/ComputeFrameworkEditor/Private/ComputeFramework/ComputeKernelAssetActions.h
@@ -5,20 +5,24 @@
 #include "AssetTypeCategories.h"
 #include "AssetTypeActions_Base.h"
 
-class COMPUTEFRAMEWORKEDITOR_API FAssetTypeActions_ComputeKernel : public FAssetTypeActions_Base
+#define UE_API COMPUTEFRAMEWORKEDITOR_API
+
+class FAssetTypeActions_ComputeKernel : public FAssetTypeActions_Base
 {
 public:
-	FAssetTypeActions_ComputeKernel(EAssetTypeCategories::Type InAssetCategoryBit = EAssetTypeCategories::Animation);
+	UE_API FAssetTypeActions_ComputeKernel(EAssetTypeCategories::Type InAssetCategoryBit = EAssetTypeCategories::Animation);
 
 protected:
 	//~ Begin FAssetTypeActions_Base Interface.
-	FText GetName() const override;
-	FColor GetTypeColor() const override;
-	UClass* GetSupportedClass() const override;
-	uint32 GetCategories() override;
-	const TArray<FText>& GetSubMenus() const override;
+	UE_API FText GetName() const override;
+	UE_API FColor GetTypeColor() const override;
+	UE_API UClass* GetSupportedClass() const override;
+	UE_API uint32 GetCategories() override;
+	UE_API const TArray<FText>& GetSubMenus() const override;
 	//~ End FAssetTypeActions_Base Interface.
 
 private:
 	EAssetTypeCategories::Type AssetCategoryBit;
 };
+
+#undef UE_API
diff --git a/Engine/Plugins/Runtime/ComputeFramework/Source/ComputeFrameworkEditor/Private/ComputeFramework/ComputeKernelFromTextAssetActions.h b/Engine/Plugins/Runtime/ComputeFramework/Source/ComputeFrameworkEditor/Private/ComputeFramework/ComputeKernelFromTextAssetActions.h
--- a/Engine/Plugins/Runtime/ComputeFramework/Source/ComputeFrameworkEditor/Private/ComputeFramework/ComputeKernelFromTextAssetActions.h
+++ b/Engine/Plugins/Runtime/ComputeFramework/Source/ComputeFrameworkEditor/Private/ComputeFramework/ComputeKernelFromTextAssetActions.h
@@ -5,20 +5,24 @@
 #include "AssetTypeCategories.h"
 #include "AssetTypeActions_Base.h"
 
-class COMPUTEFRAMEWORKEDITOR_API FAssetTypeActions_ComputeKernelFromText : public FAssetTypeActions_Base
+#define UE_API COMPUTEFRAMEWORKEDITOR_API
+
+class FAssetTypeActions_ComputeKernelFromText : public FAssetTypeActions_Base
 {
 public:
-	FAssetTypeActions_ComputeKernelFromText(EAssetTypeCategories::Type InAssetCategoryBit = EAssetTypeCategories::Animation);
+	UE_API FAssetTypeActions_ComputeKernelFromText(EAssetTypeCategories::Type InAssetCategoryBit = EAssetTypeCategories::Animation);
 
 protected:
 	//~ Begin FAssetTypeActions_Base Interface.
-	FText GetName() const override;
-	FColor GetTypeColor() const override;
-	UClass* GetSupportedClass() const override;
-	uint32 GetCategories() override;
-	const TArray<FText>& GetSubMenus() const override;
+	UE_API FText GetName() const override;
+	UE_API FColor GetTypeColor() const override;
+	UE_API UClass* GetSupportedClass() const override;
+	UE_API uint32 GetCategories() override;
+	UE_API const TArray<FText>& GetSubMenus() const override;
 	//~ End FAssetTypeActions_Base Interface.
 
 private:
 	EAssetTypeCategories::Type AssetCategoryBit;
 };
+
+#undef UE_API
diff --git a/Engine/Plugins/Runtime/DataRegistry/Source/DataRegistryEditor/Public/DataRegistryEditorModule.h b/Engine/Plugins/Runtime/DataRegistry/Source/DataRegistryEditor/Public/DataRegistryEditorModule.h
--- a/Engine/Plugins/Runtime/DataRegistry/Source/DataRegistryEditor/Public/DataRegistryEditorModule.h
+++ b/Engine/Plugins/Runtime/DataRegistry/Source/DataRegistryEditor/Public/DataRegistryEditorModule.h
@@ -5,6 +5,8 @@
 #include "Modules/ModuleManager.h"
 #include "Toolkits/AssetEditorToolkit.h"
 
+#define UE_API DATAREGISTRYEDITOR_API
+
 class SToolTip;
 struct FDataRegistryId;
 struct FDataRegistryType;
@@ -18,7 +20,7 @@ DECLARE_DELEGATE_OneParam(FOnGetCustomDataRegistryItemNames, TArray<FName>&);
 /**
  * The public interface to this module
  */
-class DATAREGISTRYEDITOR_API FDataRegistryEditorModule : public IModuleInterface, public IHasMenuExtensibility, public IHasToolBarExtensibility
+class FDataRegistryEditorModule : public IModuleInterface, public IHasMenuExtensibility, public IHasToolBarExtensibility
 {
 public:
 
@@ -51,10 +53,10 @@ public:
 	 * @param bAllowClear If true, add None option to top
 	 * @param FilterStructName If not empty, restrict to types that use a struct that inherts from something named this
 	 */
-	static TSharedRef<SWidget> MakeDataRegistryTypeSelector(FOnGetDataRegistryDisplayText OnGetDisplayText, FOnSetDataRegistryType OnSetType, bool bAllowClear = true, FName FilterStructName = NAME_None);
+	static UE_API TSharedRef<SWidget> MakeDataRegistryTypeSelector(FOnGetDataRegistryDisplayText OnGetDisplayText, FOnSetDataRegistryType OnSetType, bool bAllowClear = true, FName FilterStructName = NAME_None);
 
 	/** Called to get list of valid data registry types, used by MakeDataRegistryTypeSelector or can be used for a property handle version */
-	static void GenerateDataRegistryTypeComboBoxStrings(TArray< TSharedPtr<FString> >& OutComboBoxStrings, TArray<TSharedPtr<SToolTip>>& OutToolTips, TArray<bool>& OutRestrictedItems, bool bAllowClear, FName FilterStructName);
+	static UE_API void GenerateDataRegistryTypeComboBoxStrings(TArray< TSharedPtr<FString> >& OutComboBoxStrings, TArray<TSharedPtr<SToolTip>>& OutToolTips, TArray<bool>& OutRestrictedItems, bool bAllowClear, FName FilterStructName);
 
 	/**
 	 * Creates a simple version of a Data Registry Item Name selector, not bound to a PropertyHandle.
@@ -66,7 +68,7 @@ public:
 	 * @param OnGetCustomItemNames Optional delegate to call to get custom names if OnGetId returns CustomContextType
 	 * @param bAllowClear If true, add None option to top
 	 */
-	static TSharedRef<SWidget> MakeDataRegistryItemNameSelector(FOnGetDataRegistryDisplayText OnGetDisplayText, FOnGetDataRegistryId OnGetId, FOnSetDataRegistryId OnSetId, FOnGetCustomDataRegistryItemNames OnGetCustomItemNames = FOnGetCustomDataRegistryItemNames(), bool bAllowClear = true);
+	static UE_API TSharedRef<SWidget> MakeDataRegistryItemNameSelector(FOnGetDataRegistryDisplayText OnGetDisplayText, FOnGetDataRegistryId OnGetId, FOnSetDataRegistryId OnSetId, FOnGetCustomDataRegistryItemNames OnGetCustomItemNames = FOnGetCustomDataRegistryItemNames(), bool bAllowClear = true);
 
 	/** Gets the extensibility managers for outside entities to extend data registry editor's menus and toolbars */
 	virtual TSharedPtr<FExtensibilityManager> GetMenuExtensibilityManager() override { return MenuExtensibilityManager; }
@@ -74,10 +76,12 @@ public:
 
 private:
 
-	virtual void StartupModule() override;
-	virtual void ShutdownModule() override;
+	UE_API virtual void StartupModule() override;
+	UE_API virtual void ShutdownModule() override;
 
 	TSharedPtr<FExtensibilityManager> MenuExtensibilityManager;
 	TSharedPtr<FExtensibilityManager> ToolBarExtensibilityManager;
 	TSharedPtr<class FDataRegistryGraphPanelPinFactory> DataRegistryGraphPanelPinFactory;
 };
+
+#undef UE_API
diff --git a/Engine/Plugins/Runtime/GameFeatures/Source/GameFeaturesEditor/Private/GameFeaturesEditorSettings.h b/Engine/Plugins/Runtime/GameFeatures/Source/GameFeaturesEditor/Private/GameFeaturesEditorSettings.h
--- a/Engine/Plugins/Runtime/GameFeatures/Source/GameFeaturesEditor/Private/GameFeaturesEditorSettings.h
+++ b/Engine/Plugins/Runtime/GameFeatures/Source/GameFeaturesEditor/Private/GameFeaturesEditorSettings.h
@@ -53,8 +53,8 @@ struct FPluginTemplateData
 };
 
 
-UCLASS(config = Editor, defaultconfig)
-class GAMEFEATURESEDITOR_API UGameFeaturesEditorSettings : public UDeveloperSettings
+UCLASS(MinimalAPI, config = Editor, defaultconfig)
+class UGameFeaturesEditorSettings : public UDeveloperSettings
 {
 	GENERATED_BODY()
 
diff --git a/Engine/Plugins/Runtime/GameFeatures/Source/GameFeaturesEditor/Public/GameFeaturePluginTemplate.h b/Engine/Plugins/Runtime/GameFeatures/Source/GameFeaturesEditor/Public/GameFeaturePluginTemplate.h
--- a/Engine/Plugins/Runtime/GameFeatures/Source/GameFeaturesEditor/Public/GameFeaturePluginTemplate.h
+++ b/Engine/Plugins/Runtime/GameFeatures/Source/GameFeaturesEditor/Public/GameFeaturePluginTemplate.h
@@ -6,26 +6,28 @@
 #include "GameFeatureData.h"
 #include "PluginDescriptor.h"
 
+#define UE_API GAMEFEATURESEDITOR_API
+
 /**
  * Used to create custom templates for GameFeaturePlugins.
  */
-struct GAMEFEATURESEDITOR_API FGameFeaturePluginTemplateDescription : public FPluginTemplateDescription
+struct FGameFeaturePluginTemplateDescription : public FPluginTemplateDescription
 {
-	FGameFeaturePluginTemplateDescription(FText InName, FText InDescription, FString InOnDiskPath, FString InDefaultSubfolder, FString InDefaultPluginName
+	UE_API FGameFeaturePluginTemplateDescription(FText InName, FText InDescription, FString InOnDiskPath, FString InDefaultSubfolder, FString InDefaultPluginName
 		, TSubclassOf<UGameFeatureData> GameFeatureDataClassOverride, FString GameFeatureDataNameOverride, EPluginEnabledByDefault InEnabledByDefault);
 
-	virtual bool ValidatePathForPlugin(const FString& ProposedAbsolutePluginPath, FText& OutErrorMessage) override;
-	virtual void UpdatePathWhenTemplateSelected(FString& InOutPath) override;
-	virtual void UpdatePathWhenTemplateUnselected(FString& InOutPath) override;
+	UE_API virtual bool ValidatePathForPlugin(const FString& ProposedAbsolutePluginPath, FText& OutErrorMessage) override;
+	UE_API virtual void UpdatePathWhenTemplateSelected(FString& InOutPath) override;
+	UE_API virtual void UpdatePathWhenTemplateUnselected(FString& InOutPath) override;
 
-	virtual void UpdatePluginNameTextWhenTemplateSelected(FText& OutPluginNameText) override;
-	virtual void UpdatePluginNameTextWhenTemplateUnselected(FText& OutPluginNameText) override;
+	UE_API virtual void UpdatePluginNameTextWhenTemplateSelected(FText& OutPluginNameText) override;
+	UE_API virtual void UpdatePluginNameTextWhenTemplateUnselected(FText& OutPluginNameText) override;
 
-	virtual void CustomizeDescriptorBeforeCreation(FPluginDescriptor& Descriptor) override;
-	virtual void OnPluginCreated(TSharedPtr<IPlugin> NewPlugin) override;
+	UE_API virtual void CustomizeDescriptorBeforeCreation(FPluginDescriptor& Descriptor) override;
+	UE_API virtual void OnPluginCreated(TSharedPtr<IPlugin> NewPlugin) override;
 
-	FString GetGameFeatureRoot() const;
-	bool IsRootedInGameFeaturesRoot(const FString& InStr) const;
+	UE_API FString GetGameFeatureRoot() const;
+	UE_API bool IsRootedInGameFeaturesRoot(const FString& InStr) const;
 
 	FString DefaultSubfolder;
 	FString DefaultPluginName;
@@ -33,3 +35,5 @@ struct GAMEFEATURESEDITOR_API FGameFeaturePluginTemplateDescription : public FPl
 	FString GameFeatureDataName;
 	EPluginEnabledByDefault PluginEnabledByDefault = EPluginEnabledByDefault::Disabled;
 };
+
+#undef UE_API
diff --git a/Engine/Plugins/Runtime/GameplayAbilities/Source/GameplayAbilities/Public/GameplayCue_Types.h b/Engine/Plugins/Runtime/GameplayAbilities/Source/GameplayAbilities/Public/GameplayCue_Types.h
--- a/Engine/Plugins/Runtime/GameplayAbilities/Source/GameplayAbilities/Public/GameplayCue_Types.h
+++ b/Engine/Plugins/Runtime/GameplayAbilities/Source/GameplayAbilities/Public/GameplayCue_Types.h
@@ -11,6 +11,8 @@
 #include "GameplayEffect.h"
 #include "GameplayCue_Types.generated.h"
 
+#define UE_API GAMEPLAYABILITIES_API
+
 class AGameplayCueNotify_Actor;
 class UAbilitySystemComponent;
 class UGameplayCueSet;
@@ -137,11 +139,13 @@ PRAGMA_ENABLE_DEPRECATION_WARNINGS
  *	FScopedGameplayCueSendContext
  *	Add this around code that sends multiple gameplay cues to allow grouping them into a smalkler number of cues for more efficient networking
  */
-struct GAMEPLAYABILITIES_API FScopedGameplayCueSendContext
+struct FScopedGameplayCueSendContext
 {
-	FScopedGameplayCueSendContext();
-	~FScopedGameplayCueSendContext();
+	UE_API FScopedGameplayCueSendContext();
+	UE_API ~FScopedGameplayCueSendContext();
 };
 
 /** Delegate for when GC notifies are added or removed from manager */
-DECLARE_MULTICAST_DELEGATE(FOnGameplayCueNotifyChange);
\ No newline at end of file
+DECLARE_MULTICAST_DELEGATE(FOnGameplayCueNotifyChange);
+
+#undef UE_API
diff --git a/Engine/Plugins/Runtime/GeometryCache/Source/GeometryCacheSequencer/Classes/GeometryCacheTrackEditor.h b/Engine/Plugins/Runtime/GeometryCache/Source/GeometryCacheSequencer/Classes/GeometryCacheTrackEditor.h
--- a/Engine/Plugins/Runtime/GeometryCache/Source/GeometryCacheSequencer/Classes/GeometryCacheTrackEditor.h
+++ b/Engine/Plugins/Runtime/GeometryCache/Source/GeometryCacheSequencer/Classes/GeometryCacheTrackEditor.h
@@ -5,6 +5,8 @@
 #include "ISequencerSection.h"
 #include "MovieSceneTrackEditor.h"
 
+#define UE_API GEOMETRYCACHESEQUENCER_API
+
 struct FBuildEditWidgetParams;
 
 struct FAssetData;
@@ -17,7 +19,7 @@ class UGeometryCacheComponent;
 /**
  * Tools for Geometry Cache tracks
  */
-class GEOMETRYCACHESEQUENCER_API FGeometryCacheTrackEditor : public FMovieSceneTrackEditor
+class FGeometryCacheTrackEditor : public FMovieSceneTrackEditor
 {
 public:
 	/**
@@ -25,7 +27,7 @@ public:
 	 *
 	 * @param InSequencer The sequencer instance to be used by this tool
 	 */
-	FGeometryCacheTrackEditor( TSharedRef<ISequencer> InSequencer );
+	UE_API FGeometryCacheTrackEditor( TSharedRef<ISequencer> InSequencer );
 
 	/** Virtual destructor. */
 	virtual ~FGeometryCacheTrackEditor() { }
@@ -36,25 +38,25 @@ public:
 	 * @param OwningSequencer The sequencer instance to be used by this tool
 	 * @return The new instance of this class
 	 */
-	static TSharedRef<ISequencerTrackEditor> CreateTrackEditor( TSharedRef<ISequencer> OwningSequencer );
+	static UE_API TSharedRef<ISequencerTrackEditor> CreateTrackEditor( TSharedRef<ISequencer> OwningSequencer );
 
 public:
 
 	// ISequencerTrackEditor interface
-	virtual FText GetDisplayName() const override;
-	virtual void BuildObjectBindingTrackMenu(FMenuBuilder& MenuBuilder, const TArray<FGuid>& ObjectBindings, const UClass* ObjectClass) override;
-	virtual TSharedRef<ISequencerSection> MakeSectionInterface( UMovieSceneSection& SectionObject, UMovieSceneTrack& Track, FGuid ObjectBinding ) override;
-	virtual bool SupportsSequence(UMovieSceneSequence* InSequence) const override;
-	virtual bool SupportsType( TSubclassOf<UMovieSceneTrack> Type ) const override;
-	virtual TSharedPtr<SWidget> BuildOutlinerEditWidget(const FGuid& ObjectBinding, UMovieSceneTrack* Track, const FBuildEditWidgetParams& Params) override;
-	virtual const FSlateBrush* GetIconBrush() const override;
+	UE_API virtual FText GetDisplayName() const override;
+	UE_API virtual void BuildObjectBindingTrackMenu(FMenuBuilder& MenuBuilder, const TArray<FGuid>& ObjectBindings, const UClass* ObjectClass) override;
+	UE_API virtual TSharedRef<ISequencerSection> MakeSectionInterface( UMovieSceneSection& SectionObject, UMovieSceneTrack& Track, FGuid ObjectBinding ) override;
+	UE_API virtual bool SupportsSequence(UMovieSceneSequence* InSequence) const override;
+	UE_API virtual bool SupportsType( TSubclassOf<UMovieSceneTrack> Type ) const override;
+	UE_API virtual TSharedPtr<SWidget> BuildOutlinerEditWidget(const FGuid& ObjectBinding, UMovieSceneTrack* Track, const FBuildEditWidgetParams& Params) override;
+	UE_API virtual const FSlateBrush* GetIconBrush() const override;
 
 private:
 
-	void BuildGeometryCacheTrack(TArray<FGuid> ObjectBindings, UMovieSceneTrack* Track);
+	UE_API void BuildGeometryCacheTrack(TArray<FGuid> ObjectBindings, UMovieSceneTrack* Track);
 
 	/** Delegate for AnimatablePropertyChanged in AddKey */
-	FKeyPropertyResult AddKeyInternal(FFrameNumber KeyTime, UObject* Object, UGeometryCacheComponent* GeomCacheComp, UMovieSceneTrack* Track);
+	UE_API FKeyPropertyResult AddKeyInternal(FFrameNumber KeyTime, UObject* Object, UGeometryCacheComponent* GeomCacheComp, UMovieSceneTrack* Track);
 };
 
 
@@ -101,3 +103,5 @@ private:
 	/** Cached start time valid only during resize */
 	FFrameNumber InitialStartTimeDuringResize;
 };
+
+#undef UE_API
diff --git a/Engine/Plugins/Runtime/GeometryCache/Source/GeometryCacheStreamer/Public/GeometryCacheStreamBase.h b/Engine/Plugins/Runtime/GeometryCache/Source/GeometryCacheStreamer/Public/GeometryCacheStreamBase.h
--- a/Engine/Plugins/Runtime/GeometryCache/Source/GeometryCacheStreamer/Public/GeometryCacheStreamBase.h
+++ b/Engine/Plugins/Runtime/GeometryCache/Source/GeometryCacheStreamer/Public/GeometryCacheStreamBase.h
@@ -6,6 +6,8 @@
 #include "IGeometryCacheStream.h"
 #include "HAL/CriticalSection.h"
 
+#define UE_API GEOMETRYCACHESTREAMER_API
+
 struct FGeometryCacheStreamReadRequest;
 
 /* Details about the animation to be streamed */
@@ -25,25 +27,25 @@ struct FGeometryCacheStreamDetails
  * Derived classes need to implement a way to retrieve the mesh data for a frame
  * through GetMeshData.
  */
-class GEOMETRYCACHESTREAMER_API FGeometryCacheStreamBase : public IGeometryCacheStream
+class FGeometryCacheStreamBase : public IGeometryCacheStream
 {
 public:
-	FGeometryCacheStreamBase(int32 ReadConcurrency, FGeometryCacheStreamDetails&& Details);
-	virtual ~FGeometryCacheStreamBase();
+	UE_API FGeometryCacheStreamBase(int32 ReadConcurrency, FGeometryCacheStreamDetails&& Details);
+	UE_API virtual ~FGeometryCacheStreamBase();
 
 	//~ Begin IGeometryCacheStream Interface
-	virtual void Prefetch(int32 StartFrameIndex, int32 NumFrames = 0) override;
-	virtual uint32 GetNumFramesNeeded() override;
-	virtual bool RequestFrameData() override;
-	virtual void UpdateRequestStatus(TArray<int32>& OutFramesCompleted) override;
-	virtual bool GetFrameData(int32 FrameIndex, FGeometryCacheMeshData& OutMeshData) override;
-	virtual int32 CancelRequests() override;
-	virtual const FGeometryCacheStreamStats& GetStreamStats() const override;
-	virtual void SetLimits(float MaxMemoryAllowed, float MaxCachedDuration) override;
+	UE_API virtual void Prefetch(int32 StartFrameIndex, int32 NumFrames = 0) override;
+	UE_API virtual uint32 GetNumFramesNeeded() override;
+	UE_API virtual bool RequestFrameData() override;
+	UE_API virtual void UpdateRequestStatus(TArray<int32>& OutFramesCompleted) override;
+	UE_API virtual bool GetFrameData(int32 FrameIndex, FGeometryCacheMeshData& OutMeshData) override;
+	UE_API virtual int32 CancelRequests() override;
+	UE_API virtual const FGeometryCacheStreamStats& GetStreamStats() const override;
+	UE_API virtual void SetLimits(float MaxMemoryAllowed, float MaxCachedDuration) override;
 	//~ End IGeometryCacheStream Interface
 
 	/* Updates the current position in the stream */
-	void UpdateCurrentFrameIndex(int32 FrameIndex);
+	UE_API void UpdateCurrentFrameIndex(int32 FrameIndex);
 
 protected:
 	/* Function called from main thread to prepare for GetMeshData */
@@ -52,10 +54,10 @@ protected:
 	/* Derived class must provide a way to get the mesh data for the given FrameIndex, called from worker threads */
 	virtual void GetMeshData(int32 FrameIndex, int32 ReadConcurrencyIndex, FGeometryCacheMeshData& OutMeshData) = 0;
 
-	void LoadFrameData(int32 FrameIndex);
-	void UpdateFramesNeeded(int32 StartIndex, int32 NumFrames);
-	void IncrementMemoryStat(const FGeometryCacheMeshData& MeshData);
-	void DecrementMemoryStat(const FGeometryCacheMeshData& MeshData);
+	UE_API void LoadFrameData(int32 FrameIndex);
+	UE_API void UpdateFramesNeeded(int32 StartIndex, int32 NumFrames);
+	UE_API void IncrementMemoryStat(const FGeometryCacheMeshData& MeshData);
+	UE_API void DecrementMemoryStat(const FGeometryCacheMeshData& MeshData);
 
 	TArray<int32> ReadIndices;
 	TArray<FGeometryCacheStreamReadRequest*> ReadRequestsPool;
@@ -79,3 +81,5 @@ protected:
 	std::atomic<bool> bCancellationRequested;
 	bool bCacheNeedsUpdate;
 };
+
+#undef UE_API
diff --git a/Engine/Plugins/Runtime/GeometryCache/Source/GeometryCacheStreamer/Public/GeometryCacheStreamerSettings.h b/Engine/Plugins/Runtime/GeometryCache/Source/GeometryCacheStreamer/Public/GeometryCacheStreamerSettings.h
--- a/Engine/Plugins/Runtime/GeometryCache/Source/GeometryCacheStreamer/Public/GeometryCacheStreamerSettings.h
+++ b/Engine/Plugins/Runtime/GeometryCache/Source/GeometryCacheStreamer/Public/GeometryCacheStreamerSettings.h
@@ -6,15 +6,17 @@
 
 #include "GeometryCacheStreamerSettings.generated.h"
 
+#define UE_API GEOMETRYCACHESTREAMER_API
+
 /** Settings for the GeometryCache streamer */
-UCLASS(config = Engine, meta = (DisplayName = "Geometry Cache"))
-class GEOMETRYCACHESTREAMER_API UGeometryCacheStreamerSettings : public UDeveloperSettings
+UCLASS(MinimalAPI, config = Engine, meta = (DisplayName = "Geometry Cache"))
+class UGeometryCacheStreamerSettings : public UDeveloperSettings
 {
 	GENERATED_BODY()
 
 public:
 	 
-	UGeometryCacheStreamerSettings();
+	UE_API UGeometryCacheStreamerSettings();
 
 	/** The amount of animation (in seconds) to stream ahead of time (per stream) */
 	UPROPERTY(config, EditAnywhere, BlueprintReadWrite, Category = "Geometry Cache Streamer", meta = (DisplayName = "Look-Ahead Buffer (in seconds)", ClampMin = "0.01", ClampMax = "3600.0"))
@@ -24,3 +26,5 @@ public:
 	UPROPERTY(config, EditAnywhere, BlueprintReadWrite, Category = "Geometry Cache Streamer", meta = (DisplayName = "Maximum Memory Allowed (in MB)", ClampMin = "1.0", ClampMax = "262144.0"))
 	float MaxMemoryAllowed;
 };
+
+#undef UE_API
diff --git a/Engine/Plugins/Runtime/GeometryCache/Source/GeometryCacheStreamer/Public/IGeometryCacheStreamer.h b/Engine/Plugins/Runtime/GeometryCache/Source/GeometryCacheStreamer/Public/IGeometryCacheStreamer.h
--- a/Engine/Plugins/Runtime/GeometryCache/Source/GeometryCacheStreamer/Public/IGeometryCacheStreamer.h
+++ b/Engine/Plugins/Runtime/GeometryCache/Source/GeometryCacheStreamer/Public/IGeometryCacheStreamer.h
@@ -4,15 +4,17 @@
 
 #include "HAL/Platform.h"
 
+#define UE_API GEOMETRYCACHESTREAMER_API
+
 struct FGeometryCacheMeshData;
 class IGeometryCacheStream;
 class UGeometryCacheTrack;
 
 /** Interface to the GeometryCache Streamer that streams data for registered GeometryCacheTracks */
-class GEOMETRYCACHESTREAMER_API IGeometryCacheStreamer
+class IGeometryCacheStreamer
 {
 public:
-	static IGeometryCacheStreamer& Get();
+	static UE_API IGeometryCacheStreamer& Get();
 
 	/** Register the given Track and its associated Stream with the Streamer. The Streamer takes ownership of the Stream */
 	virtual void RegisterTrack(UGeometryCacheTrack* Track, IGeometryCacheStream* Stream) = 0;
@@ -28,3 +30,5 @@ public:
 	 */
 	virtual bool TryGetFrameData(UGeometryCacheTrack* Track, int32 FrameIndex, FGeometryCacheMeshData& OutMeshData) = 0;
 };
+
+#undef UE_API
diff --git a/Engine/Plugins/Runtime/GeometryScripting/Source/GeometryScriptingEditor/Public/GeometryActors/EditorGeometryGenerationSubsystem.h b/Engine/Plugins/Runtime/GeometryScripting/Source/GeometryScriptingEditor/Public/GeometryActors/EditorGeometryGenerationSubsystem.h
--- a/Engine/Plugins/Runtime/GeometryScripting/Source/GeometryScriptingEditor/Public/GeometryActors/EditorGeometryGenerationSubsystem.h
+++ b/Engine/Plugins/Runtime/GeometryScripting/Source/GeometryScriptingEditor/Public/GeometryActors/EditorGeometryGenerationSubsystem.h
@@ -7,6 +7,8 @@
 
 #include "EditorGeometryGenerationSubsystem.generated.h"
 
+#define UE_API GEOMETRYSCRIPTINGEDITOR_API
+
 class AGeneratedDynamicMeshActor;
 
 
@@ -22,14 +24,14 @@ class AGeneratedDynamicMeshActor;
  * this regeneration will be more carefully managed via throttling / timeslicing / etc.
  * 
  */
-UCLASS()
-class GEOMETRYSCRIPTINGEDITOR_API UEditorGeometryGenerationSubsystem : public UEditorSubsystem
+UCLASS(MinimalAPI)
+class UEditorGeometryGenerationSubsystem : public UEditorSubsystem
 {
 	GENERATED_BODY()
 
 public:
-	virtual void Initialize(FSubsystemCollectionBase& Collection) override;
-	virtual void Deinitialize() override;
+	UE_API virtual void Initialize(FSubsystemCollectionBase& Collection) override;
+	UE_API virtual void Deinitialize() override;
 
 public:
 	UPROPERTY()
@@ -40,16 +42,16 @@ public:
 	// Static functions that simplify registration of an Actor w/ the Subsystem
 	//
 public:
-	static bool RegisterGeneratedMeshActor(AGeneratedDynamicMeshActor* Actor);
-	static void UnregisterGeneratedMeshActor(AGeneratedDynamicMeshActor* Actor);
+	static UE_API bool RegisterGeneratedMeshActor(AGeneratedDynamicMeshActor* Actor);
+	static UE_API void UnregisterGeneratedMeshActor(AGeneratedDynamicMeshActor* Actor);
 
 
 protected:
 	// callback connected to engine/editor shutdown events to set bIsShuttingDown, which disables the subsystem static functions above
-	virtual void OnShutdown();
+	UE_API virtual void OnShutdown();
 
 private:
-	static bool bIsShuttingDown;
+	static UE_API bool bIsShuttingDown;
 };
 
 
@@ -59,26 +61,28 @@ private:
  * UEditorGeometryGenerationManager is a class used by UEditorGeometryGenerationSubsystem to
  * store registrations and provide a Tick()
  */
-UCLASS()
-class GEOMETRYSCRIPTINGEDITOR_API UEditorGeometryGenerationManager : public UObject, public FTickableEditorObject
+UCLASS(MinimalAPI)
+class UEditorGeometryGenerationManager : public UObject, public FTickableEditorObject
 {
 	GENERATED_BODY()
 
 public:
-	virtual void Shutdown();
+	UE_API virtual void Shutdown();
 
-	virtual void RegisterGeneratedMeshActor(AGeneratedDynamicMeshActor* Actor);
-	virtual void UnregisterGeneratedMeshActor(AGeneratedDynamicMeshActor* Actor);
+	UE_API virtual void RegisterGeneratedMeshActor(AGeneratedDynamicMeshActor* Actor);
+	UE_API virtual void UnregisterGeneratedMeshActor(AGeneratedDynamicMeshActor* Actor);
 
 public:
 
 	//~ Begin FTickableEditorObject interface
-	virtual void Tick(float DeltaTime) override;
+	UE_API virtual void Tick(float DeltaTime) override;
 	virtual ETickableTickType GetTickableTickType() const override { return ETickableTickType::Conditional; }
-	virtual bool IsTickable() const override;
-	virtual TStatId GetStatId() const override;
+	UE_API virtual bool IsTickable() const override;
+	UE_API virtual TStatId GetStatId() const override;
 	//~ End FTickableEditorObject interface
 
 protected:
 	TSet<AGeneratedDynamicMeshActor*> ActiveGeneratedActors;
 };
+
+#undef UE_API
diff --git a/Engine/Plugins/Runtime/GeometryScripting/Source/GeometryScriptingEditor/Public/GeometryActors/GeneratedDynamicMeshActor.h b/Engine/Plugins/Runtime/GeometryScripting/Source/GeometryScriptingEditor/Public/GeometryActors/GeneratedDynamicMeshActor.h
--- a/Engine/Plugins/Runtime/GeometryScripting/Source/GeometryScriptingEditor/Public/GeometryActors/GeneratedDynamicMeshActor.h
+++ b/Engine/Plugins/Runtime/GeometryScripting/Source/GeometryScriptingEditor/Public/GeometryActors/GeneratedDynamicMeshActor.h
@@ -6,6 +6,8 @@
 
 #include "GeneratedDynamicMeshActor.generated.h"
 
+#define UE_API GEOMETRYSCRIPTINGEDITOR_API
+
 class AStaticMeshActor;
 struct FScopedSlowTask;
 
@@ -20,14 +22,14 @@ struct FScopedSlowTask;
  * and doing the procedural mesh regeneration when that function fires will generally provide
  * better in-Editor interactive performance.
  */
-UCLASS()
-class GEOMETRYSCRIPTINGEDITOR_API AGeneratedDynamicMeshActor : public ADynamicMeshActor
+UCLASS(MinimalAPI)
+class AGeneratedDynamicMeshActor : public ADynamicMeshActor
 {
 	GENERATED_UCLASS_BODY()
 
 public:
 
-	virtual ~AGeneratedDynamicMeshActor();
+	UE_API virtual ~AGeneratedDynamicMeshActor();
 
 
 public:
@@ -45,22 +47,22 @@ public:
 	 * meshes on this event, instead of doing so directly from the Construction Script.
 	 */
 	UFUNCTION(BlueprintImplementableEvent, CallInEditor, Category = "Events")
-	void OnRebuildGeneratedMesh(UDynamicMesh* TargetMesh);
+	UE_API void OnRebuildGeneratedMesh(UDynamicMesh* TargetMesh);
 
 	/**
 	 * Mark this Actor as modified so that OnRebuildGeneratedMesh runs 
 	 */
 	UFUNCTION(BlueprintCallable, Category = "DynamicMeshActor")
-	void MarkForMeshRebuild(bool bImmediate = false, bool bImmediateEvenIfFrozen = false);
+	UE_API void MarkForMeshRebuild(bool bImmediate = false, bool bImmediateEvenIfFrozen = false);
 
 
 	/**
 	 * This function will fire the RebuildGeneratedMesh function if the actor has been
 	 * marked for a pending rebuild (eg via OnConstruction)
 	 */
-	virtual void ExecuteRebuildGeneratedMeshIfPending();
+	UE_API virtual void ExecuteRebuildGeneratedMeshIfPending();
 
-	virtual bool WantsRebuild();
+	UE_API virtual bool WantsRebuild();
 
 
 protected:
@@ -68,7 +70,7 @@ protected:
 	 * Overridable native event for when the generated Mesh should be rebuilt.
 	 * This function will fire the OnRebuildGeneratedMesh event to the BP.
 	 */
-	virtual void RebuildGeneratedMesh(UDynamicMesh* TargetMesh);
+	UE_API virtual void RebuildGeneratedMesh(UDynamicMesh* TargetMesh);
 
 
 public:
@@ -79,7 +81,7 @@ public:
 	 * many configured Actor settings to be preserved (like assigned DataLayers, etc)
 	 */
 	UFUNCTION(BlueprintCallable, Category = DynamicMeshActor)
-	void CopyPropertiesToStaticMesh(AStaticMeshActor* StaticMeshActor, bool bCopyComponentMaterials = false);
+	UE_API void CopyPropertiesToStaticMesh(AStaticMeshActor* StaticMeshActor, bool bCopyComponentMaterials = false);
 
 	/**
 	 * Attempt to copy Actor Properties from a StaticMeshActor. Optionally copy DynamicMeshComponent material list to the StaticMeshComponent.
@@ -87,25 +89,25 @@ public:
 	 * many configured Actor settings to be preserved (like assigned DataLayers, etc) 
 	 */
 	UFUNCTION(BlueprintCallable, Category = DynamicMeshActor)
-	void CopyPropertiesFromStaticMesh(AStaticMeshActor* StaticMeshActor, bool bCopyComponentMaterials = false);
+	UE_API void CopyPropertiesFromStaticMesh(AStaticMeshActor* StaticMeshActor, bool bCopyComponentMaterials = false);
 
 
 
 public:
 
 	//~ Begin UObject/AActor Interface
-	virtual void PostLoad() override;
-	virtual void PostActorCreated() override;
-	virtual void OnConstruction(const FTransform& Transform) override;
-	virtual void Destroyed() override;
+	UE_API virtual void PostLoad() override;
+	UE_API virtual void PostActorCreated() override;
+	UE_API virtual void OnConstruction(const FTransform& Transform) override;
+	UE_API virtual void Destroyed() override;
 
 	// these are called when Actor exists in a sublevel that is hidden/shown
-	virtual void PreRegisterAllComponents() override;
-	virtual void PostUnregisterAllComponents() override;
+	UE_API virtual void PreRegisterAllComponents() override;
+	UE_API virtual void PostUnregisterAllComponents() override;
 
 #if WITH_EDITOR
-	virtual void PostEditUndo() override;
-	virtual void PostEditUndo(TSharedPtr<ITransactionObjectAnnotation> TransactionAnnotation) override;
+	UE_API virtual void PostEditUndo() override;
+	UE_API virtual void PostEditUndo(TSharedPtr<ITransactionObjectAnnotation> TransactionAnnotation) override;
 #endif
 
 
@@ -123,8 +125,8 @@ private:
 	bool bIsCurrentlyRebuilding = false;
 protected:
 
-	virtual void RegisterWithGenerationManager();
-	virtual void UnregisterWithGenerationManager();
+	UE_API virtual void RegisterWithGenerationManager();
+	UE_API virtual void UnregisterWithGenerationManager();
 
 
 
@@ -154,10 +156,12 @@ public:
 
 	/** Call this function from within OnRebuildGeneratedMesh to update progress tracking. */
 	UFUNCTION(BlueprintCallable, Category = "DynamicMeshActor|Progress")
-	void IncrementProgress(int NumSteps, FString Message);
+	UE_API void IncrementProgress(int NumSteps, FString Message);
 
 protected:
 	FScopedSlowTask* ActiveSlowTask = nullptr;
 	int32 CurProgressAccumSteps = 0;
 
 };
+
+#undef UE_API
diff --git a/Engine/Plugins/Runtime/GeometryScripting/Source/GeometryScriptingEditor/Public/GeometryScript/CreateNewAssetUtilityFunctions.h b/Engine/Plugins/Runtime/GeometryScripting/Source/GeometryScriptingEditor/Public/GeometryScript/CreateNewAssetUtilityFunctions.h
--- a/Engine/Plugins/Runtime/GeometryScripting/Source/GeometryScriptingEditor/Public/GeometryScript/CreateNewAssetUtilityFunctions.h
+++ b/Engine/Plugins/Runtime/GeometryScripting/Source/GeometryScriptingEditor/Public/GeometryScript/CreateNewAssetUtilityFunctions.h
@@ -9,6 +9,8 @@
 #include "BodySetupEnums.h"
 #include "CreateNewAssetUtilityFunctions.generated.h"
 
+#define UE_API GEOMETRYSCRIPTINGEDITOR_API
+
 
 class UMaterialInterface;
 class UStaticMesh;
@@ -17,7 +19,7 @@ class AVolume;
 
 
 USTRUCT(BlueprintType)
-struct GEOMETRYSCRIPTINGEDITOR_API FGeometryScriptUniqueAssetNameOptions
+struct FGeometryScriptUniqueAssetNameOptions
 {
 	GENERATED_BODY()
 public:
@@ -27,7 +29,7 @@ public:
 
 
 USTRUCT(BlueprintType)
-struct GEOMETRYSCRIPTINGEDITOR_API FGeometryScriptCreateNewVolumeFromMeshOptions
+struct FGeometryScriptCreateNewVolumeFromMeshOptions
 {
 	GENERATED_BODY()
 public:
@@ -43,7 +45,7 @@ public:
 
 
 USTRUCT(BlueprintType)
-struct GEOMETRYSCRIPTINGEDITOR_API FGeometryScriptCreateNewStaticMeshAssetOptions
+struct FGeometryScriptCreateNewStaticMeshAssetOptions
 {
 	GENERATED_BODY()
 public:
@@ -78,7 +80,7 @@ public:
 
 
 USTRUCT(BlueprintType)
-struct GEOMETRYSCRIPTINGEDITOR_API FGeometryScriptCreateNewSkeletalMeshAssetOptions
+struct FGeometryScriptCreateNewSkeletalMeshAssetOptions
 {
 	GENERATED_BODY()
 public:
@@ -111,7 +113,7 @@ public:
 
 
 USTRUCT(BlueprintType)
-struct GEOMETRYSCRIPTINGEDITOR_API FGeometryScriptCreateNewTexture2DAssetOptions
+struct FGeometryScriptCreateNewTexture2DAssetOptions
 {
 	GENERATED_BODY()
 public:
@@ -121,14 +123,14 @@ public:
 };
 
 
-UCLASS(meta = (ScriptName = "GeometryScript_NewAssetUtils"))
-class GEOMETRYSCRIPTINGEDITOR_API UGeometryScriptLibrary_CreateNewAssetFunctions : public UBlueprintFunctionLibrary
+UCLASS(MinimalAPI, meta = (ScriptName = "GeometryScript_NewAssetUtils"))
+class UGeometryScriptLibrary_CreateNewAssetFunctions : public UBlueprintFunctionLibrary
 {
 	GENERATED_BODY()
 public:
 
 	UFUNCTION(BlueprintCallable, Category = "GeometryScript|AssetManagement", meta = (HidePin = "Debug", ExpandEnumAsExecs = "Outcome"))
-	static void
+	static UE_API void
 	CreateUniqueNewAssetPathName(
 		FString AssetFolderPath,
 		FString BaseAssetName,
@@ -143,7 +145,7 @@ public:
 	 * Create a new Volume from a Dynamic Mesh, in the same world as the calling blueprint.
 	 */
 	UFUNCTION(BlueprintCallable, Category = "GeometryScript|AssetManagement", meta = (HidePin = "Debug", ExpandEnumAsExecs = "Outcome", DisplayName = "Create New Volume From Mesh", WorldContext = "WorldContextObject"))
-	static UPARAM(DisplayName = "Volume Actor") AVolume*
+	static UE_API UPARAM(DisplayName = "Volume Actor") AVolume*
 	CreateNewVolumeFromMesh_WorldContext(
 		UObject* WorldContextObject,
 		UDynamicMesh* FromDynamicMesh,
@@ -157,7 +159,7 @@ public:
 	 * Create a new Volume from a Dynamic Mesh, in the specified world.
 	 */
 	UFUNCTION(BlueprintCallable, Category = "GeometryScript|AssetManagement", meta = (HidePin = "Debug", ExpandEnumAsExecs = "Outcome", DisplayName = "Create New Volume From Mesh In Target World"))
-	static UPARAM(DisplayName = "Volume Actor") AVolume* 
+	static UE_API UPARAM(DisplayName = "Volume Actor") AVolume* 
 	CreateNewVolumeFromMesh(
 		UDynamicMesh* FromDynamicMesh, 
 		UPARAM(ref) UWorld* CreateInWorld,
@@ -171,7 +173,7 @@ public:
 	 * Create a new StaticMesh asset from a DynamicMesh. Save the asset at the AssetPathAndName location.
 	 */
 	UFUNCTION(BlueprintCallable, Category = "GeometryScript|AssetManagement", meta = (HidePin = "Debug", ExpandEnumAsExecs = "Outcome"))
-	static UPARAM(DisplayName = "Static Mesh Asset") UStaticMesh* 
+	static UE_API UPARAM(DisplayName = "Static Mesh Asset") UStaticMesh* 
 	CreateNewStaticMeshAssetFromMesh(
 		UDynamicMesh* FromDynamicMesh, 
 		FString AssetPathAndName,
@@ -184,7 +186,7 @@ public:
 	 * Save the asset at the AssetPathAndName location.
 	 */
 	UFUNCTION(BlueprintCallable, Category = "GeometryScript|AssetManagement", meta = (HidePin = "Debug", ScriptName="CreateNewStaticMeshAssetFromMeshLods;CreateNewStaticMeshAssetFromMeshLODs", ExpandEnumAsExecs = "Outcome"))
-	static UPARAM(DisplayName = "Static Mesh Asset") UStaticMesh* 
+	static UE_API UPARAM(DisplayName = "Static Mesh Asset") UStaticMesh* 
 	CreateNewStaticMeshAssetFromMeshLODs(
 		TArray<UDynamicMesh*> FromDynamicMesh, 
 		FString AssetPathAndName,
@@ -197,7 +199,7 @@ public:
 	 * Save the asset at the AssetPathAndName location.
 	 */
 	UFUNCTION(BlueprintCallable, Category = "GeometryScript|AssetManagement", meta = (HidePin = "Debug", ExpandEnumAsExecs = "Outcome"))
-	static UPARAM(DisplayName = "Skeletal Mesh Asset") USkeletalMesh* 
+	static UE_API UPARAM(DisplayName = "Skeletal Mesh Asset") USkeletalMesh* 
 	CreateNewSkeletalMeshAssetFromMesh(
 		UDynamicMesh* FromDynamicMesh, 
 		USkeleton* InSkeleton,
@@ -211,7 +213,7 @@ public:
 	 * Save the asset at the AssetPathAndName location.
 	 */
 	UFUNCTION(BlueprintCallable, Category = "GeometryScript|AssetManagement", meta = (HidePin = "Debug", ScriptName="CreateNewSkeletalMeshAssetFromMeshLods;CreateNewSkeletalMeshAssetFromMeshLODs", ExpandEnumAsExecs = "Outcome"))
-	static UPARAM(DisplayName = "Skeletal Mesh Asset") USkeletalMesh* 
+	static UE_API UPARAM(DisplayName = "Skeletal Mesh Asset") USkeletalMesh* 
 	CreateNewSkeletalMeshAssetFromMeshLODs(
 		TArray<UDynamicMesh*> FromDynamicMeshLODs, 
 		USkeleton* InSkeleton,
@@ -221,7 +223,7 @@ public:
 		UGeometryScriptDebug* Debug = nullptr);
 
 	UFUNCTION(BlueprintCallable, Category = "GeometryScript|AssetManagement", meta = (HidePin = "Debug", ExpandEnumAsExecs = "Outcome"))
-	static UPARAM(DisplayName = "Texture 2D Asset") UTexture2D* 
+	static UE_API UPARAM(DisplayName = "Texture 2D Asset") UTexture2D* 
 	CreateNewTexture2DAsset(
 		UTexture2D* FromTexture, 
 		FString AssetPathAndName,
@@ -231,3 +233,4 @@ public:
 
 };
 
+#undef UE_API
diff --git a/Engine/Plugins/Runtime/GeometryScripting/Source/GeometryScriptingEditor/Public/GeometryScript/EditorDynamicMeshUtilityFunctions.h b/Engine/Plugins/Runtime/GeometryScripting/Source/GeometryScriptingEditor/Public/GeometryScript/EditorDynamicMeshUtilityFunctions.h
--- a/Engine/Plugins/Runtime/GeometryScripting/Source/GeometryScriptingEditor/Public/GeometryScript/EditorDynamicMeshUtilityFunctions.h
+++ b/Engine/Plugins/Runtime/GeometryScripting/Source/GeometryScriptingEditor/Public/GeometryScript/EditorDynamicMeshUtilityFunctions.h
@@ -7,6 +7,8 @@
 #include "GeometryScript/GeometryScriptTypes.h"
 #include "EditorDynamicMeshUtilityFunctions.generated.h"
 
+#define UE_API GEOMETRYSCRIPTINGEDITOR_API
+
 namespace UE::Geometry
 { 
 	struct FDynamicMeshChangeContainerInternals;
@@ -20,7 +22,7 @@ namespace UE::Geometry
  * and EmitTrackedMeshChange() functions
  */
 USTRUCT(BlueprintType, meta = (DisplayName = "Dynamic Mesh ChangeTracker"))
-struct GEOMETRYSCRIPTINGEDITOR_API FDynamicMeshChangeContainer
+struct FDynamicMeshChangeContainer
 {
 	GENERATED_BODY()
 public:
@@ -38,8 +40,8 @@ public:
 
 
 
-UCLASS(meta = (ScriptName = "GeometryScript_EditorDynamicMeshUtil"))
-class GEOMETRYSCRIPTINGEDITOR_API UGeometryScriptLibrary_EditorDynamicMeshFunctions : public UBlueprintFunctionLibrary
+UCLASS(MinimalAPI, meta = (ScriptName = "GeometryScript_EditorDynamicMeshUtil"))
+class UGeometryScriptLibrary_EditorDynamicMeshFunctions : public UBlueprintFunctionLibrary
 {
 	GENERATED_BODY()
 public:
@@ -50,7 +52,7 @@ public:
 	 * @param ChangeTracker output structure containing initial TargetMesh state
 	 */
 	UFUNCTION(BlueprintCallable, Category = "DynamicMesh|Changes")
-	static UPARAM(DisplayName = "Dynamic Mesh") UDynamicMesh* 
+	static UE_API UPARAM(DisplayName = "Dynamic Mesh") UDynamicMesh* 
 	BeginTrackedMeshChange(UDynamicMesh* TargetMesh, FDynamicMeshChangeContainer& ChangeTracker);
 
 	/**
@@ -59,7 +61,7 @@ public:
 	 * called in the context of a Transaction (ie BeginTransaction / EndTransaction pair)
 	 */
 	UFUNCTION(BlueprintCallable, Category = "DynamicMesh|Changes")
-	static UPARAM(DisplayName = "Dynamic Mesh") UDynamicMesh* 
+	static UE_API UPARAM(DisplayName = "Dynamic Mesh") UDynamicMesh* 
 	EmitTrackedMeshChange(UDynamicMesh* TargetMesh, UPARAM(ref) FDynamicMeshChangeContainer& ChangeTracker);
 
 
@@ -69,7 +71,7 @@ public:
 	 * @warning This function stores the mesh in a global data structure, the caller must take care to avoid storing large numbers of debug meshes
 	 */
 	UFUNCTION(BlueprintCallable, Category = "DynamicMesh|Changes")
-	static UPARAM(DisplayName = "Dynamic Mesh") UDynamicMesh* 
+	static UE_API UPARAM(DisplayName = "Dynamic Mesh") UDynamicMesh* 
 	StashDebugMesh(UDynamicMesh* TargetMesh, FString DebugMeshName);
 
 
@@ -80,6 +82,8 @@ public:
 	 * @param bDebugMeshExists will return as true if DebugMeshName existed
 	 */
 	UFUNCTION(BlueprintCallable, Category = "DynamicMesh|Changes")
-	static UPARAM(DisplayName = "Dynamic Mesh") UDynamicMesh* 
+	static UE_API UPARAM(DisplayName = "Dynamic Mesh") UDynamicMesh* 
 	FetchDebugMesh(FString DebugMeshName, UDynamicMesh* ToTargetMesh, bool bClearDebugMesh, bool& bDebugMeshExists);
-};
\ No newline at end of file
+};
+
+#undef UE_API
diff --git a/Engine/Plugins/Runtime/GeometryScripting/Source/GeometryScriptingEditor/Public/GeometryScript/EditorTextureMapFunctions.h b/Engine/Plugins/Runtime/GeometryScripting/Source/GeometryScriptingEditor/Public/GeometryScript/EditorTextureMapFunctions.h
--- a/Engine/Plugins/Runtime/GeometryScripting/Source/GeometryScriptingEditor/Public/GeometryScript/EditorTextureMapFunctions.h
+++ b/Engine/Plugins/Runtime/GeometryScripting/Source/GeometryScriptingEditor/Public/GeometryScript/EditorTextureMapFunctions.h
@@ -7,6 +7,8 @@
 #include "GeometryScript/GeometryScriptTypes.h"
 #include "EditorTextureMapFunctions.generated.h"
 
+#define UE_API GEOMETRYSCRIPTINGEDITOR_API
+
 #if WITH_EDITOR
 
 class UTexture2D;
@@ -32,7 +34,7 @@ enum class EGeometryScriptReadGammaSpace : uint8
 };
 
 USTRUCT(BlueprintType)
-struct GEOMETRYSCRIPTINGEDITOR_API FGeometryScriptChannelPackSource
+struct FGeometryScriptChannelPackSource
 {
 	GENERATED_BODY()
 
@@ -54,7 +56,7 @@ struct GEOMETRYSCRIPTINGEDITOR_API FGeometryScriptChannelPackSource
 };
 
 USTRUCT(BlueprintType)
-struct GEOMETRYSCRIPTINGEDITOR_API FGeometryScriptChannelPackResult
+struct FGeometryScriptChannelPackResult
 {
 	GENERATED_BODY();
 
@@ -63,15 +65,15 @@ struct GEOMETRYSCRIPTINGEDITOR_API FGeometryScriptChannelPackResult
 };
 
 
-UCLASS(meta = (ScriptName = "GeometryScript_EditorTextureUtils"))
-class GEOMETRYSCRIPTINGEDITOR_API UGeometryScriptLibrary_EditorTextureMapFunctions : public UBlueprintFunctionLibrary
+UCLASS(MinimalAPI, meta = (ScriptName = "GeometryScript_EditorTextureUtils"))
+class UGeometryScriptLibrary_EditorTextureMapFunctions : public UBlueprintFunctionLibrary
 {
 	GENERATED_BODY()
 
 public:
 
 	UFUNCTION(BlueprintCallable, Category = "GeometryScript|TextureUtils", meta = (HidePin = "Debug"))
-	static FGeometryScriptChannelPackResult ChannelPack(
+	static UE_API FGeometryScriptChannelPackResult ChannelPack(
 		FGeometryScriptChannelPackSource RChannelSource,
 		FGeometryScriptChannelPackSource GChannelSource,
 		FGeometryScriptChannelPackSource BChannelSource,
@@ -81,3 +83,5 @@ public:
 };
 
 #endif // WITH_EDITOR
+
+#undef UE_API
diff --git a/Engine/Plugins/Runtime/GeometryScripting/Source/GeometryScriptingEditor/Public/GeometryScript/OpenSubdivUtilityFunctions.h b/Engine/Plugins/Runtime/GeometryScripting/Source/GeometryScriptingEditor/Public/GeometryScript/OpenSubdivUtilityFunctions.h
--- a/Engine/Plugins/Runtime/GeometryScripting/Source/GeometryScriptingEditor/Public/GeometryScript/OpenSubdivUtilityFunctions.h
+++ b/Engine/Plugins/Runtime/GeometryScripting/Source/GeometryScriptingEditor/Public/GeometryScript/OpenSubdivUtilityFunctions.h
@@ -7,15 +7,17 @@
 #include "GeometryScript/GeometryScriptTypes.h"
 #include "OpenSubdivUtilityFunctions.generated.h"
 
-UCLASS(meta = (ScriptName = "GeometryScript_OpenSubdiv"))
-class GEOMETRYSCRIPTINGEDITOR_API UGeometryScriptLibrary_OpenSubdivFunctions : public UBlueprintFunctionLibrary
+#define UE_API GEOMETRYSCRIPTINGEDITOR_API
+
+UCLASS(MinimalAPI, meta = (ScriptName = "GeometryScript_OpenSubdiv"))
+class UGeometryScriptLibrary_OpenSubdivFunctions : public UBlueprintFunctionLibrary
 {
 	GENERATED_BODY()
 public:
 
 
 	UFUNCTION(BlueprintCallable, Category = "GeometryScript|OpenSubdiv", meta=(ScriptMethod, HidePin = "Debug", DisplayName = "Apply PolyGroup Catmull Clark SubD"))
-	static UPARAM(DisplayName = "Target Mesh") UDynamicMesh* 
+	static UE_API UPARAM(DisplayName = "Target Mesh") UDynamicMesh* 
 	ApplyPolygroupCatmullClarkSubD(
 		UDynamicMesh* FromDynamicMesh, 
 		int32 Subdivisions,
@@ -24,10 +26,12 @@ public:
 
 
 	UFUNCTION(BlueprintCallable, Category = "GeometryScript|OpenSubdiv", meta=(ScriptMethod, HidePin = "Debug", DisplayName = "Apply Triangle Loop SubD"))
-	static UPARAM(DisplayName = "Target Mesh") UDynamicMesh* 
+	static UE_API UPARAM(DisplayName = "Target Mesh") UDynamicMesh* 
 	ApplyTriangleLoopSubD(
 		UDynamicMesh* FromDynamicMesh, 
 		int32 Subdivisions,
 		UGeometryScriptDebug* Debug = nullptr);
 
-};
\ No newline at end of file
+};
+
+#undef UE_API
diff --git a/Engine/Plugins/Runtime/HairStrands/Source/HairCardGeneratorFramework/Public/IHairCardGenerator.h b/Engine/Plugins/Runtime/HairStrands/Source/HairCardGeneratorFramework/Public/IHairCardGenerator.h
--- a/Engine/Plugins/Runtime/HairStrands/Source/HairCardGeneratorFramework/Public/IHairCardGenerator.h
+++ b/Engine/Plugins/Runtime/HairStrands/Source/HairCardGeneratorFramework/Public/IHairCardGenerator.h
@@ -5,14 +5,16 @@
 #include "Features/IModularFeature.h"
 #include "Features/IModularFeatures.h" // for RegisterModularFeature()/UnregisterModularFeature()
 
+#define UE_API HAIRCARDGENERATORFRAMEWORK_API
+
 class  UGroomAsset;
 struct FHairGroupsCardsSourceDescription;
 class  UHairCardGenerationSettings;
 
-class HAIRCARDGENERATORFRAMEWORK_API IHairCardGenerator : public IModularFeature
+class IHairCardGenerator : public IModularFeature
 {
 public:
-	static const FName ModularFeatureName; // "HairCardGenerator"
+	static UE_API const FName ModularFeatureName; // "HairCardGenerator"
 
 	virtual bool GenerateHairCardsForLOD(UGroomAsset* Groom, FHairGroupsCardsSourceDescription& CardsDesc) = 0;
 
@@ -30,4 +32,6 @@ namespace HairCardGenerator_Utils
 	{
 		IModularFeatures::Get().UnregisterModularFeature(IHairCardGenerator::ModularFeatureName, Generator);
 	}
-}
\ No newline at end of file
+}
+
+#undef UE_API
diff --git a/Engine/Plugins/Runtime/HairStrands/Source/HairStrandsDataflow/Public/GroomCollectionFacades.h b/Engine/Plugins/Runtime/HairStrands/Source/HairStrandsDataflow/Public/GroomCollectionFacades.h
--- a/Engine/Plugins/Runtime/HairStrands/Source/HairStrandsDataflow/Public/GroomCollectionFacades.h
+++ b/Engine/Plugins/Runtime/HairStrands/Source/HairStrandsDataflow/Public/GroomCollectionFacades.h
@@ -9,6 +9,8 @@
 
 #include "GroomCollectionFacades.generated.h"
 
+#define UE_API HAIRSTRANDSDATAFLOW_API
+
 class FGeometryCollection;
 namespace Chaos
 {
@@ -33,7 +35,7 @@ namespace UE::Groom
 	* @brief Base facade to store the TArray<T> groups necessary to setup the groom asset.
 	*/
 	template<typename DerivedType>
-	class HAIRSTRANDSDATAFLOW_API FGroomCollectionFacade
+	class FGroomCollectionFacade
 	{
 	public:
 		/** Groom collection group names */
@@ -172,15 +174,15 @@ namespace UE::Groom
 	* FGroomStrandsFacade 
 	* @brief Strands facade to store the TArray<T> groups necessary to setup the groom strands.
 	*/
-	class HAIRSTRANDSDATAFLOW_API FGroomStrandsFacade : public FGroomCollectionFacade<FGroomStrandsFacade>
+	class FGroomStrandsFacade : public FGroomCollectionFacade<FGroomStrandsFacade>
 	{
 	public:
 		/** Groom collection group names */
-		static const FName GroupPrefix;
+		static UE_API const FName GroupPrefix;
 		using FEditableType = FEditableHairStrand;
 
-		FGroomStrandsFacade(FManagedArrayCollection& InCollection);
-		FGroomStrandsFacade(const FManagedArrayCollection& InCollection);
+		UE_API FGroomStrandsFacade(FManagedArrayCollection& InCollection);
+		UE_API FGroomStrandsFacade(const FManagedArrayCollection& InCollection);
 
 		/** Create the facade attributes. */
 		void DefineFacadeSchema() {}
@@ -202,34 +204,34 @@ namespace UE::Groom
 	* FGroomGuidesFacade 
 	* @brief Guides facade to store the TArray<T> groups necessary to setup the groom guides.
 	*/
-	class HAIRSTRANDSDATAFLOW_API FGroomGuidesFacade : public FGroomCollectionFacade<FGroomGuidesFacade>
+	class FGroomGuidesFacade : public FGroomCollectionFacade<FGroomGuidesFacade>
 	{
 	public:
 		/** Groom collection group names */
-		static const FName GroupPrefix;
+		static UE_API const FName GroupPrefix;
 		using FEditableType = FEditableHairGuide;
 
-		FGroomGuidesFacade(FManagedArrayCollection& InCollection);
-		FGroomGuidesFacade(const FManagedArrayCollection& InCollection);
-
-		static const FName ObjectMeshLODsAttribute;
-		static const FName ObjectSkeletalMeshesAttribute;
-		static const FName PointKinematicWeightsAttribute;
-		static const FName PointBoneIndicesAttribute;
-		static const FName PointBoneWeightsAttribute;
-		static const FName ObjectPointSamplesAttribute;
-		static const FName CurveStrandIndicesAttribute;
-		static const FName CurveParentIndicesAttribute;
-		static const FName CurveLodIndicesAttribute;
+		UE_API FGroomGuidesFacade(FManagedArrayCollection& InCollection);
+		UE_API FGroomGuidesFacade(const FManagedArrayCollection& InCollection);
+
+		static UE_API const FName ObjectMeshLODsAttribute;
+		static UE_API const FName ObjectSkeletalMeshesAttribute;
+		static UE_API const FName PointKinematicWeightsAttribute;
+		static UE_API const FName PointBoneIndicesAttribute;
+		static UE_API const FName PointBoneWeightsAttribute;
+		static UE_API const FName ObjectPointSamplesAttribute;
+		static UE_API const FName CurveStrandIndicesAttribute;
+		static UE_API const FName CurveParentIndicesAttribute;
+		static UE_API const FName CurveLodIndicesAttribute;
 		
 		/** Create the facade attributes. */
-		void DefineFacadeSchema();
+		UE_API void DefineFacadeSchema();
 
 		/** Is the Facade defined on the collection? */
-		bool IsFacadeValid() const;
+		UE_API bool IsFacadeValid() const;
 
 		/** Init facade collection attributes */
-		void InitFacadeCollection();
+		UE_API void InitFacadeCollection();
 
 		/** Get the editable groom types */
 		static const TArray<FEditableType>& GetEditableGroom(const FEditableGroomGroup& GroomGroup)
@@ -280,7 +282,7 @@ namespace UE::Groom
 		void SetCurveLodIndices(const TArray<int32>& LodIndices) { CurveLodIndices.Modify() = LodIndices; }
 
 		/** Resize the geometry groups (vertex, points, edges) based on a new total number of points */
-		void ResizePointsGroups(const int32 NumPoints) const;
+		UE_API void ResizePointsGroups(const int32 NumPoints) const;
 
 	protected :
 		
@@ -306,3 +308,5 @@ namespace UE::Groom
 		TManagedArrayAccessor<int32> CurveLodIndices;
 	};
 }
+
+#undef UE_API
diff --git a/Engine/Plugins/Runtime/HairStrands/Source/HairStrandsDataflow/Public/GroomDataflowRendering.h b/Engine/Plugins/Runtime/HairStrands/Source/HairStrandsDataflow/Public/GroomDataflowRendering.h
--- a/Engine/Plugins/Runtime/HairStrands/Source/HairStrandsDataflow/Public/GroomDataflowRendering.h
+++ b/Engine/Plugins/Runtime/HairStrands/Source/HairStrandsDataflow/Public/GroomDataflowRendering.h
@@ -1,10 +1,12 @@
-﻿// Copyright Epic Games, Inc. All Rights Reserved.
+// Copyright Epic Games, Inc. All Rights Reserved.
 
 #pragma once
 
 #include "Dataflow/DataflowRenderingFactory.h"
 #include "Dataflow/DataflowRenderingViewMode.h"
 
+#define UE_API HAIRSTRANDSDATAFLOW_API
+
 namespace UE::Groom
 {
 
@@ -12,10 +14,10 @@ class FGroomStrandsFacade;
 class FGroomGuidesFacade;
 
 /** Strands rendering callback for the dataflow editor */
-class HAIRSTRANDSDATAFLOW_API FGroomStrandsRenderingCallbacks : public Dataflow::FRenderingFactory::ICallbackInterface
+class FGroomStrandsRenderingCallbacks : public Dataflow::FRenderingFactory::ICallbackInterface
 {
 public :
-	static Dataflow::FRenderKey RenderKey;
+	static UE_API Dataflow::FRenderKey RenderKey;
 	
 protected:
 
@@ -30,20 +32,20 @@ protected:
 		return (ViewMode.GetName() == Dataflow::FDataflowConstruction3DViewMode::Name);
 	}
 
-	virtual void Render(GeometryCollection::Facades::FRenderingFacade& RenderCollection, const Dataflow::FGraphRenderingState& State) override;
+	UE_API virtual void Render(GeometryCollection::Facades::FRenderingFacade& RenderCollection, const Dataflow::FGraphRenderingState& State) override;
 	//¬ End ICallbackInterface interface
 	
-	virtual void ComputeVertexColors(const FGroomStrandsFacade& StrandsFacade, TArray<FLinearColor>& VertexColors) const;
+	UE_API virtual void ComputeVertexColors(const FGroomStrandsFacade& StrandsFacade, TArray<FLinearColor>& VertexColors) const;
 
 	/** Get the group attribute used to create each geometry groups */
-	virtual int32 GetGroupAttribute(const FGroomStrandsFacade& StrandsFacade, FString& GroupAttribute, FString& GroupName) const;
+	UE_API virtual int32 GetGroupAttribute(const FGroomStrandsFacade& StrandsFacade, FString& GroupAttribute, FString& GroupName) const;
 };
 
 /** Guides rendering callback for the dataflow editor */
-class HAIRSTRANDSDATAFLOW_API FGroomGuidesRenderingCallbacks : public Dataflow::FRenderingFactory::ICallbackInterface
+class FGroomGuidesRenderingCallbacks : public Dataflow::FRenderingFactory::ICallbackInterface
 {
 	public :
-		static Dataflow::FRenderKey RenderKey;
+		static UE_API Dataflow::FRenderKey RenderKey;
 	
 protected:
 
@@ -58,12 +60,12 @@ protected:
 		return (ViewMode.GetName() == Dataflow::FDataflowConstruction3DViewMode::Name);
 	}
 
-	virtual void Render(GeometryCollection::Facades::FRenderingFacade& RenderCollection, const Dataflow::FGraphRenderingState& State) override;
+	UE_API virtual void Render(GeometryCollection::Facades::FRenderingFacade& RenderCollection, const Dataflow::FGraphRenderingState& State) override;
 	
-	virtual void ComputeVertexColors(const FGroomGuidesFacade& GuidesFacade, TArray<FLinearColor>& VertexColors) const;
+	UE_API virtual void ComputeVertexColors(const FGroomGuidesFacade& GuidesFacade, TArray<FLinearColor>& VertexColors) const;
 	
 	/** Get the group attribute used to create each geometry groups */
-	virtual int32 GetGroupAttribute(const FGroomGuidesFacade& StrandsFacade, FString& GroupAttribute, FString& GroupName) const;
+	UE_API virtual int32 GetGroupAttribute(const FGroomGuidesFacade& StrandsFacade, FString& GroupAttribute, FString& GroupName) const;
     	
 	//¬ End ICallbackInterface interface
 };
@@ -75,3 +77,5 @@ void RegisterRenderingCallbacks();
 void DeregisterRenderingCallbacks();
 	
 }  // End namespace UE::Groom
+
+#undef UE_API
diff --git a/Engine/Plugins/Runtime/HairStrands/Source/HairStrandsEditor/Public/GroomCacheImporter.h b/Engine/Plugins/Runtime/HairStrands/Source/HairStrandsEditor/Public/GroomCacheImporter.h
--- a/Engine/Plugins/Runtime/HairStrands/Source/HairStrandsEditor/Public/GroomCacheImporter.h
+++ b/Engine/Plugins/Runtime/HairStrands/Source/HairStrandsEditor/Public/GroomCacheImporter.h
@@ -4,6 +4,8 @@
 
 #include "CoreMinimal.h"
 
+#define UE_API HAIRSTRANDSEDITOR_API
+
 struct FGroomAnimationInfo;
 struct FHairImportContext;
 class FGroomCacheProcessor;
@@ -13,10 +15,12 @@ class UGroomCache;
 
 enum class EGroomCacheImportType : uint8;
 
-struct HAIRSTRANDSEDITOR_API FGroomCacheImporter
+struct FGroomCacheImporter
 {
-	static TArray<UGroomCache*> ImportGroomCache(const FString& SourceFilename, TSharedPtr<IGroomTranslator> Translator, const FGroomAnimationInfo& InAnimInfo, FHairImportContext& HairImportContext, UGroomAsset* GroomAssetForCache, EGroomCacheImportType ImportType);
-	static void SetupImportSettings(struct FGroomCacheImportSettings& ImportSettings, const FGroomAnimationInfo& AnimInfo);
-	static void ApplyImportSettings(struct FGroomCacheImportSettings& ImportSettings, FGroomAnimationInfo& AnimInfo);
-	static UGroomCache* ProcessToGroomCache(FGroomCacheProcessor& Processor, const FGroomAnimationInfo& AnimInfo, FHairImportContext& ImportContext, const FString& ObjectNameSuffix);
+	static UE_API TArray<UGroomCache*> ImportGroomCache(const FString& SourceFilename, TSharedPtr<IGroomTranslator> Translator, const FGroomAnimationInfo& InAnimInfo, FHairImportContext& HairImportContext, UGroomAsset* GroomAssetForCache, EGroomCacheImportType ImportType);
+	static UE_API void SetupImportSettings(struct FGroomCacheImportSettings& ImportSettings, const FGroomAnimationInfo& AnimInfo);
+	static UE_API void ApplyImportSettings(struct FGroomCacheImportSettings& ImportSettings, FGroomAnimationInfo& AnimInfo);
+	static UE_API UGroomCache* ProcessToGroomCache(FGroomCacheProcessor& Processor, const FGroomAnimationInfo& AnimInfo, FHairImportContext& ImportContext, const FString& ObjectNameSuffix);
 };
+
+#undef UE_API
diff --git a/Engine/Plugins/Runtime/HairStrands/Source/HairStrandsEditor/Public/GroomCustomAssetEditorToolkit.h b/Engine/Plugins/Runtime/HairStrands/Source/HairStrandsEditor/Public/GroomCustomAssetEditorToolkit.h
--- a/Engine/Plugins/Runtime/HairStrands/Source/HairStrandsEditor/Public/GroomCustomAssetEditorToolkit.h
+++ b/Engine/Plugins/Runtime/HairStrands/Source/HairStrandsEditor/Public/GroomCustomAssetEditorToolkit.h
@@ -11,6 +11,8 @@
 #include "HairStrandsInterface.h"
 #include "GroomEditorStyle.h"
 
+#define UE_API HAIRSTRANDSEDITOR_API
+
 class FAssetThumbnailPool;
 class UGroomAsset; 
 class UGroomComponent;
@@ -25,7 +27,7 @@ class UGroomBindingAsset;
 class UGroomBindingAssetList;
 class UAnimationAsset;
 
-class HAIRSTRANDSEDITOR_API IGroomCustomAssetEditorToolkit : public FAssetEditorToolkit
+class IGroomCustomAssetEditorToolkit : public FAssetEditorToolkit
 {
 public:
 	/** Retrieves the current custom asset. */
@@ -42,12 +44,12 @@ public:
 	virtual FGroomEditorStyle* GetSlateStyle() const { return nullptr; };
 };
 
-class HAIRSTRANDSEDITOR_API FGroomCustomAssetEditorToolkit : public IGroomCustomAssetEditorToolkit
+class FGroomCustomAssetEditorToolkit : public IGroomCustomAssetEditorToolkit
 {
 public:
 
-	virtual void RegisterTabSpawners(const TSharedRef<class FTabManager>& InTabManager) override;
-	virtual void UnregisterTabSpawners(const TSharedRef<class FTabManager>& InTabManager) override;
+	UE_API virtual void RegisterTabSpawners(const TSharedRef<class FTabManager>& InTabManager) override;
+	UE_API virtual void UnregisterTabSpawners(const TSharedRef<class FTabManager>& InTabManager) override;
 
 	/**
 	 * Edits the specified asset object
@@ -56,99 +58,99 @@ public:
 	 * @param	InitToolkitHost			When Mode is WorldCentric, this is the level editor instance to spawn this editor within
 	 * @param	InCustomAsset			The Custom Asset to Edit
 	 */
-	void InitCustomAssetEditor(const EToolkitMode::Type Mode, const TSharedPtr<class IToolkitHost>& InitToolkitHost, UGroomAsset* InCustomAsset);
+	UE_API void InitCustomAssetEditor(const EToolkitMode::Type Mode, const TSharedPtr<class IToolkitHost>& InitToolkitHost, UGroomAsset* InCustomAsset);
 
-	FGroomCustomAssetEditorToolkit();
+	UE_API FGroomCustomAssetEditorToolkit();
 
 	/** Destructor */
-	virtual ~FGroomCustomAssetEditorToolkit();
+	UE_API virtual ~FGroomCustomAssetEditorToolkit();
 
 	/** Begin IToolkit interface */
-	virtual FName GetToolkitFName() const override;
-	virtual FText GetBaseToolkitName() const override;
-	virtual FText GetToolkitName() const override;
-	virtual FText GetToolkitToolTipText() const override;
-	virtual FString GetWorldCentricTabPrefix() const override;
-	virtual FLinearColor GetWorldCentricTabColorScale() const override;
+	UE_API virtual FName GetToolkitFName() const override;
+	UE_API virtual FText GetBaseToolkitName() const override;
+	UE_API virtual FText GetToolkitName() const override;
+	UE_API virtual FText GetToolkitToolTipText() const override;
+	UE_API virtual FString GetWorldCentricTabPrefix() const override;
+	UE_API virtual FLinearColor GetWorldCentricTabColorScale() const override;
 	virtual bool IsPrimaryEditor() const override { return true; }
-	virtual void OnClose() override;
+	UE_API virtual void OnClose() override;
 	/** End IToolkit interface */
 
 	/** Retrieves the current custom asset. */
-	virtual UGroomAsset* GetCustomAsset() const override;
+	UE_API virtual UGroomAsset* GetCustomAsset() const override;
 
 	/** Set the current custom asset. */
-	virtual void SetCustomAsset(UGroomAsset* InCustomAsset) override;	
+	UE_API virtual void SetCustomAsset(UGroomAsset* InCustomAsset) override;	
 
 	/** Set preview of a particular binding. */
-	virtual void PreviewBinding(int32 BindingIndex) override;
+	UE_API virtual void PreviewBinding(int32 BindingIndex) override;
 
 	/** Return the index of the active binding. */
-	virtual int32 GetActiveBindingIndex() const override;
+	UE_API virtual int32 GetActiveBindingIndex() const override;
 
 	/** Return groom editor style. */
-	FGroomEditorStyle* GetSlateStyle() const override;
+	UE_API FGroomEditorStyle* GetSlateStyle() const override;
 
 private:
 
 	// called when The Play simulation button is pressed
-	void OnPlaySimulation();
-	bool CanPlaySimulation() const;
+	UE_API void OnPlaySimulation();
+	UE_API bool CanPlaySimulation() const;
 	
 	// Called when the pause simulation button is pressed
-	void OnPauseSimulation();
-	bool CanPauseSimulation() const;
+	UE_API void OnPauseSimulation();
+	UE_API bool CanPauseSimulation() const;
 
 	// Called when the reset simulation button is pressed
-	void OnResetSimulation();
-	bool CanResetSimulation() const;
+	UE_API void OnResetSimulation();
+	UE_API bool CanResetSimulation() const;
 
 	// Called when the play animation button is pressed
-	void OnPlayAnimation();
-	bool CanPlayAnimation() const;
+	UE_API void OnPlayAnimation();
+	UE_API bool CanPlayAnimation() const;
 
 	// Called when the stop animation button is pressed
-	void OnStopAnimation();
-	bool CanStopAnimation() const;
+	UE_API void OnStopAnimation();
+	UE_API bool CanStopAnimation() const;
 
 	// Add buttons to to the toolbar
-	void ExtendToolbar();
+	UE_API void ExtendToolbar();
 
 	// THis should be called when the properties of the Document are changed
-	void DocPropChanged(UObject *, FPropertyChangedEvent &);
+	UE_API void DocPropChanged(UObject *, FPropertyChangedEvent &);
 
 	// THis is called when the groom target object changes and needs updating
-	void OnSkeletalGroomTargetChanged(USkeletalMesh *NewTarget);
+	UE_API void OnSkeletalGroomTargetChanged(USkeletalMesh *NewTarget);
 
 	// Return true if the animation asset should be filtered out (not compatible with the preview skel. mesh)
-	bool OnShouldFilterAnimAsset(const FAssetData& AssetData);
-	void OnObjectChangedAnimAsset(const FAssetData& AssetData);
-	bool OnIsEnabledAnimAsset();
-	FString GetCurrentAnimAssetPath() const;
+	UE_API bool OnShouldFilterAnimAsset(const FAssetData& AssetData);
+	UE_API void OnObjectChangedAnimAsset(const FAssetData& AssetData);
+	UE_API bool OnIsEnabledAnimAsset();
+	UE_API FString GetCurrentAnimAssetPath() const;
 
 	// create the custom components we need
-	void InitPreviewComponents();
+	UE_API void InitPreviewComponents();
 
 	// Initialized the content of the binding tab
-	void InitializeBindingAssetTabContent();
+	UE_API void InitializeBindingAssetTabContent();
 
 	// return a pointer to the groom preview component
-	UGroomComponent*		GetPreview_GroomComponent() const;
-	USkeletalMeshComponent*	GetPreview_SkeletalMeshComponent() const;
-
-	TSharedRef<SDockTab> SpawnViewportTab(const FSpawnTabArgs& Args);
-
-	TSharedRef<SDockTab> SpawnTab_LODProperties(const FSpawnTabArgs& Args);
-	TSharedRef<SDockTab> SpawnTab_InterpolationProperties(const FSpawnTabArgs& Args);
-	TSharedRef<SDockTab> SpawnTab_RenderingProperties(const FSpawnTabArgs& Args);
-	TSharedRef<SDockTab> SpawnTab_CardsProperties(const FSpawnTabArgs& Args);
-	TSharedRef<SDockTab> SpawnTab_MeshesProperties(const FSpawnTabArgs& Args);
-	TSharedRef<SDockTab> SpawnTab_MaterialProperties(const FSpawnTabArgs& Args);
-	TSharedRef<SDockTab> SpawnTab_PhysicsProperties(const FSpawnTabArgs& Args);
-	TSharedRef<SDockTab> SpawnTab_PreviewGroomComponent(const FSpawnTabArgs& Args);
-	TSharedRef<SDockTab> SpawnTab_BindingProperties(const FSpawnTabArgs& Args);
-	TSharedRef<SDockTab> SpawnTab_PreviewSettings(const FSpawnTabArgs& Args);
-	TSharedRef<SDockTab> SpawnTab_DataflowProperties(const FSpawnTabArgs& Args);
+	UE_API UGroomComponent*		GetPreview_GroomComponent() const;
+	UE_API USkeletalMeshComponent*	GetPreview_SkeletalMeshComponent() const;
+
+	UE_API TSharedRef<SDockTab> SpawnViewportTab(const FSpawnTabArgs& Args);
+
+	UE_API TSharedRef<SDockTab> SpawnTab_LODProperties(const FSpawnTabArgs& Args);
+	UE_API TSharedRef<SDockTab> SpawnTab_InterpolationProperties(const FSpawnTabArgs& Args);
+	UE_API TSharedRef<SDockTab> SpawnTab_RenderingProperties(const FSpawnTabArgs& Args);
+	UE_API TSharedRef<SDockTab> SpawnTab_CardsProperties(const FSpawnTabArgs& Args);
+	UE_API TSharedRef<SDockTab> SpawnTab_MeshesProperties(const FSpawnTabArgs& Args);
+	UE_API TSharedRef<SDockTab> SpawnTab_MaterialProperties(const FSpawnTabArgs& Args);
+	UE_API TSharedRef<SDockTab> SpawnTab_PhysicsProperties(const FSpawnTabArgs& Args);
+	UE_API TSharedRef<SDockTab> SpawnTab_PreviewGroomComponent(const FSpawnTabArgs& Args);
+	UE_API TSharedRef<SDockTab> SpawnTab_BindingProperties(const FSpawnTabArgs& Args);
+	UE_API TSharedRef<SDockTab> SpawnTab_PreviewSettings(const FSpawnTabArgs& Args);
+	UE_API TSharedRef<SDockTab> SpawnTab_DataflowProperties(const FSpawnTabArgs& Args);
 	
 private:
 
@@ -168,20 +170,20 @@ private:
 	TSharedPtr<class IDetailsView> DetailView_PreviewSceneProperties;
 	TSharedPtr<class IDetailsView> DetailView_DataflowProperties;
 
-	static const FName ToolkitFName;
-	static const FName TabId_Viewport;
+	static UE_API const FName ToolkitFName;
+	static UE_API const FName TabId_Viewport;
 
-	static const FName TabId_LODProperties;
-	static const FName TabId_InterpolationProperties;
-	static const FName TabId_RenderingProperties;
-	static const FName TabId_CardsProperties;
-	static const FName TabId_MeshesProperties;
-	static const FName TabId_MaterialProperties;
-	static const FName TabId_PhysicsProperties;
-	static const FName TabId_PreviewGroomComponent;
-	static const FName TabId_BindingProperties;
-	static const FName TabId_PreviewSceneProperties;
-	static const FName TabId_DataflowProperties;
+	static UE_API const FName TabId_LODProperties;
+	static UE_API const FName TabId_InterpolationProperties;
+	static UE_API const FName TabId_RenderingProperties;
+	static UE_API const FName TabId_CardsProperties;
+	static UE_API const FName TabId_MeshesProperties;
+	static UE_API const FName TabId_MaterialProperties;
+	static UE_API const FName TabId_PhysicsProperties;
+	static UE_API const FName TabId_PreviewGroomComponent;
+	static UE_API const FName TabId_BindingProperties;
+	static UE_API const FName TabId_PreviewSceneProperties;
+	static UE_API const FName TabId_DataflowProperties;
 
 
 	bool bIsTabManagerInitialized = false;
@@ -201,3 +203,4 @@ private:
 	TSharedPtr<FGroomEditorStyle> GroomEditorStyle;
 };
 
+#undef UE_API
