Add support for custom/additionnal UV channels for hair cards/hair meshes.

diff --git a/Engine/Plugins/Experimental/HairModelingToolset/Source/HairModelingToolset/Private/GroomQueryUtil.cpp b/Engine/Plugins/Experimental/HairModelingToolset/Source/HairModelingToolset/Private/GroomQueryUtil.cpp
--- a/Engine/Plugins/Experimental/HairModelingToolset/Source/HairModelingToolset/Private/GroomQueryUtil.cpp
+++ b/Engine/Plugins/Experimental/HairModelingToolset/Source/HairModelingToolset/Private/GroomQueryUtil.cpp
@@ -70,6 +70,7 @@ void UE::GroomQueries::ExtractAllHairCards(AGroomActor* GroomActor,
 		}
 		const FHairCardsGeometry& CardGeo = CardData.Cards;
 
+		const uint32 CardUVIndex = FHairCardsGeometry::FUVChannel::GetCardUVIndex();
 		int32 NumVerts = CardGeo.GetNumVertices();
 		TArray<int32> MapV, MapN, MapU;
 		MapV.SetNum(NumVerts);
@@ -83,7 +84,7 @@ void UE::GroomQueries::ExtractAllHairCards(AGroomActor* GroomActor,
 			int32 NewNormID = Normals->AppendElement((FVector3f)CardGeo.Normals[k]);
 			MapN[k] = NewNormID;
 
-			int32 NewUVID = UVs->AppendElement( FVector2f(CardGeo.UVs[k].X, CardGeo.UVs[k].Y) );
+			int32 NewUVID = UVs->AppendElement( FVector2f(CardGeo.UVs[CardUVIndex].UV[k].X, CardGeo.UVs[CardUVIndex].UV[k].Y) );
 			MapU[k] = NewUVID;
 		}
 
diff --git a/Engine/Plugins/Runtime/HairStrands/Source/HairStrandsCore/Private/HairCardsBuilder.cpp b/Engine/Plugins/Runtime/HairStrands/Source/HairStrandsCore/Private/HairCardsBuilder.cpp
--- a/Engine/Plugins/Runtime/HairStrands/Source/HairStrandsCore/Private/HairCardsBuilder.cpp
+++ b/Engine/Plugins/Runtime/HairStrands/Source/HairStrandsCore/Private/HairCardsBuilder.cpp
@@ -66,7 +66,48 @@ namespace FHairCardsBuilder
 FString GetVersion()
 {
 	// Important to update the version when cards building or importing changes
-	return TEXT("12");
+	return TEXT("13i");
+}
+
+#define HAIR_CARD_UV_ROOTUV 0
+#define HAIR_CARD_UV_CARDUV 1
+
+// Example of layout with 3 texture coords
+//  -----------------------------------------------  -----------------------------------------------
+// |          Vertex0          ||          Vertex1          ||          Vertex2          ||
+// | 0.xy   0.zw   1.xy   1.zw || 0.xy   0.zw   1.xy   1.zw || 0.xy   0.zw   1.xy   1.zw ||
+// | UV0  | UV1  | UV2  | ---  || UV0  | UV1  | UV2  | ---  || UV0  | UV1  | UV2  | ---  ||
+//  -----------------------------------------------  -----------------------------------------------
+static void SetPackedUV(uint32 VertexIndex, uint32 UVIndex, uint32 UVCount, const FVector2f& InUV, TArray<FVector4f>& OutPackedUVs)
+{
+	const uint32 StrideInUint = FMath::DivideAndRoundUp(UVCount, 2u);
+	uint32 UVVertexIndex = VertexIndex * StrideInUint + (UVIndex / 2);
+	if (UVIndex % 2 == 0)
+	{
+		OutPackedUVs[UVVertexIndex].X = InUV.X; 
+		OutPackedUVs[UVVertexIndex].Y = InUV.Y; 
+	}
+	else
+	{
+		OutPackedUVs[UVVertexIndex].Z = InUV.X; 
+		OutPackedUVs[UVVertexIndex].W = InUV.Y; 
+	}
+}
+static FVector2f GetPackedUV(uint32 VertexIndex, uint32 UVIndex, uint32 UVCount, const TArray<FVector4f>& InPackedUVs)
+{
+	const uint32 StrideInUint = FMath::DivideAndRoundUp(UVCount, 2u);
+	const uint32 UVVertexIndex = VertexIndex * StrideInUint + (UVIndex / 2);
+	const FVector4f PackedUV = InPackedUVs[UVVertexIndex];
+	return (UVIndex % 2 == 0) ? FVector2f(PackedUV.X, PackedUV.Y) : FVector2f(PackedUV.Z, PackedUV.W);
+}
+
+static uint32 GetNumPackedUVs(uint32 PointCount, uint32 NumUVs)
+{
+	static_assert(FHairCardsGeometry::FUVChannel::GetCardUVIndex() == HAIR_CARD_UV_CARDUV);
+
+	// Packed 2 UVs into a single FVector4f
+	const uint32 StrideInUint = FMath::DivideAndRoundUp(NumUVs, 2u);
+	return PointCount * StrideInUint;
 }
 
 static bool InternalCreateCardsGuides(
@@ -94,7 +135,6 @@ static bool InternalCreateCardsGuides(
 	{
 		uint32 MainDirectionUCount = 0;
 		uint32 ValidCount = 0;
-		FVector4f TriangleUVs[3];
 		for (uint32 CardIt = 0; CardIt < NumCards; ++CardIt)
 		{
 			const uint32 NumTriangles = InCards.IndexCounts[CardIt] / 3;
@@ -142,7 +182,7 @@ static bool InternalCreateCardsGuides(
 				for (uint32 VertexIt = 0; VertexIt < 3; ++VertexIt)
 				{
 					const uint32 VertexIndex = InCards.Indices[VertexIndexOffset + VertexIt];
-					const FVector4f UV = TriangleUVs[VertexIt] = InCards.UVs[VertexIndex];
+					const FVector2f UV = InCards.UVs[HAIR_CARD_UV_CARDUV].UV[VertexIndex];
 					AddSimilarUV(SimilarVertexU, VertexIndex, UV.X, UV.Y, UVCoordTreshold);
 					AddSimilarUV(SimilarVertexV, VertexIndex, UV.Y, UV.X, UVCoordTreshold);
 				}
@@ -163,7 +203,6 @@ static bool InternalCreateCardsGuides(
 		bIsMainDirectionU = (float(MainDirectionUCount) / float(FMath::Max(1u,ValidCount))) > 0.5f;
 	}
 
-	FVector4f TriangleUVs[3];
 	for (uint32 CardIt = 0; CardIt < NumCards; ++CardIt)
 	{
 		const uint32 NumTriangles = InCards.IndexCounts[CardIt] / 3;
@@ -211,7 +250,7 @@ static bool InternalCreateCardsGuides(
 			for (uint32 VertexIt = 0; VertexIt < 3; ++VertexIt)
 			{
 				const uint32 VertexIndex = InCards.Indices[VertexIndexOffset + VertexIt];
-				const FVector4f UV = TriangleUVs[VertexIt] = InCards.UVs[VertexIndex];
+				const FVector2f UV = InCards.UVs[HAIR_CARD_UV_CARDUV].UV[VertexIndex];
 				AddSimilarUV(SimilarVertexU, VertexIndex, UV.X, UV.Y, UVCoordTreshold);
 				AddSimilarUV(SimilarVertexV, VertexIndex, UV.Y, UV.X, UVCoordTreshold);
 			}
@@ -534,12 +573,20 @@ static bool InternalImportGeometry_WithGeneratedGuides(
 		}
 	}
 
+	const uint32 MeshDescriptionNumChannels = MeshDescription->GetNumUVElementChannels();
+	const uint32 NumUVSource = MeshDescriptionNumChannels;
+	const uint32 NumUVs = NumUVSource + 1 /*RootUV*/;
+
 	// Fill vertex data
 	Out.Cards.Positions.SetNum(VertexCount);
 	Out.Cards.Normals.SetNum(VertexCount);
 	Out.Cards.Tangents.SetNum(VertexCount);
-	Out.Cards.UVs.SetNum(VertexCount);
 	Out.Cards.VertexColors.Empty();
+	Out.Cards.UVs.SetNum(NumUVs);
+	for (uint32 UVIt = 0;UVIt<NumUVs; ++UVIt)
+	{
+		Out.Cards.UVs[UVIt].UV.Init(FVector2f::ZeroVector, VertexCount);
+	}
 	TArray<float> TangentFrameSigns;
 	TangentFrameSigns.SetNum(VertexCount);
 
@@ -571,7 +618,6 @@ static bool InternalImportGeometry_WithGeneratedGuides(
 		const uint32 VertexIndex = VertexId.GetValue();
 		check(VertexIndex < VertexCount);
 		Out.Cards.Positions[VertexIndex]	= VertexPositions[VertexId];
-		Out.Cards.UVs[VertexIndex]			= FVector4f(VertexInstanceUVs[VertexInstanceId0].Component(0), VertexInstanceUVs[VertexInstanceId0].Component(1), 0, 0); // RootUV are not set here, but will be 'patched' later once guides & interpolation data are built
 		Out.Cards.Tangents[VertexIndex]		= VertexInstanceTangents[VertexInstanceId0];
 		Out.Cards.Normals[VertexIndex]		= VertexInstanceNormals[VertexInstanceId0];
 		TangentFrameSigns[VertexIndex] 		= VertexInstanceBinormalSigns[VertexInstanceId0];
@@ -579,17 +625,21 @@ static bool InternalImportGeometry_WithGeneratedGuides(
 		{
 			Out.Cards.VertexColors[VertexIndex] = FLinearColor(VertexColors[VertexInstanceId0]);
 		}
-
+		for (uint32 UVIt = 0; UVIt < NumUVSource; ++UVIt)
+		{
+			Out.Cards.UVs[HAIR_CARD_UV_CARDUV + UVIt].UV[VertexIndex] = VertexInstanceUVs.Get(VertexInstanceId0, UVIt);
+		}
 		Out.Cards.BoundingBox += Out.Cards.Positions[VertexIndex];
 	}
 	OutBulk.Header.BoundingBox = ToFBox3d(Out.Cards.BoundingBox);
 	OutBulk.Header.bVertexColor = bHasVertexColor;
+	OutBulk.Header.NumUVs = NumUVs;
 
 	// Fill in render resources (do we need to keep it separated? e.g, format compression, packing)
 	const uint32 PointCount = Out.Cards.Positions.Num();
 	OutBulk.Positions.SetNum(PointCount);
 	OutBulk.Normals.SetNum(PointCount * FHairCardsNormalFormat::ComponentCount);
-	OutBulk.UVs.SetNum(PointCount);
+	OutBulk.UVs.SetNum(GetNumPackedUVs(PointCount, NumUVs));
 	OutBulk.Materials.SetNum(PointCount);
 	OutBulk.VertexColors.Empty();
 	if (bHasVertexColor)
@@ -599,7 +649,6 @@ static bool InternalImportGeometry_WithGeneratedGuides(
 	for (uint32 PointIt = 0; PointIt < PointCount; ++PointIt)
 	{
 		OutBulk.Positions[PointIt] = FVector4f(Out.Cards.Positions[PointIt], 0);
-		OutBulk.UVs[PointIt] = FVector4f(Out.Cards.UVs[PointIt]);
 		OutBulk.Normals[PointIt * 2] = FVector4f(Out.Cards.Tangents[PointIt], 0);
 		OutBulk.Normals[PointIt * 2 + 1] = FVector4f(Out.Cards.Normals[PointIt], TangentFrameSigns[PointIt]);
 		OutBulk.Materials[PointIt] = 0;
@@ -874,13 +923,19 @@ static bool InternalImportGeometry_WithGeneratedGuides(
 				// Root0  Root1
 				const float TexCoordV = Out.Cards.LocalUVs[VertexIndex].Y;
 				const FVector2f RootUV = FMath::Lerp(RootUVData.RootUV0, RootUVData.RootUV1, TexCoordV);
-
-				Out.Cards.UVs[VertexIndex].Z = RootUV.X;
-				Out.Cards.UVs[VertexIndex].W = RootUV.Y;
-				OutBulk.UVs[VertexIndex].Z   = RootUV.X;
-				OutBulk.UVs[VertexIndex].W   = RootUV.Y;
+				Out.Cards.UVs[HAIR_CARD_UV_ROOTUV].UV[VertexIndex] = RootUV;
 			}
 		}
+
+	}
+
+	// Pack texture coords
+	for (uint32 PointIt = 0; PointIt < PointCount; ++PointIt)
+	{
+		for (uint32 UVIt = 0; UVIt < NumUVs; ++UVIt)
+		{
+			SetPackedUV(PointIt, UVIt, NumUVs, Out.Cards.UVs[UVIt].UV[PointIt], OutBulk.UVs);
+		}
 	}
 
 	// Build the FHairStrandsData from the FHairStrandsRawData
@@ -919,12 +974,14 @@ static bool InternalImportGeometry_WithImportedGuides(
 	const TVertexInstanceAttributesRef<const FVector2f> VertexInstanceUVs		= MeshDescription->VertexInstanceAttributes().GetAttributesRef<FVector2f>(MeshAttribute::VertexInstance::TextureCoordinate);
 	const TVertexInstanceAttributesRef<const FVector4f> VertexColors			= MeshDescription->VertexInstanceAttributes().GetAttributesRef<FVector4f>(MeshAttribute::VertexInstance::Color);
 
-	const bool bHasVertexColor = VertexColors.GetNumElements() > 0;
+	const bool bHasVertexColor = VertexColors.GetNumElements() > 0;	
+	const uint32 NumSourceUVs = StaticMesh->GetNumTexCoords(0);
+	const uint32 NumUVs = NumSourceUVs + 1u /*RootUV*/;
 
 	// Write out BulkData
 	OutBulk.Positions.SetNum(PointCount);
 	OutBulk.Normals.SetNum(PointCount * FHairCardsNormalFormat::ComponentCount);
-	OutBulk.UVs.SetNum(PointCount);
+	OutBulk.UVs.SetNum(GetNumPackedUVs(PointCount, NumUVs));
 	OutBulk.Materials.SetNum(PointCount);
 	OutBulk.VertexColors.Empty();
 	if (bHasVertexColor)
@@ -944,8 +1001,7 @@ static bool InternalImportGeometry_WithImportedGuides(
 
 		const uint32 VertexIndex = VertexId.GetValue();
 		check(VertexIndex < PointCount);
-		OutBulk.Positions[VertexIndex] 		= FVector4f(VertexPositions[VertexId], 0);
-		OutBulk.UVs[VertexIndex] 			= FVector4f(VertexInstanceUVs[VertexInstanceId0].Component(0), VertexInstanceUVs[VertexInstanceId0].Component(1), 0, 0); // RootUV are not set here, but will be 'patched' later once guides & interpolation data are built
+		OutBulk.Positions[VertexIndex] 		= FVector4f(VertexPositions[VertexId], 0);		
 		OutBulk.Normals[VertexIndex * 2]	= FVector4f(VertexInstanceTangents[VertexInstanceId0], 0);
 		OutBulk.Normals[VertexIndex * 2 + 1]= FVector4f(VertexInstanceNormals[VertexInstanceId0], VertexInstanceBinormalSigns[VertexInstanceId0] /*TangentFrameSigns*/);
 		OutBulk.Materials[VertexIndex] 		= 0;
@@ -953,11 +1009,23 @@ static bool InternalImportGeometry_WithImportedGuides(
 		{
 			OutBulk.VertexColors[VertexIndex]= FLinearColor(VertexColors[VertexInstanceId0]).ToFColor(true/*sRGB*/).Bits;
 		}
+		// UVs are packed 2 channels into a single float4
+		// * UVs[0]   -> Root UV - Not set here, but will be 'patched' later once guides & interpolation data are built
+		// * UVs[1-N] -> Card UV & else
+		for (uint32 UVIndex=0;UVIndex<NumSourceUVs;++UVIndex)
+		{
+			FVector2f UV = FVector2f::ZeroVector;
+			UV.X = VertexInstanceUVs[VertexInstanceId0].Component(UVIndex*2 + 0);
+			UV.Y = VertexInstanceUVs[VertexInstanceId0].Component(UVIndex*2 + 1);
+
+			SetPackedUV(VertexIndex, HAIR_CARD_UV_CARDUV + UVIndex, NumUVs, UV, OutBulk.UVs);
+		}
 
 		BoundingBox += OutBulk.Positions[VertexIndex];
 	}
 	OutBulk.Header.BoundingBox = ToFBox3d(BoundingBox);
 	OutBulk.Header.bVertexColor = bHasVertexColor;
+	OutBulk.Header.NumUVs = NumUVs;
 
 	TArray<FVector3f> CardBaseColor;
 	TArray<float> CardsRoughness;
@@ -1076,8 +1144,7 @@ static bool InternalImportGeometry_WithImportedGuides(
 			// Root UV bulk data
 			{
 				const FVector2f RootUV = InGuides.StrandsCurves.CurvesRootUV[ClosestGuide.GuideIndex];
-				OutBulk.UVs[PointIt].Z = RootUV.X;
-				OutBulk.UVs[PointIt].W = RootUV.Y;
+				SetPackedUV(PointIt, HAIR_CARD_UV_ROOTUV, NumUVs, RootUV, OutBulk.UVs);
 			}
 		}
 	}
@@ -1176,7 +1243,7 @@ namespace FHairMeshesBuilder
 FString GetVersion()
 {
 	// Important to update the version when meshes building or importing changes
-	return TEXT("4");
+	return TEXT("4g");
 }
 
 void BuildGeometry(
@@ -1195,8 +1262,9 @@ void BuildGeometry(
 	Out.Meshes.Positions.SetNum(TotalPointCount);
 	Out.Meshes.Normals.SetNum(TotalPointCount);
 	Out.Meshes.Tangents.SetNum(TotalPointCount);
-	Out.Meshes.UVs.SetNum(TotalPointCount);
 	Out.Meshes.Indices.SetNum(TotalIndexCount);
+	Out.Meshes.UVs.SetNum(1);
+	Out.Meshes.UVs[0].UV.SetNum(TotalPointCount);
 
 	Out.Meshes.Positions[0] = Center + FVector3f(-Extent.X, -Extent.Y, -Extent.Z);
 	Out.Meshes.Positions[1] = Center + FVector3f(+Extent.X, -Extent.Y, -Extent.Z);
@@ -1207,14 +1275,14 @@ void BuildGeometry(
 	Out.Meshes.Positions[6] = Center + FVector3f(+Extent.X, +Extent.Y, +Extent.Z);
 	Out.Meshes.Positions[7] = Center + FVector3f(-Extent.X, +Extent.Y, +Extent.Z);
 
-	Out.Meshes.UVs[0] = FVector2f(0, 0);
-	Out.Meshes.UVs[1] = FVector2f(1, 0);
-	Out.Meshes.UVs[2] = FVector2f(1, 1);
-	Out.Meshes.UVs[3] = FVector2f(0, 1);
-	Out.Meshes.UVs[4] = FVector2f(0, 0);
-	Out.Meshes.UVs[5] = FVector2f(1, 0);
-	Out.Meshes.UVs[6] = FVector2f(1, 1);
-	Out.Meshes.UVs[7] = FVector2f(0, 1);
+	Out.Meshes.UVs[HAIR_CARD_UV_CARDUV].UV[0] = FVector2f(0, 0);
+	Out.Meshes.UVs[HAIR_CARD_UV_CARDUV].UV[1] = FVector2f(1, 0);
+	Out.Meshes.UVs[HAIR_CARD_UV_CARDUV].UV[2] = FVector2f(1, 1);
+	Out.Meshes.UVs[HAIR_CARD_UV_CARDUV].UV[3] = FVector2f(0, 1);
+	Out.Meshes.UVs[HAIR_CARD_UV_CARDUV].UV[4] = FVector2f(0, 0);
+	Out.Meshes.UVs[HAIR_CARD_UV_CARDUV].UV[5] = FVector2f(1, 0);
+	Out.Meshes.UVs[HAIR_CARD_UV_CARDUV].UV[6] = FVector2f(1, 1);
+	Out.Meshes.UVs[HAIR_CARD_UV_CARDUV].UV[7] = FVector2f(0, 1);
 
 	Out.Meshes.Normals[0] = FVector3f(0, 0, 1);
 	Out.Meshes.Normals[1] = FVector3f(0, 0, 1);
@@ -1279,23 +1347,26 @@ void BuildGeometry(
 	Out.Meshes.BoundingBox.Init();
 
 	// Fill in render resources (do we need to keep it separated? e.g, format compression, packing)
+	const uint32 NumUVs = 2; // RootUV+CardUV
 	const uint32 PointCount = Out.Meshes.Positions.Num();
 	OutBulk.Positions.SetNum(PointCount);
 	OutBulk.Normals.SetNum(PointCount * FHairCardsNormalFormat::ComponentCount);
-	OutBulk.UVs.SetNum(PointCount);
+	OutBulk.UVs.SetNum(FHairCardsBuilder::GetNumPackedUVs(PointCount, NumUVs));
 	OutBulk.VertexColors.Empty();
 	for (uint32 PointIt = 0; PointIt < PointCount; ++PointIt)
 	{
 		OutBulk.Positions[PointIt] = FVector4f(Out.Meshes.Positions[PointIt], 0);
-		OutBulk.UVs[PointIt] = FVector4f(Out.Meshes.UVs[PointIt].X, Out.Meshes.UVs[PointIt].Y, 0, 0);
 		OutBulk.Normals[PointIt * 2] = FVector4f(Out.Meshes.Tangents[PointIt], 0);
 		OutBulk.Normals[PointIt * 2 + 1] = FVector4f(Out.Meshes.Normals[PointIt], 1);
+		FHairCardsBuilder::SetPackedUV(PointIt, HAIR_CARD_UV_CARDUV, NumUVs, Out.Meshes.UVs[HAIR_CARD_UV_CARDUV].UV[PointIt], OutBulk.UVs);
+		FHairCardsBuilder::SetPackedUV(PointIt, HAIR_CARD_UV_ROOTUV, NumUVs, FVector2f::ZeroVector, OutBulk.UVs);
 
 		Out.Meshes.BoundingBox += FVector4f(OutBulk.Positions[PointIt]);
 	}
 	OutBulk.Header.BoundingBox = ToFBox3d(Out.Meshes.BoundingBox);
 	OutBulk.Header.bVertexColor = false;
-	
+	OutBulk.Header.NumUVs = NumUVs;
+
 	const uint32 IndexCount = Out.Meshes.Indices.Num();
 	OutBulk.Indices.SetNum(IndexCount);
 	for (uint32 IndexIt = 0; IndexIt < IndexCount; ++IndexIt)
@@ -1320,13 +1391,20 @@ void ImportGeometry(
 	const uint32 VertexCount = LODData.VertexBuffers.PositionVertexBuffer.GetNumVertices();
 	const uint32 IndexCount = LODData.IndexBuffer.GetNumIndices();
 	const bool bHasVertexColor = LODData.VertexBuffers.ColorVertexBuffer.GetNumVertices() > 0;
-
+	const uint32 NumUVsSource = LODData.VertexBuffers.StaticMeshVertexBuffer.GetNumTexCoords();
+	const uint32 NumUVs = NumUVsSource + 1; // +1 for RootUV
+	
 	Out.Meshes.Positions.SetNum(VertexCount);
 	Out.Meshes.Normals.SetNum(VertexCount);
 	Out.Meshes.Tangents.SetNum(VertexCount);
-	Out.Meshes.UVs.SetNum(VertexCount);
 	Out.Meshes.VertexColors.Empty();
 	Out.Meshes.Indices.SetNum(IndexCount);
+	Out.Meshes.UVs.SetNum(NumUVs);
+	for (uint32 UVIt = 0; UVIt < NumUVs; ++UVIt)
+	{
+		Out.Meshes.UVs[UVIt].UV.Init(FVector2f::ZeroVector, VertexCount);
+	}
+
 	if (bHasVertexColor)
 	{
 		Out.Meshes.VertexColors.SetNum(VertexCount);
@@ -1337,10 +1415,15 @@ void ImportGeometry(
 	for (uint32 VertexIt = 0; VertexIt < VertexCount; ++VertexIt)
 	{
 		Out.Meshes.Positions[VertexIt]	= LODData.VertexBuffers.PositionVertexBuffer.VertexPosition(VertexIt);
-		Out.Meshes.UVs[VertexIt]		= LODData.VertexBuffers.StaticMeshVertexBuffer.GetVertexUV(VertexIt, 0);
 		Out.Meshes.Tangents[VertexIt]	= FVector4f(LODData.VertexBuffers.StaticMeshVertexBuffer.VertexTangentX(VertexIt));
 		Out.Meshes.Normals[VertexIt]	= FVector4f(LODData.VertexBuffers.StaticMeshVertexBuffer.VertexTangentZ(VertexIt));
 
+		for (uint32 UVIt=0; UVIt<NumUVsSource; ++UVIt)
+		{
+			const FVector2f UV = LODData.VertexBuffers.StaticMeshVertexBuffer.GetVertexUV(VertexIt, UVIt);
+			Out.Meshes.UVs[HAIR_CARD_UV_CARDUV + UVIt].UV[VertexIt] = UV;
+		}
+
 		Out.Meshes.BoundingBox += Out.Meshes.Positions[VertexIt];
 	}
 
@@ -1352,11 +1435,12 @@ void ImportGeometry(
 	// Fill in render resources (do we need to keep it separated? e.g, format compression, packing)
 	OutBulk.Header.BoundingBox = ToFBox3d(Out.Meshes.BoundingBox);
 	OutBulk.Header.bVertexColor = bHasVertexColor;
+	OutBulk.Header.NumUVs = NumUVs;
 
 	const uint32 PointCount = Out.Meshes.Positions.Num();
 	OutBulk.Positions.SetNum(PointCount);
 	OutBulk.Normals.SetNum(PointCount * FHairCardsNormalFormat::ComponentCount);
-	OutBulk.UVs.SetNum(PointCount);
+	OutBulk.UVs.SetNum(FHairCardsBuilder::GetNumPackedUVs(PointCount, NumUVs));
 	OutBulk.VertexColors.Empty();
 	if (bHasVertexColor)
 	{
@@ -1365,13 +1449,17 @@ void ImportGeometry(
 	for (uint32 PointIt = 0; PointIt < PointCount; ++PointIt)
 	{
 		OutBulk.Positions[PointIt] = FVector4f(Out.Meshes.Positions[PointIt], 0);
-		OutBulk.UVs[PointIt] = FVector4f(Out.Meshes.UVs[PointIt].X, Out.Meshes.UVs[PointIt].Y, 0, 0);
 		OutBulk.Normals[PointIt * 2] = FVector4f(Out.Meshes.Tangents[PointIt], 0);
 		OutBulk.Normals[PointIt * 2 + 1] = FVector4f(Out.Meshes.Normals[PointIt], 1);
 		if (bHasVertexColor)
 		{
 			OutBulk.VertexColors[PointIt] = Out.Meshes.VertexColors[PointIt].Bits;
 		}
+		for (uint32 UVIt=0; UVIt<NumUVs; ++UVIt)
+		{
+			const FVector2f CardUVs = Out.Meshes.UVs[UVIt].UV[PointIt];
+			FHairCardsBuilder::SetPackedUV(PointIt, UVIt, NumUVs, CardUVs, OutBulk.UVs);
+		}
 	}
 
 	OutBulk.Indices.SetNum(IndexCount);
@@ -1533,10 +1621,13 @@ void ConvertCardsGeometryToMeshDescription(const FHairCardsGeometry& In, FMeshDe
 {
 	Out.Empty();
 
+	// Do not write out RootUV
+	const uint32 NumUVs = FMath::Max(1, In.UVs.Num()-1);
+
 	FMeshDescriptionBuilder Builder;
 	Builder.SetMeshDescription(&Out);
 	Builder.EnablePolyGroups();
-	Builder.SetNumUVLayers(1);
+	Builder.SetNumUVLayers(NumUVs);
 
 	// create vertices
 	TArray<FVertexInstanceID> VertexInstanceIDs;
@@ -1546,9 +1637,11 @@ void ConvertCardsGeometryToMeshDescription(const FHairCardsGeometry& In, FMeshDe
 	{
 		FVertexID VertexID = Builder.AppendVertex(In.Positions[VIndex]);
 		FVertexInstanceID InstanceID = Builder.AppendInstance(VertexID);
-		FVector2f UV = FVector2f(In.UVs[VIndex].X, In.UVs[VIndex].Y);
 		Builder.SetInstanceNormal(InstanceID, In.Normals[VIndex]);
-		Builder.SetInstanceUV(InstanceID, UV, 0);
+		for (uint32 UVIt = 0; UVIt < NumUVs; ++UVIt)
+		{
+			Builder.SetInstanceUV(InstanceID, In.UVs[HAIR_CARD_UV_CARDUV + UVIt].UV[VIndex], UVIt);
+		}
 
 		VertexInstanceIDs.Add(InstanceID);
 	}
diff --git a/Engine/Plugins/Runtime/HairStrands/Source/HairStrandsCore/Private/HairCardsDatas.cpp b/Engine/Plugins/Runtime/HairStrands/Source/HairStrandsCore/Private/HairCardsDatas.cpp
--- a/Engine/Plugins/Runtime/HairStrands/Source/HairStrandsCore/Private/HairCardsDatas.cpp
+++ b/Engine/Plugins/Runtime/HairStrands/Source/HairStrandsCore/Private/HairCardsDatas.cpp
@@ -16,6 +16,7 @@ void FHairCardsBulkData::Serialize(FArchive& Ar)
 	Ar << VertexColors;
 	Ar << Indices;
 	Ar << Header.BoundingBox;
+	Ar << Header.NumUVs;
 	Ar << Header.bVertexColor;
 }
 
@@ -45,6 +46,7 @@ void FHairMeshesBulkData::Serialize(FArchive& Ar)
 	Ar << VertexColors;
 	Ar << Indices;
 	Ar << Header.BoundingBox;
+	Ar << Header.NumUVs;
 	Ar << Header.bVertexColor;
 }
 
diff --git a/Engine/Plugins/Runtime/HairStrands/Source/HairStrandsCore/Private/HairCardsVertexFactory.cpp b/Engine/Plugins/Runtime/HairStrands/Source/HairStrandsCore/Private/HairCardsVertexFactory.cpp
--- a/Engine/Plugins/Runtime/HairStrands/Source/HairStrandsCore/Private/HairCardsVertexFactory.cpp
+++ b/Engine/Plugins/Runtime/HairStrands/Source/HairStrandsCore/Private/HairCardsVertexFactory.cpp
@@ -119,6 +119,7 @@ FHairCardsUniformBuffer CreateHairCardsVFUniformBuffer(
 		UniformParameters.AttributeChannelIndex = 0; // TODO pack attribute texture/channel description. Currently use hardcoded layout in shader attribute code.
 		UniformParameters.MaxVertexCount = LOD.RestResource->GetVertexCount();
 		UniformParameters.CoverageBias = FMath::Clamp(GHairCardCoverageBias, -16.f, 16.f);
+		UniformParameters.NumUVs = LOD.RestResource->BulkData.Header.NumUVs;
 
 		// When the geometry is not-dynamic (no binding to skeletal mesh, no simulation), only a single vertex buffer is allocated. 
 		// In this case we force the buffer index to 0
@@ -174,6 +175,7 @@ FHairCardsUniformBuffer CreateHairCardsVFUniformBuffer(
 		UniformParameters.AttributeChannelIndex = 0; // TODO pack attribute texture/channel description. Currently use hardcoded layout in shader attribute code.
 		UniformParameters.MaxVertexCount = LOD.RestResource->GetVertexCount();
 		UniformParameters.CoverageBias = 0;
+		UniformParameters.NumUVs = LOD.RestResource->BulkData.Header.NumUVs;
 
 		UniformParameters.NormalsBuffer = LOD.RestResource->NormalsBuffer.ShaderResourceViewRHI.GetReference();
 		UniformParameters.UVsBuffer = LOD.RestResource->UVsBuffer.ShaderResourceViewRHI.GetReference();
@@ -328,17 +330,17 @@ void FHairCardsVertexFactory::GetPSOPrecacheVertexFetchElements(EVertexInputStre
 	{
 		uint8 NumStreams = Elements.Num();
 		
-		Elements.Add(FVertexElement(0+NumStreams, 0, FHairCardsPositionFormat::VertexElementType, 0, FHairCardsPositionFormat::SizeInByte, false));
-		Elements.Add(FVertexElement(0+NumStreams, 0, FHairCardsPositionFormat::VertexElementType, 5, FHairCardsPositionFormat::SizeInByte, false));
+		Elements.Add(FVertexElement(0+NumStreams, 0, FHairCardsPositionFormat::VertexElementType, 0, FHairCardsPositionFormat::SizeInByte, false)); // Current Position
+		Elements.Add(FVertexElement(0+NumStreams, 0, FHairCardsPositionFormat::VertexElementType, 1, FHairCardsPositionFormat::SizeInByte, false)); // Previous Position
 
 		uint8 NormalOffset0 = 0u;
 		uint8 NormalOffset1 = FHairCardsNormalFormat::SizeInByte;
-		Elements.Add(FVertexElement(1+NumStreams, NormalOffset0, FHairCardsNormalFormat::VertexElementType, 1, FHairCardsNormalFormat::SizeInByte * FHairCardsNormalFormat::ComponentCount, false));
-		Elements.Add(FVertexElement(1+NumStreams, NormalOffset1, FHairCardsNormalFormat::VertexElementType, 2, FHairCardsNormalFormat::SizeInByte * FHairCardsNormalFormat::ComponentCount, false));
+		Elements.Add(FVertexElement(1+NumStreams, NormalOffset0, FHairCardsNormalFormat::VertexElementType, 2, FHairCardsNormalFormat::SizeInByte * FHairCardsNormalFormat::ComponentCount, false)); // Normal
+		Elements.Add(FVertexElement(1+NumStreams, NormalOffset1, FHairCardsNormalFormat::VertexElementType, 3, FHairCardsNormalFormat::SizeInByte * FHairCardsNormalFormat::ComponentCount, false)); // Normal
 
-		Elements.Add(FVertexElement(2+NumStreams, 0, FHairCardsUVFormat::VertexElementType, 3, FHairCardsUVFormat::SizeInByte, false));
-		Elements.Add(FVertexElement(3+NumStreams, 0, FHairCardsMaterialFormat::VertexElementType, 4, FHairCardsMaterialFormat::SizeInByte, false)); // Material
-		Elements.Add(FVertexElement(3+NumStreams, 0, FHairCardsMaterialFormat::VertexElementType, 5, FHairCardsMaterialFormat::SizeInByte, false)); // VertexColor
+		Elements.Add(FVertexElement(2+NumStreams, 0, FHairCardsMaterialFormat::VertexElementType, 4, FHairCardsMaterialFormat::SizeInByte, false)); // Material
+		Elements.Add(FVertexElement(2+NumStreams, 0, FHairCardsMaterialFormat::VertexElementType, 5, FHairCardsMaterialFormat::SizeInByte, false)); // VertexColor
+		Elements.Add(FVertexElement(2+NumStreams, 0, FHairCardsUVFormat::VertexElementType, 6, FHairCardsUVFormat::SizeInByte, false));
 	}
 }
 
@@ -397,11 +399,26 @@ void FHairCardsVertexFactory::InitResources(FRHICommandListBase& RHICmdList)
 	const bool bManualFetch = SupportsManualVertexFetch(CurrentFeatureLevel);
 	if (!bManualFetch)
 	{
+		auto AddUVChannels = [this, &Elements](const FHairCardsVertexBuffer* UVsBuffer, uint32 NumUVs, uint32 StartStreamIndex)
+		{
+			// Add additional TexCoord
+			const uint32 UVSize = sizeof(FVector2f);
+			const EVertexElementType UVVertexElementType = EVertexElementType::VET_Float4;
+
+			// UV[0] is rootUV which is not part of the regular UVs, hence UV starts at index 1
+			for (uint32 UVIt = 1; UVIt < NumUVs; UVIt +=2)
+			{
+				Elements.Add(this->AccessStreamComponent(FVertexStreamComponent(UVsBuffer, 0, UVIt * UVSize, UVSize * NumUVs, UVVertexElementType, EVertexStreamUsage::Default), StartStreamIndex++));
+			}
+		};
+
 		if (Data.GeometryType == EHairGeometryType::Cards)
 		{
 			const FHairGroupInstance::FCards::FLOD& LOD = Data.Instance->Cards.LODs[Data.LODIndex];
+			const uint32 NumUVs = LOD.RestResource->BulkData.GetNumUVs();
 			const bool bHasVertexColor = LOD.RestResource->BulkData.Header.bVertexColor;
 			const bool bDynamic = LOD.DeformedResource != nullptr;
+			const FHairCardsVertexBuffer* VertexColorsBuffer = bHasVertexColor ? &LOD.RestResource->VertexColorsBuffer : &LOD.RestResource->NormalsBuffer;
 			if (bDynamic)
 			{
 				DeformedPositionVertexBuffer[0] = FRDGWrapperVertexBuffer(LOD.DeformedResource->GetBuffer(FHairCardsDeformedResource::EFrameType::Current));
@@ -414,29 +431,25 @@ void FHairCardsVertexFactory::InitResources(FRHICommandListBase& RHICmdList)
 				DeformedNormalVertexBuffer.InitResource(RHICmdList);
 
 				Elements.Add(AccessStreamComponent(FVertexStreamComponent(&(DeformedPositionVertexBuffer[0]),		0, 0,									FHairCardsPositionFormat::SizeInByte,												FHairCardsPositionFormat::VertexElementType,	EVertexStreamUsage::Default), 0));
-				Elements.Add(AccessStreamComponent(FVertexStreamComponent(&(DeformedPositionVertexBuffer[1]),		0, 0,									FHairCardsPositionFormat::SizeInByte,												FHairCardsPositionFormat::VertexElementType,	EVertexStreamUsage::Default), 5));
+				Elements.Add(AccessStreamComponent(FVertexStreamComponent(&(DeformedPositionVertexBuffer[1]),		0, 0,									FHairCardsPositionFormat::SizeInByte,												FHairCardsPositionFormat::VertexElementType,	EVertexStreamUsage::Default), 1));
 
-				Elements.Add(AccessStreamComponent(FVertexStreamComponent(&DeformedNormalVertexBuffer,				0, 0,									FHairCardsNormalFormat::SizeInByte * FHairCardsNormalFormat::ComponentCount,		FHairCardsNormalFormat::VertexElementType,		EVertexStreamUsage::Default), 1));
-				Elements.Add(AccessStreamComponent(FVertexStreamComponent(&DeformedNormalVertexBuffer,				0, FHairCardsNormalFormat::SizeInByte,	FHairCardsNormalFormat::SizeInByte * FHairCardsNormalFormat::ComponentCount,		FHairCardsNormalFormat::VertexElementType,		EVertexStreamUsage::Default), 2));
+				Elements.Add(AccessStreamComponent(FVertexStreamComponent(&DeformedNormalVertexBuffer,				0, 0,									FHairCardsNormalFormat::SizeInByte * FHairCardsNormalFormat::ComponentCount,		FHairCardsNormalFormat::VertexElementType,		EVertexStreamUsage::Default), 2));
+				Elements.Add(AccessStreamComponent(FVertexStreamComponent(&DeformedNormalVertexBuffer,				0, FHairCardsNormalFormat::SizeInByte,	FHairCardsNormalFormat::SizeInByte * FHairCardsNormalFormat::ComponentCount,		FHairCardsNormalFormat::VertexElementType,		EVertexStreamUsage::Default), 3));
 			}
 			else
 			{
 				Elements.Add(AccessStreamComponent(FVertexStreamComponent(&LOD.RestResource->RestPositionBuffer,	0, 0,									FHairCardsPositionFormat::SizeInByte,												FHairCardsPositionFormat::VertexElementType,	EVertexStreamUsage::Default), 0));
-				Elements.Add(AccessStreamComponent(FVertexStreamComponent(&LOD.RestResource->RestPositionBuffer,	0, 0,									FHairCardsPositionFormat::SizeInByte,												FHairCardsPositionFormat::VertexElementType,	EVertexStreamUsage::Default), 5));
+				Elements.Add(AccessStreamComponent(FVertexStreamComponent(&LOD.RestResource->RestPositionBuffer,	0, 0,									FHairCardsPositionFormat::SizeInByte,												FHairCardsPositionFormat::VertexElementType,	EVertexStreamUsage::Default), 1));
 
-				Elements.Add(AccessStreamComponent(FVertexStreamComponent(&LOD.RestResource->NormalsBuffer,			0, 0,									FHairCardsNormalFormat::SizeInByte * FHairCardsNormalFormat::ComponentCount,		FHairCardsNormalFormat::VertexElementType,		EVertexStreamUsage::Default), 1));
-				Elements.Add(AccessStreamComponent(FVertexStreamComponent(&LOD.RestResource->NormalsBuffer,			0, FHairCardsNormalFormat::SizeInByte,	FHairCardsNormalFormat::SizeInByte * FHairCardsNormalFormat::ComponentCount,		FHairCardsNormalFormat::VertexElementType,		EVertexStreamUsage::Default), 2));
+				Elements.Add(AccessStreamComponent(FVertexStreamComponent(&LOD.RestResource->NormalsBuffer,			0, 0,									FHairCardsNormalFormat::SizeInByte * FHairCardsNormalFormat::ComponentCount,		FHairCardsNormalFormat::VertexElementType,		EVertexStreamUsage::Default), 2));
+				Elements.Add(AccessStreamComponent(FVertexStreamComponent(&LOD.RestResource->NormalsBuffer,			0, FHairCardsNormalFormat::SizeInByte,	FHairCardsNormalFormat::SizeInByte * FHairCardsNormalFormat::ComponentCount,		FHairCardsNormalFormat::VertexElementType,		EVertexStreamUsage::Default), 3));
 			}
-			Elements.Add(AccessStreamComponent(FVertexStreamComponent(&LOD.RestResource->UVsBuffer,					0, 0,									FHairCardsUVFormat::SizeInByte,														FHairCardsUVFormat::VertexElementType,			EVertexStreamUsage::Default), 3));
 			Elements.Add(AccessStreamComponent(FVertexStreamComponent(&LOD.RestResource->MaterialsBuffer,			0, 0,									FHairCardsMaterialFormat::SizeInByte,												FHairCardsMaterialFormat::VertexElementType,	EVertexStreamUsage::Default), 4));
-			if (bHasVertexColor)
-			{
-				Elements.Add(AccessStreamComponent(FVertexStreamComponent(&LOD.RestResource->VertexColorsBuffer,	0, 0,									FHairCardsMaterialFormat::SizeInByte,												FHairCardsMaterialFormat::VertexElementType,	EVertexStreamUsage::Default), 6));
-			}
-			else
-			{
-				Elements.Add(AccessStreamComponent(FVertexStreamComponent(&LOD.RestResource->NormalsBuffer,			0, 0,									FHairCardsMaterialFormat::SizeInByte,												FHairCardsMaterialFormat::VertexElementType,	EVertexStreamUsage::Default), 6));
-			}
+			Elements.Add(AccessStreamComponent(FVertexStreamComponent(VertexColorsBuffer,							0, 0,									FHairCardsMaterialFormat::SizeInByte,												FHairCardsMaterialFormat::VertexElementType,	EVertexStreamUsage::Default), 5));
+			Elements.Add(AccessStreamComponent(FVertexStreamComponent(&LOD.RestResource->UVsBuffer,					0, 0,									FHairCardsUVFormat::SizeInByte,														FHairCardsUVFormat::VertexElementType,			EVertexStreamUsage::Default), 6));
+
+			// Add additional UVs
+			AddUVChannels(&LOD.RestResource->UVsBuffer, LOD.RestResource->BulkData.GetNumUVs(), 7);
 
 			// Ensure the rest resources are in correct states for the VF
 			const FHairCardsRestResource* RestResource = LOD.RestResource;
@@ -457,6 +470,7 @@ void FHairCardsVertexFactory::InitResources(FRHICommandListBase& RHICmdList)
 		{
 			const FHairGroupInstance::FMeshes::FLOD& LOD = Data.Instance->Meshes.LODs[Data.LODIndex];
 			const bool bHasVertexColor = LOD.RestResource->BulkData.Header.bVertexColor;
+			const FHairCardsVertexBuffer* VertexColorsBuffer = bHasVertexColor ? &LOD.RestResource->VertexColorsBuffer : &LOD.RestResource->NormalsBuffer;
 
 			// Note: Use the 'Normal' buffer as a dummy input for 'Material' buffer, as Material data is fetched through textures for meshes
 			const bool bDynamic = LOD.DeformedResource != nullptr;
@@ -469,25 +483,22 @@ void FHairCardsVertexFactory::InitResources(FRHICommandListBase& RHICmdList)
 				DeformedPositionVertexBuffer[1].InitResource(RHICmdList);
 
 				Elements.Add(AccessStreamComponent(FVertexStreamComponent(&(DeformedPositionVertexBuffer[0]),		0, 0,									FHairCardsPositionFormat::SizeInByte,												FHairCardsPositionFormat::VertexElementType,	EVertexStreamUsage::Default), 0));
-				Elements.Add(AccessStreamComponent(FVertexStreamComponent(&(DeformedPositionVertexBuffer[1]),		0, 0,									FHairCardsPositionFormat::SizeInByte,												FHairCardsPositionFormat::VertexElementType,	EVertexStreamUsage::Default), 5));
+				Elements.Add(AccessStreamComponent(FVertexStreamComponent(&(DeformedPositionVertexBuffer[1]),		0, 0,									FHairCardsPositionFormat::SizeInByte,												FHairCardsPositionFormat::VertexElementType,	EVertexStreamUsage::Default), 1));
 			}
 			else
 			{
 				Elements.Add(AccessStreamComponent(FVertexStreamComponent(&LOD.RestResource->RestPositionBuffer,	0, 0,									FHairCardsPositionFormat::SizeInByte,												FHairCardsPositionFormat::VertexElementType,	EVertexStreamUsage::Default), 0));
-				Elements.Add(AccessStreamComponent(FVertexStreamComponent(&LOD.RestResource->RestPositionBuffer,	0, 0,									FHairCardsPositionFormat::SizeInByte,												FHairCardsPositionFormat::VertexElementType,	EVertexStreamUsage::Default), 5));
+				Elements.Add(AccessStreamComponent(FVertexStreamComponent(&LOD.RestResource->RestPositionBuffer,	0, 0,									FHairCardsPositionFormat::SizeInByte,												FHairCardsPositionFormat::VertexElementType,	EVertexStreamUsage::Default), 1));
 			}
-			Elements.Add(AccessStreamComponent(FVertexStreamComponent(&LOD.RestResource->NormalsBuffer,				0, 0,									FHairCardsNormalFormat::SizeInByte * FHairCardsNormalFormat::ComponentCount,		FHairCardsNormalFormat::VertexElementType,		EVertexStreamUsage::Default), 1));
-			Elements.Add(AccessStreamComponent(FVertexStreamComponent(&LOD.RestResource->NormalsBuffer,				0, FHairCardsNormalFormat::SizeInByte,	FHairCardsNormalFormat::SizeInByte * FHairCardsNormalFormat::ComponentCount,		FHairCardsNormalFormat::VertexElementType,		EVertexStreamUsage::Default), 2));
-			Elements.Add(AccessStreamComponent(FVertexStreamComponent(&LOD.RestResource->UVsBuffer,					0, 0,									FHairCardsUVFormat::SizeInByte,														FHairCardsUVFormat::VertexElementType,			EVertexStreamUsage::Default), 3));
+			Elements.Add(AccessStreamComponent(FVertexStreamComponent(&LOD.RestResource->NormalsBuffer,				0, 0,									FHairCardsNormalFormat::SizeInByte * FHairCardsNormalFormat::ComponentCount,		FHairCardsNormalFormat::VertexElementType,		EVertexStreamUsage::Default), 2));
+			Elements.Add(AccessStreamComponent(FVertexStreamComponent(&LOD.RestResource->NormalsBuffer,				0, FHairCardsNormalFormat::SizeInByte,	FHairCardsNormalFormat::SizeInByte * FHairCardsNormalFormat::ComponentCount,		FHairCardsNormalFormat::VertexElementType,		EVertexStreamUsage::Default), 3));
+
 			Elements.Add(AccessStreamComponent(FVertexStreamComponent(&LOD.RestResource->NormalsBuffer,				0, 0,									FHairCardsMaterialFormat::SizeInByte,												FHairCardsMaterialFormat::VertexElementType,	EVertexStreamUsage::Default), 4)); // Dummy
-			if (bHasVertexColor)
-			{
-				Elements.Add(AccessStreamComponent(FVertexStreamComponent(&LOD.RestResource->VertexColorsBuffer,	0, 0,									FHairCardsMaterialFormat::SizeInByte,												FHairCardsMaterialFormat::VertexElementType,	EVertexStreamUsage::Default), 6));
-			}
-			else
-			{
-				Elements.Add(AccessStreamComponent(FVertexStreamComponent(&LOD.RestResource->NormalsBuffer,			0, 0,									FHairCardsMaterialFormat::SizeInByte,												FHairCardsMaterialFormat::VertexElementType,	EVertexStreamUsage::Default), 6));
-			}
+			Elements.Add(AccessStreamComponent(FVertexStreamComponent(&LOD.RestResource->VertexColorsBuffer,		0, 0,									FHairCardsMaterialFormat::SizeInByte,												FHairCardsMaterialFormat::VertexElementType,	EVertexStreamUsage::Default), 5));
+			Elements.Add(AccessStreamComponent(FVertexStreamComponent(&LOD.RestResource->UVsBuffer,					0, 0,									FHairCardsUVFormat::SizeInByte,														FHairCardsUVFormat::VertexElementType,			EVertexStreamUsage::Default), 6));
+
+			// Add additional UVs
+			AddUVChannels(&LOD.RestResource->UVsBuffer, LOD.RestResource->BulkData.GetNumUVs(), 7);
 
 			// Ensure the rest resources are in correct states for the VF
 			const FHairMeshesRestResource* RestResource = LOD.RestResource;
diff --git a/Engine/Plugins/Runtime/HairStrands/Source/HairStrandsCore/Public/GroomInstance.h b/Engine/Plugins/Runtime/HairStrands/Source/HairStrandsCore/Public/GroomInstance.h
--- a/Engine/Plugins/Runtime/HairStrands/Source/HairStrandsCore/Public/GroomInstance.h
+++ b/Engine/Plugins/Runtime/HairStrands/Source/HairStrandsCore/Public/GroomInstance.h
@@ -31,6 +31,7 @@ BEGIN_GLOBAL_SHADER_PARAMETER_STRUCT(FHairCardsVertexFactoryUniformShaderParamet
 	SHADER_PARAMETER(uint32, AttributeTextureIndex)
 	SHADER_PARAMETER(uint32, AttributeChannelIndex)
 	SHADER_PARAMETER(uint32, MaxVertexCount)
+	SHADER_PARAMETER(uint32, NumUVs)
 	SHADER_PARAMETER(float, CoverageBias)
 	SHADER_PARAMETER_SRV(Buffer<float4>, PositionBuffer)
 	SHADER_PARAMETER_SRV(Buffer<float4>, PreviousPositionBuffer)
diff --git a/Engine/Plugins/Runtime/HairStrands/Source/HairStrandsCore/Public/HairCardsDatas.h b/Engine/Plugins/Runtime/HairStrands/Source/HairStrandsCore/Public/HairCardsDatas.h
--- a/Engine/Plugins/Runtime/HairStrands/Source/HairStrandsCore/Public/HairCardsDatas.h
+++ b/Engine/Plugins/Runtime/HairStrands/Source/HairStrandsCore/Public/HairCardsDatas.h
@@ -84,8 +84,14 @@ struct FHairCardsInterpolationFormat
 // Data structure holding cards geometry information
 struct FHairCardsGeometry
 {
+	struct FUVChannel
+	{
+		TArray<FVector2f> UV;
+		static constexpr uint32 GetCardUVIndex() { return 1u; }
+	};
+
 	// Geometry
-	TArray<FVector4f>UVs;
+	TArray<FUVChannel>UVs;
 	TArray<FVector3f>Normals;
 	TArray<FVector3f>Tangents;
 	TArray<FVector3f>Positions;
@@ -182,6 +188,7 @@ struct FHairCardsBulkData
 	uint32 GetNumTriangles() const { return Indices.Num() / 3; }
 	uint32 GetNumVertices() const { return Positions.Num(); }
 	const FBox& GetBounds() const { return Header.BoundingBox; }
+	uint32 GetNumUVs() const { return Header.NumUVs; }
 
 	UTexture2D* DepthTexture = nullptr;
 	UTexture2D* TangentTexture = nullptr;
@@ -198,6 +205,7 @@ struct FHairCardsBulkData
 	struct FHeader
 	{
 		FBox BoundingBox = FBox(EForceInit::ForceInit);
+		uint32 NumUVs = 0;
 		bool bVertexColor = false;
 	} Header;
 
@@ -241,8 +249,13 @@ struct FHairCardsInterpolationBulkData
 
 struct FHairMeshes
 {
+	struct FUVChannel
+	{
+		TArray<FVector2f> UV;
+	};
+
 	// Geometry
-	TArray<FVector2f> UVs;
+	TArray<FUVChannel>UVs;
 	TArray<FColor>    VertexColors;
 	TArray<FVector3f> Normals;
 	TArray<FVector3f> Tangents;
@@ -282,6 +295,7 @@ struct FHairMeshesBulkData
 	uint32 GetNumTriangles() const { return Indices.Num() / 3; }
 	uint32 GetNumVertices() const { return Positions.Num(); }
 	const FBox& GetBounds() const { return Header.BoundingBox; }
+	uint32 GetNumUVs() const { return Header.NumUVs; }
 
 	TArray<FHairCardsPositionFormat::Type> Positions;
 	TArray<FHairCardsNormalFormat::Type> Normals;
@@ -292,6 +306,7 @@ struct FHairMeshesBulkData
 	struct FHeader
 	{
 		FBox BoundingBox = FBox(EForceInit::ForceInit);
+		uint32 NumUVs = 0;
 		bool bVertexColor = false;
 	} Header;
 
diff --git a/Engine/Shaders/Private/Common.ush b/Engine/Shaders/Private/Common.ush
--- a/Engine/Shaders/Private/Common.ush
+++ b/Engine/Shaders/Private/Common.ush
@@ -2358,6 +2358,7 @@ uint DivideAndRoundUp(uint Dividend, uint Divisor, uint DivisorAsBitShift)
 	return (Dividend + Divisor - 1) >> DivisorAsBitShift;
 }
 
+#define DivideAndRoundUp2(Dividend) DivideAndRoundUp((Dividend),    2u,  1u)
 #define DivideAndRoundUp4(Dividend) DivideAndRoundUp((Dividend),    4u,  2u)
 #define DivideAndRoundUp8(Dividend) DivideAndRoundUp((Dividend),    8u,  3u)
 #define DivideAndRoundUp16(Dividend) DivideAndRoundUp((Dividend),  16u,  4u)
diff --git a/Engine/Shaders/Private/HairStrands/HairCardsVertexFactory.ush b/Engine/Shaders/Private/HairStrands/HairCardsVertexFactory.ush
--- a/Engine/Shaders/Private/HairStrands/HairCardsVertexFactory.ush
+++ b/Engine/Shaders/Private/HairStrands/HairCardsVertexFactory.ush
@@ -25,7 +25,7 @@
 
 ////////////////////////////////////////////////////////////////////////////////
 // HAIR_CARD_MESH_FACTORY
-// This is set by compilation enviromenent of the vertex factory
+// This is set by compilation environment of the vertex factory
 ///////////////////////////////////////////////////////////////////////////////
 
 /**
@@ -38,19 +38,35 @@ struct FVertexFactoryInput
 	uint VertexId : SV_VertexID;
 #else
 	float4 Position      : ATTRIBUTE0; 
-	float4 PrevPosition  : ATTRIBUTE5;
+	float4 PrevPosition  : ATTRIBUTE1;
 
 	// TangentZ.w contains sign of tangent basis determinant
   #if METAL_ES3_1_PROFILE
-	float3	TangentX : ATTRIBUTE1;
-	float4	TangentZ : ATTRIBUTE2;
+	float3	TangentX : ATTRIBUTE2;
+	float4	TangentZ : ATTRIBUTE3;
   #else
-	half3	TangentX : ATTRIBUTE1;
-	half4	TangentZ : ATTRIBUTE2;
+	half3	TangentX : ATTRIBUTE2;
+	half4	TangentZ : ATTRIBUTE3;
   #endif
-	float4 UVs		 : ATTRIBUTE3;
 	float4 Materials : ATTRIBUTE4;
-	float4 VertexColors : ATTRIBUTE6;
+	float4 VertexColors : ATTRIBUTE5;
+	float4 UVs : ATTRIBUTE6; // RootUV/AtlasUV
+
+  #if NUM_MATERIAL_TEXCOORDS_VERTEX
+	// These used to be packed texcoord arrays, but these cause problems with alighnment on some Vulkan drivers
+	#if NUM_MATERIAL_TEXCOORDS_VERTEX >= 1
+	float4	TexCoords0 : ATTRIBUTE7;
+	#endif
+	#if NUM_MATERIAL_TEXCOORDS_VERTEX >= 3
+	float4	TexCoords1 : ATTRIBUTE8;
+	#endif
+	#if NUM_MATERIAL_TEXCOORDS_VERTEX >= 5
+	float4	TexCoords2 : ATTRIBUTE9;
+	#endif
+	#if NUM_MATERIAL_TEXCOORDS_VERTEX >= 7
+	float4	TexCoords3 : ATTRIBUTE10;
+	#endif
+  #endif
 #endif
 
 	VF_GPUSCENE_DECLARE_INPUT_BLOCK(13)
@@ -84,8 +100,6 @@ struct FVertexFactoryIntermediates
 	half3x3 TangentToWorld;
 	half TangentToWorldSign;
 	half4 VertexColors;
-	float2 HairPrimitiveUV;
-	float2 HairPrimitiveRootUV;
 	float2 HairPrimitiveLength;
 	float4 HairPrimitiveMaterial;
 	float  HairPrimitiveGroupIndex;
@@ -111,8 +125,6 @@ uint GetHairVertexFetchId(FVertexFactoryInput Input)
 struct FHairCardsPoint
 {
 	float3 Position;
-	float2 CardsUVs;
-	float2 RootUVs;
 	float  CardLength;
 	float  GroupIndex;
 	float4 Materials;
@@ -125,7 +137,6 @@ FHairCardsPoint GetVertexPosition(FVertexFactoryInput Input)
 #if MANUAL_VERTEX_FETCH
 	const uint VertexFetchId = GetHairVertexFetchId(Input);
 	const float4 Data = HairCardsVF.PositionBuffer[VertexFetchId];
-	const float4 UVs  = HairCardsVF.UVsBuffer[VertexFetchId];
 	const float4 Materials = HairCardsVF.MaterialsBuffer[VertexFetchId];
 	if (HasHairCardsVertexColor())
 	{
@@ -133,7 +144,6 @@ FHairCardsPoint GetVertexPosition(FVertexFactoryInput Input)
 	}
 #else
 	const float4 Data = Input.Position;
-	const float4 UVs  = Input.UVs;
 	const float4 Materials = Input.Materials;
 	if (HasHairCardsVertexColor())
 	{
@@ -145,8 +155,6 @@ FHairCardsPoint GetVertexPosition(FVertexFactoryInput Input)
 
 	FHairCardsPoint Out;
 	Out.Position	= Data.xyz;
-	Out.CardsUVs	= UVs.xy;
-	Out.RootUVs		= UVs.zw;
 	Out.CardLength	= f16tof32(EncodedW & 0xFFFF);
 	Out.GroupIndex	= (EncodedW>>16u);
 	Out.Materials	= float4(Materials.xyz*Materials.xyz, Materials.w); // Decoding (cheap) sRGB -> linear color
@@ -200,7 +208,6 @@ FMaterialPixelParameters GetMaterialPixelParameters(FVertexFactoryInterpolantsVS
 	Result.HairPrimitiveLength = Interpolants.HairPrimitiveLength;
 	Result.HairPrimitiveGroupIndex = Interpolants.HairPrimitiveGroupIndex;
 	Result.HairPrimitiveMaterial = Interpolants.HairPrimitiveMaterial;
-
 	Result.VertexColor = GetColor(Interpolants);
 
 #if NUM_TEX_COORD_INTERPOLATORS
@@ -227,6 +234,65 @@ half3x3 CalcTangentToWorldNoScale(FVertexFactoryIntermediates Intermediates, hal
 float3 VertexFactoryGetPreviousInstanceSpacePosition(FVertexFactoryInput Input, FVertexFactoryIntermediates Intermediates);
 float3 VertexFactoryGetInstanceSpacePosition(FVertexFactoryInput Input, FVertexFactoryIntermediates Intermediates);
 
+void GetRootUVsAtlasUVs(FVertexFactoryInput Input, inout float2 OutAtlasUV, inout float2 OutRootUV)
+{
+#if MANUAL_VERTEX_FETCH
+	const uint NumFetchUVs = DivideAndRoundUp2(HairCardsVF.NumUVs);
+	const uint UVFetchIndex = Input.VertexId * NumFetchUVs;
+	const float4 UVs = HairCardsVF.UVsBuffer[UVFetchIndex];
+	OutRootUV  = UVs.xy;
+	OutAtlasUV = UVs.zw;
+#else
+	OutRootUV  = Input.UVs.xy;
+	OutAtlasUV = Input.UVs.zw;
+#endif
+}
+
+void GetTexCoords(FVertexFactoryInput Input, inout float2 OutTexCoords[NUM_MATERIAL_TEXCOORDS_VERTEX])
+{
+#if NUM_MATERIAL_TEXCOORDS_VERTEX
+#if MANUAL_VERTEX_FETCH
+	UNROLL
+	for (uint CoordIndex = 0; CoordIndex < NUM_MATERIAL_TEXCOORDS_VERTEX; ++CoordIndex)
+	{
+		OutTexCoords[CoordIndex] = 0;
+		const uint ReadCoordIndex = CoordIndex + 1; // UVs[0] == RootUV, which is skipped 
+		if (ReadCoordIndex < HairCardsVF.NumUVs)
+		{
+			const uint FetchIndex = ReadCoordIndex >> 1u;
+			const float4 PackedUVs = HairCardsVF.UVsBuffer[FetchIndex];
+			OutTexCoords[CoordIndex] = (CoordIndex & 1) == 0 ? PackedUVs.zw : PackedUVs.xy; //.xy/.zw are inversed, as TexCoord are shifted by 1 due to UVs[0] == RootUV, which is skipped 
+		}
+	}
+#else // MANUAL_VERTEX_FETCH
+	#if NUM_MATERIAL_TEXCOORDS_VERTEX > 0
+	OutTexCoords[0] = Input.TexCoords0.xy;
+	#endif
+	#if NUM_MATERIAL_TEXCOORDS_VERTEX > 1
+	OutTexCoords[1] = Input.TexCoords0.zw;
+	#endif
+	#if NUM_MATERIAL_TEXCOORDS_VERTEX > 2
+	OutTexCoords[2] = Input.TexCoords1.xy;
+	#endif
+	#if NUM_MATERIAL_TEXCOORDS_VERTEX > 3
+	OutTexCoords[3] = Input.TexCoords1.zw;
+	#endif
+	#if NUM_MATERIAL_TEXCOORDS_VERTEX > 4
+	OutTexCoords[4] = Input.TexCoords2.xy;
+	#endif
+	#if NUM_MATERIAL_TEXCOORDS_VERTEX > 5
+	OutTexCoords[5] = Input.TexCoords2.zw;
+	#endif
+	#if NUM_MATERIAL_TEXCOORDS_VERTEX > 6
+	OutTexCoords[6] = Input.TexCoords3.xy;
+	#endif
+	#if NUM_MATERIAL_TEXCOORDS_VERTEX > 7
+	OutTexCoords[7] = Input.TexCoords3.zw;
+	#endif
+#endif  //MANUAL_VERTEX_FETCH 
+#endif // NUM_MATERIAL_TEXCOORDS_VERTEX
+}
+
 /** Converts from vertex factory specific input to a FMaterialVertexParameters, which is used by vertex shader material inputs. */
 FMaterialVertexParameters GetMaterialVertexParameters(
 	FVertexFactoryInput Input, 
@@ -260,6 +326,10 @@ FMaterialVertexParameters GetMaterialVertexParameters(
 
 	Result.LWCData = MakeMaterialLWCData(Result);
 
+#if NUM_MATERIAL_TEXCOORDS_VERTEX
+	GetTexCoords(Input, Result.TexCoords);
+#endif
+
 	return Result;
 }
 
@@ -309,8 +379,6 @@ FVertexFactoryIntermediates GetVertexFactoryIntermediates(FVertexFactoryInput In
 
 	FVertexFactoryIntermediates Intermediates;
 	Intermediates.SceneData					= VF_GPUSCENE_GET_INTERMEDIATES(Input);
-	Intermediates.HairPrimitiveUV			= CardsPoint.CardsUVs;
-	Intermediates.HairPrimitiveRootUV		= CardsPoint.RootUVs;
 	Intermediates.HairPrimitiveLength		= CardsPoint.CardLength; 
 	Intermediates.HairPrimitiveGroupIndex	= CardsPoint.GroupIndex;
 	Intermediates.HairPrimitiveMaterial		= CardsPoint.Materials;
@@ -370,17 +438,27 @@ FVertexFactoryInterpolantsVSToPS VertexFactoryGetInterpolantsVSToPS(FVertexFacto
 	SetColor(Interpolants, Intermediates.VertexColors);
 
 	FHairCardsPoint CardsPoint = GetVertexPosition(Input);
-	Interpolants.HairPrimitiveUV			= CardsPoint.CardsUVs;
-	Interpolants.HairPrimitiveRootUV		= CardsPoint.RootUVs;
 	Interpolants.HairPrimitiveLength		= CardsPoint.CardLength;
 	Interpolants.HairPrimitiveGroupIndex	= CardsPoint.GroupIndex;
 	Interpolants.HairPrimitiveMaterial		= CardsPoint.Materials;
 
+	GetRootUVsAtlasUVs(Input, Interpolants.HairPrimitiveUV, Interpolants.HairPrimitiveRootUV);
+
 	SetPrimitiveId(Interpolants, Intermediates.SceneData.PrimitiveId);
 #if NUM_TEX_COORD_INTERPOLATORS 
-	SetUV(Interpolants, 0, Interpolants.HairPrimitiveUV);
+	float2 TexCoords[NUM_TEX_COORD_INTERPOLATORS];
+	GetTexCoords(Input, TexCoords);
+	UNROLL
+	for (int CoordinateIndex = 0; CoordinateIndex < NUM_TEX_COORD_INTERPOLATORS; CoordinateIndex++)
+	{
+		SetUV(Interpolants, CoordinateIndex, TexCoords[CoordinateIndex]);
+	}
 #endif
 
+	#if NUM_TEX_COORD_INTERPOLATORS 
+	SetUV(Interpolants, 0, Interpolants.HairPrimitiveUV);
+	#endif
+
 	return Interpolants;
 }
 
