Make no memory estimate log level a warning instead of verbose
Implement per asset type default memory estimates

#rnx

diff --git a/Engine/Plugins/Runtime/HairStrands/Source/HairStrandsCore/Private/GroomBindingAsset.cpp b/Engine/Plugins/Runtime/HairStrands/Source/HairStrandsCore/Private/GroomBindingAsset.cpp
--- a/Engine/Plugins/Runtime/HairStrands/Source/HairStrandsCore/Private/GroomBindingAsset.cpp
+++ b/Engine/Plugins/Runtime/HairStrands/Source/HairStrandsCore/Private/GroomBindingAsset.cpp
@@ -37,6 +37,12 @@ LLM_DECLARE_TAG(Groom);
 static int32 GHairStrandsDelayBindingDataRelease = 1;
 static FAutoConsoleVariableRef CVarHairStrandsDelayBindingDataRelease(TEXT("r.HairStrands.Binding.DelayBindingDataRelease"), GHairStrandsDelayBindingDataRelease, TEXT("Delay groom binding data release (Experimental)"));
 
+static TAutoConsoleVariable<int32> CVarMemoryForGroomBindingAssetCompile(
+	TEXT("Memory.MemoryForGroomBindingAssetCompile"),
+	512,
+	TEXT("Memory in MiB set aside for a GroomBindingAsset compile job\n"),
+	ECVF_Default);
+
 /////////////////////////////////////////////////////////////////////////////////////////
 
 bool IsHairStrandsDDCLogEnable();
@@ -1106,8 +1112,11 @@ void UGroomBindingAsset::BeginCacheDerivedDatas(
 		FQueuedThreadPool* ThreadPool = FGroomBindingCompilingManager::Get().GetThreadPool();
 		EQueuedWorkPriority BasePriority = FGroomBindingCompilingManager::Get().GetBasePriority(this);
 		check(AsyncTask == nullptr);
+
+		int64 RequiredMemory = 1024 * 1024 * CVarMemoryForGroomBindingAssetCompile.GetValueOnAnyThread();
+
 		AsyncTask = MakeUnique<FGroomBindingAsyncBuildTask>(this, MoveTemp(Context));
-		AsyncTask->StartBackgroundTask(ThreadPool, BasePriority, EQueuedWorkFlags::DoNotRunInsideBusyWait);
+		AsyncTask->StartBackgroundTask(ThreadPool, BasePriority, EQueuedWorkFlags::DoNotRunInsideBusyWait, RequiredMemory, TEXT("GroomBindingAsset"));
 		FGroomBindingCompilingManager::Get().AddGroomBindings({ this });
 	}
 	else
diff --git a/Engine/Source/Runtime/Engine/Private/Animation/AnimBank.cpp b/Engine/Source/Runtime/Engine/Private/Animation/AnimBank.cpp
--- a/Engine/Source/Runtime/Engine/Private/Animation/AnimBank.cpp
+++ b/Engine/Source/Runtime/Engine/Private/Animation/AnimBank.cpp
@@ -46,6 +46,12 @@ CSV_DECLARE_CATEGORY_MODULE_EXTERN(ENGINE_API, Animation);
 
 DEFINE_LOG_CATEGORY(LogAnimBank);
 
+static TAutoConsoleVariable<int32> CVarMemoryForAnimBankAssetCompile(
+	TEXT("Memory.MemoryForAnimBankAssetCompile"),
+	512,
+	TEXT("Memory in MiB set aside for a AnimBank asset compile job\n"),
+	ECVF_Default);
+
 #if WITH_EDITOR
 
 UAnimBank::FOnDependenciesChanged UAnimBank::OnDependenciesChanged;
@@ -370,7 +376,7 @@ void FAnimBankBuildAsyncCacheTask::BeginCache(const FIoHash& InKeyHash)
 		FQueuedThreadPool* ThreadPool = FAnimBankCompilingManager::Get().GetThreadPool();
 		EQueuedWorkPriority BasePriority = FAnimBankCompilingManager::Get().GetBasePriority(Bank);
 
-		int64 RequiredMemory = -1; // @todo RequiredMemory - possibly follow FAnimationSequenceAsyncCacheTask::GetRequiredMemoryEstimate()
+		int64 RequiredMemory = 1024 * 1024 * CVarMemoryForAnimBankAssetCompile.GetValueOnAnyThread();
 
 		check(BuildTask == nullptr);
 		BuildTask = MakeUnique<FAnimBankAsyncBuildTask>(this, InKeyHash);
diff --git a/Engine/Source/Runtime/Engine/Private/AssetCompilingManager.cpp b/Engine/Source/Runtime/Engine/Private/AssetCompilingManager.cpp
--- a/Engine/Source/Runtime/Engine/Private/AssetCompilingManager.cpp
+++ b/Engine/Source/Runtime/Engine/Private/AssetCompilingManager.cpp
@@ -355,10 +355,6 @@ public:
 		int64 RequiredMemory = InQueuedWork->GetRequiredMemory();
 		const TCHAR * DebugName = InQueuedWork->GetDebugName();
 		
-		//these conditions should in fact both be true :
-		//check( DebugName != nullptr );
-		//check( RequiredMemory != 0 );
-
 		if ( DebugName == nullptr )
 		{
 			DebugName = TEXT("No DebugName");
@@ -367,8 +363,7 @@ public:
 		if ( RequiredMemory < 0 )
 		{
 			// log the types that are failing to set required memory
-			// @todo RequiredMemory : make this a Warning when possible
-			UE_LOG(LogAsyncCompilation, Verbose, TEXT("AddQueuedWork [%s] RequiredMemory unknown; fix me!"),DebugName);
+			UE_LOG(LogAsyncCompilation, Warning, TEXT("AddQueuedWork [%s] did not supply a RequiredMemory estimate. This is a bug, all asset types should provide an estimate. When no estimate is provided asset compilation is slowed as the queue has to operate very conservatively."),DebugName);
 		}
 		else
 		{
diff --git a/Engine/Source/Runtime/Engine/Private/SkeletalMesh.cpp b/Engine/Source/Runtime/Engine/Private/SkeletalMesh.cpp
--- a/Engine/Source/Runtime/Engine/Private/SkeletalMesh.cpp
+++ b/Engine/Source/Runtime/Engine/Private/SkeletalMesh.cpp
@@ -130,6 +130,12 @@ static TAutoConsoleVariable<int32> CVarRayTracingSkeletalMeshLODBias(
 	TEXT("Final LOD level to use in ray tracing is the sum of this global bias and the bias set on each skeletal mesh asset."),
 	ECVF_RenderThreadSafe);
 
+static TAutoConsoleVariable<int32> CVarMemoryForSkeletalMeshAssetCompile(
+	TEXT("Memory.MemoryForSkeletalMeshAssetCompile"),
+	1024 * 4, // 4GiB
+	TEXT("Memory in MiB set aside for a SkeletalMeshAsset compile job\n"),
+	ECVF_Default);
+
 const TCHAR* GSkeletalMeshMinLodQualityLevelCVarName = TEXT("r.SkeletalMesh.MinLodQualityLevel");
 const TCHAR* GSkeletalMeshMinLodQualityLevelScalabilitySection = TEXT("ViewDistanceQuality");
 int32 GSkeletalMeshMinLodQualityLevel = -1;
@@ -2244,8 +2250,11 @@ FEvent* USkeletalMesh::LockPropertiesUntil()
 	FQueuedThreadPool* SkeletalMeshThreadPool = FSkinnedAssetCompilingManager::Get().GetThreadPool();
 	EQueuedWorkPriority BasePriority = FSkinnedAssetCompilingManager::Get().GetBasePriority(this);
 	check(AsyncTask == nullptr);
+
+	int64 RequiredMemory = 1024 * 1024 * CVarMemoryForSkeletalMeshAssetCompile.GetValueOnAnyThread();
+
 	AsyncTask = MakeUnique<FSkinnedAssetAsyncBuildTask>(this, MoveTemp(Context));
-	AsyncTask->StartBackgroundTask(SkeletalMeshThreadPool, BasePriority, EQueuedWorkFlags::DoNotRunInsideBusyWait);
+	AsyncTask->StartBackgroundTask(SkeletalMeshThreadPool, BasePriority, EQueuedWorkFlags::DoNotRunInsideBusyWait, RequiredMemory, TEXT("SkeletalMeshAsset"));
 	FSkinnedAssetCompilingManager::Get().AddSkinnedAssets({ this });
 	return Event;
 }
diff --git a/Engine/Source/Runtime/Engine/Private/SkinnedAsset.cpp b/Engine/Source/Runtime/Engine/Private/SkinnedAsset.cpp
--- a/Engine/Source/Runtime/Engine/Private/SkinnedAsset.cpp
+++ b/Engine/Source/Runtime/Engine/Private/SkinnedAsset.cpp
@@ -36,6 +36,12 @@ static FAutoConsoleVariableRef CVarAnimSkinnedAssetISPCEnabled(TEXT("a.SkinnedAs
 
 DEFINE_LOG_CATEGORY_STATIC(LogSkinnedAsset, Log, All);
 
+static TAutoConsoleVariable<int32> CVarMemoryForSkinnedAssetCompile (
+	TEXT("Memory.MemoryForSkinnedAssetCompile"),
+	1024 * 4, // 4GiB
+	TEXT("Memory in MiB set aside for a SkinnedAsset compile job\n"),
+	ECVF_Default);
+
 #if WITH_EDITORONLY_DATA
 static void AccumulateUVDensities(float* OutWeightedUVDensities, float* OutWeights, const FSkeletalMeshLODRenderData& LODData, const FSkelMeshRenderSection& Section)
 {
@@ -134,7 +140,7 @@ void USkinnedAsset::PostLoad()
 		EQueuedWorkPriority BasePriority = FSkinnedAssetCompilingManager::Get().GetBasePriority(this);
 
 		AsyncTask = MakeUnique<FSkinnedAssetAsyncBuildTask>(this, MoveTemp(Context));
-		int64 RequiredMemory = -1; // @todo RequiredMemory
+		int64 RequiredMemory = 1024 * 1024 * CVarMemoryForSkinnedAssetCompile.GetValueOnAnyThread();
 		AsyncTask->StartBackgroundTask(ThreadPool, BasePriority, EQueuedWorkFlags::DoNotRunInsideBusyWait, RequiredMemory, TEXT("SkinnedAsset"));
 		FSkinnedAssetCompilingManager::Get().AddSkinnedAssets({ this });
 	}
