Fix hair binding not working correctly when a skin cache has certain sections marked as disabled.

diff --git a/Engine/Plugins/Runtime/HairStrands/Source/HairStrandsCore/Private/HairStrandsMeshProjection.cpp b/Engine/Plugins/Runtime/HairStrands/Source/HairStrandsCore/Private/HairStrandsMeshProjection.cpp
--- a/Engine/Plugins/Runtime/HairStrands/Source/HairStrandsCore/Private/HairStrandsMeshProjection.cpp
+++ b/Engine/Plugins/Runtime/HairStrands/Source/HairStrandsCore/Private/HairStrandsMeshProjection.cpp
@@ -212,6 +212,7 @@ static bool AddHairStrandUpdateMeshTrianglesPass(
 	FGlobalShaderMap* ShaderMap,
 	const uint32 UniqueTriangleCount, 
 	const uint32 TotalSectionCount, // Total section count of the underlying mesh
+	const TArray<uint32>& UniqueMeshSectionIndices,
 	const FCachedGeometry& MeshLODData,
 	FRDGBufferSRVRef UniqueTriangleIndexSRV,
 	FRDGBufferUAVRef OutputCurrUAV,
@@ -269,8 +270,11 @@ static bool AddHairStrandUpdateMeshTrianglesPass(
 	// Allocate data for *all* sections, but only fill in the used/valid sections
 	TArray<FSectionData> SectionDatas;
 	SectionDatas.Init(Default, TotalSectionCount); 
-	for (const FCachedGeometry::Section& MeshSectionData : MeshLODData.Sections)
+	for (uint32 MeshSectionIndex : UniqueMeshSectionIndices)
 	{
+		check(MeshLODData.Sections.IsValidIndex(MeshSectionIndex));
+
+		const FCachedGeometry::Section& MeshSectionData = MeshLODData.Sections[MeshSectionIndex];
 		const uint32 SectionIndex = MeshSectionData.SectionIndex;
 		SectionDatas[SectionIndex].TotalIndexCount	= MeshSectionData.TotalIndexCount;
 		SectionDatas[SectionIndex].TotalVertexCount	= MeshSectionData.TotalVertexCount;
@@ -374,6 +378,7 @@ void AddHairStrandUpdateMeshTrianglesPass(
 			ShaderMap,
 			UniqueTriangleCount,
 			TotalMeshSectionCount,
+			RestLODData.GetBulkData().Header.UniqueSectionIndices,
 			MeshLODData,
 			RegisterAsSRV(GraphBuilder, RestLODData.UniqueTriangleIndexBuffer),
 			OutputCurrBuffer.UAV,
diff --git a/Engine/Source/Runtime/Engine/Private/SkeletalRenderGPUSkin.cpp b/Engine/Source/Runtime/Engine/Private/SkeletalRenderGPUSkin.cpp
--- a/Engine/Source/Runtime/Engine/Private/SkeletalRenderGPUSkin.cpp
+++ b/Engine/Source/Runtime/Engine/Private/SkeletalRenderGPUSkin.cpp
@@ -2532,7 +2532,11 @@ bool FSkeletalMeshObjectGPUSkin::GetCachedGeometry(FRDGBuilder& GraphBuilder, FC
 			FRWBuffer* PositionBuffer = FGPUSkinCache::GetPositionBuffer(GraphBuilder, SkinCacheEntry, SectionIndex);
 			if (!PositionBuffer || !PositionBuffer->SRV)
 			{
-				return false;
+				// Skip section which are not valid, but continue to fill in other section.
+				// A section might be invalid because it is disabled/deactivated for instance, but the rest of the sections are still valid
+				CachedSection = FCachedGeometry::Section();
+				CachedSection.SectionIndex = SectionIndex;
+				continue;
 			}
 			FRWBuffer* PreviousPositionBuffer = FGPUSkinCache::GetPreviousPositionBuffer(GraphBuilder, SkinCacheEntry, SectionIndex);
 
