Restrict simulation visualization to current dataflow editor only

diff --git a/Engine/Plugins/Runtime/HairStrands/Source/HairStrandsDataflow/HairStrandsDataflow.Build.cs b/Engine/Plugins/Runtime/HairStrands/Source/HairStrandsDataflow/HairStrandsDataflow.Build.cs
--- a/Engine/Plugins/Runtime/HairStrands/Source/HairStrandsDataflow/HairStrandsDataflow.Build.cs
+++ b/Engine/Plugins/Runtime/HairStrands/Source/HairStrandsDataflow/HairStrandsDataflow.Build.cs
@@ -1,4 +1,4 @@
-﻿// Copyright Epic Games, Inc. All Rights Reserved.
+// Copyright Epic Games, Inc. All Rights Reserved.
 
 namespace UnrealBuildTool.Rules
 {
@@ -29,7 +29,8 @@ namespace UnrealBuildTool.Rules
 					"SkeletalMeshDescription",
 					"StaticMeshDescription",
 					"Slate",
-					"SlateCore"
+					"SlateCore",
+					"UnrealEd"
 				}
 			);
 		}
diff --git a/Engine/Plugins/Runtime/HairStrands/Source/HairStrandsDataflow/Private/GroomDataflowVisualization.cpp b/Engine/Plugins/Runtime/HairStrands/Source/HairStrandsDataflow/Private/GroomDataflowVisualization.cpp
--- a/Engine/Plugins/Runtime/HairStrands/Source/HairStrandsDataflow/Private/GroomDataflowVisualization.cpp
+++ b/Engine/Plugins/Runtime/HairStrands/Source/HairStrandsDataflow/Private/GroomDataflowVisualization.cpp
@@ -33,71 +33,54 @@ namespace UE::Groom
 				{
 					if (GetGroomComponent(SimulationScene.Get()))
 					{
-						static IConsoleVariable* CVarGroomViewMode = nullptr; 
-						if (CVarGroomViewMode == nullptr)
-						{
-							if (IConsoleVariable* ConsoleVariable = IConsoleManager::Get().FindConsoleVariable(
-								FGroomVisualizationData::GetVisualizeConsoleCommandName()))
-							{
-								CVarGroomViewMode = ConsoleVariable->AsVariable();
-							}
-						}
-						static IConsoleVariable* CVarGroomStrandWidth = nullptr; 
-						if (CVarGroomStrandWidth == nullptr)
-						{
-							if (IConsoleVariable* ConsoleVariable = IConsoleManager::Get().FindConsoleVariable(
-								TEXT("r.HairStrands.StrandWidth")))
-							{
-								CVarGroomStrandWidth = ConsoleVariable->AsVariable();
-							}
-						}
 						constexpr uint8 NumFlags = static_cast<uint8>(EGroomViewMode::Count);
-						auto AddSimulationVisualisation = [this, &MenuBuilder, NumFlags](const TAttribute<FText>& FlagLabel, const uint8 FlagIndex, const int32 ViewMode, const float StrandWidth)
-						{
-							// Handler for visualization entry being clicked
-							const FExecuteAction ExecuteAction = FExecuteAction::CreateLambda([this, FlagIndex, ViewMode, StrandWidth]()
+						auto AddSimulationVisualisation = [this, &ViewportClient , &MenuBuilder, NumFlags](const TAttribute<FText>& FlagLabel, const EGroomViewMode ViewMode)
 							{
-								VisualizationFlags[FlagIndex] = !VisualizationFlags[FlagIndex];
-								if(VisualizationFlags[FlagIndex])
-								{
-									CVarGroomViewMode->SetWithCurrentPriority(ViewMode);
-									CVarGroomStrandWidth->SetWithCurrentPriority(StrandWidth);
-									for(int32 OtherIndex = 0; OtherIndex < NumFlags; ++OtherIndex)
+								const uint8 FlagIndex = static_cast<uint8>(ViewMode);
+
+								// Handler for visualization entry being clicked
+								const FExecuteAction ExecuteAction = FExecuteAction::CreateLambda([this, FlagIndex, ViewMode, ViewportClient]()
 									{
-										if(OtherIndex != FlagIndex)
+										VisualizationFlags[FlagIndex] = !VisualizationFlags[FlagIndex];
+										if (VisualizationFlags[FlagIndex])
 										{
-											VisualizationFlags[OtherIndex] = false;
+											ViewportClient->ChangeGroomVisualizationMode(GetGroomViewModeName(ViewMode));
+											for (int32 OtherIndex = 0; OtherIndex < NumFlags; ++OtherIndex)
+											{
+												if (OtherIndex != FlagIndex)
+												{
+													VisualizationFlags[OtherIndex] = false;
+												}
+											}
 										}
-									}
-								}
-								else
-								{
-									CVarGroomViewMode->SetWithCurrentPriority(0);
-									CVarGroomStrandWidth->SetWithCurrentPriority(0.0f);
-								}
-							});
-
-							// Checkstate function for visualization entries
-							const FIsActionChecked IsActionChecked = FIsActionChecked::CreateLambda([this, FlagIndex]()
-							{
-								return VisualizationFlags[FlagIndex];
-							});
+										else
+										{
+											ViewportClient->ChangeGroomVisualizationMode(GetGroomViewModeName(EGroomViewMode::None));
+										}
+									});
+
+								// Check state function for visualization entries
+								const FIsActionChecked IsActionChecked = FIsActionChecked::CreateLambda([this, FlagIndex]()
+									{
+										return VisualizationFlags[FlagIndex];
+									});
+
+								const FUIAction Action(ExecuteAction, FCanExecuteAction(), IsActionChecked);
 
-							const FUIAction Action(ExecuteAction, FCanExecuteAction(), IsActionChecked);
+								// Add menu entry
+								MenuBuilder.AddMenuEntry(FlagLabel, FText(), FSlateIcon(), Action, NAME_None, EUserInterfaceActionType::ToggleButton);
+							};
 
-							// Add menu entry
-							MenuBuilder.AddMenuEntry(FlagLabel, FText(), FSlateIcon(), Action, NAME_None, EUserInterfaceActionType::ToggleButton);
-						};
 						MenuBuilder.BeginSection(TEXT("GroomSimulation_Visualizations"), LOCTEXT("GroomSimulationVisualization", "Groom Simulation Visualization"));
 
-						AddSimulationVisualisation(LOCTEXT("SimulatedGuides", "Simulated guides"), static_cast<uint8>(EGroomViewMode::SimHairStrands), 21, 0.5f);
-						AddSimulationVisualisation(LOCTEXT("RenderedStrands", "Rendered strands"), static_cast<uint8>(EGroomViewMode::RenderHairStrands), 22, 0.0f);
-						AddSimulationVisualisation(LOCTEXT("StrandsGroups", "Strands groups"), static_cast<uint8>(EGroomViewMode::Group), 35, 0.0f);
-						AddSimulationVisualisation(LOCTEXT("StrandsClumps", "Strands clumps"), static_cast<uint8>(EGroomViewMode::ClumpID), 39, 0.0f);
-						AddSimulationVisualisation(LOCTEXT("StrandsClusters", "Strands clusters"), static_cast<uint8>(EGroomViewMode::Cluster), 31, 0.0f);
-						AddSimulationVisualisation(LOCTEXT("MeshProjection", "Mesh projection"), static_cast<uint8>(EGroomViewMode::MeshProjection), 12, 0.0f);
-						AddSimulationVisualisation(LOCTEXT("CardsGuides", "Cards guides"), static_cast<uint8>(EGroomViewMode::CardGuides), 37, 0.0f);
-						AddSimulationVisualisation(LOCTEXT("LODColoration", "LOD Coloration"), static_cast<uint8>(EGroomViewMode::LODColoration), 36, 0.0f);
+						AddSimulationVisualisation(LOCTEXT("SimulatedGuides", "Simulated guides"), EGroomViewMode::SimHairStrands);
+						AddSimulationVisualisation(LOCTEXT("RenderedStrands", "Rendered strands"), EGroomViewMode::RenderHairStrands);
+						AddSimulationVisualisation(LOCTEXT("StrandsGroups", "Strands groups"), EGroomViewMode::Group);
+						AddSimulationVisualisation(LOCTEXT("StrandsClumps", "Strands clumps"), EGroomViewMode::ClumpID);
+						AddSimulationVisualisation(LOCTEXT("StrandsClusters", "Strands clusters"), EGroomViewMode::Cluster);
+						AddSimulationVisualisation(LOCTEXT("MeshProjection", "Mesh projection"), EGroomViewMode::MeshProjection);
+						AddSimulationVisualisation(LOCTEXT("CardsGuides", "Cards guides"), EGroomViewMode::CardGuides);
+						AddSimulationVisualisation(LOCTEXT("LODColoration", "LOD Coloration"), EGroomViewMode::LODColoration);
 						
 						MenuBuilder.EndSection();
 					}
