Add support for groom dataflow variable overrides

diff --git a/Engine/Plugins/Runtime/HairStrands/Source/HairStrandsCore/Private/GroomAsset.cpp b/Engine/Plugins/Runtime/HairStrands/Source/HairStrandsCore/Private/GroomAsset.cpp
--- a/Engine/Plugins/Runtime/HairStrands/Source/HairStrandsCore/Private/GroomAsset.cpp
+++ b/Engine/Plugins/Runtime/HairStrands/Source/HairStrandsCore/Private/GroomAsset.cpp
@@ -891,7 +891,7 @@ void UGroomAsset::Serialize(FArchive& Ar)
 }
 
 UGroomAsset::UGroomAsset(const FObjectInitializer& ObjectInitializer)
-	: Super(ObjectInitializer)
+	: Super(ObjectInitializer), DataflowSettings(this, nullptr, TEXT("GroomAssetTerminal"))
 {
 	bIsInitialized = false;
 
@@ -1222,6 +1222,34 @@ void UGroomAsset::UpdateHairGroupsInfo()
 	}
 }
 
+static void SyncHairGroupsName(UGroomAsset* In)
+{
+#if WITH_EDITORONLY_DATA
+	check(In);
+
+	uint32 GroupIndex = 0;
+	for (FHairGroupInfoWithVisibility& Info : In->GetHairGroupsInfo())
+	{
+		if (In->GetHairGroupsPhysics().IsValidIndex(GroupIndex))
+		{
+			In->GetHairGroupsPhysics()[GroupIndex].GroupName = Info.GroupName;
+		}
+
+		if (In->GetHairGroupsRendering().IsValidIndex(GroupIndex))
+		{
+			In->GetHairGroupsRendering()[GroupIndex].GroupName = Info.GroupName;
+		}
+
+		if (In->GetHairGroupsInterpolation().IsValidIndex(GroupIndex))
+		{
+			In->GetHairGroupsInterpolation()[GroupIndex].GroupName = Info.GroupName;
+		}
+
+		++GroupIndex;
+	}
+#endif
+}
+
 static bool UpdateHairGroupsName(UGroomAsset* In)
 {
 	bool bNeedResaved = false;
@@ -1241,6 +1269,7 @@ static bool UpdateHairGroupsName(UGroomAsset* In)
 				bNeedResaved = true;
 			}
 		}
+		SyncHairGroupsName(In);
 
 		++GroupIndex;
 	}
@@ -1497,6 +1526,7 @@ void UGroomAsset::PostLoad()
 	UpdateHairGroupsInfo();
 #if WITH_EDITORONLY_DATA
 	UpdateCachedSettings();
+	SyncHairGroupsName(this);
 #endif // #if WITH_EDITORONLY_DATA
 
 	if (IsResourcePSOPrecachingEnabled())
@@ -4051,7 +4081,7 @@ void UGroomAsset::WriteDataflowContent(const TObjectPtr<UDataflowBaseContent>& D
 	if(const TObjectPtr<UDataflowGroomContent> GroomContent = Cast<UDataflowGroomContent>(DataflowContent))
 	{
 		GroomContent->SetDataflowAsset(DataflowSettings.GetDataflowAsset());
-		GroomContent->SetDataflowTerminal(DataflowSettings.GetDataflowTerminal());
+		GroomContent->SetDataflowTerminal(DataflowSettings.GetDataflowTerminal().ToString());
 
 #if WITH_EDITORONLY_DATA
 		GroomContent->SetAnimationAsset(DataflowSettings.GetPreviewAnimationAsset());
diff --git a/Engine/Plugins/Runtime/HairStrands/Source/HairStrandsCore/Private/GroomAssetDataflow.cpp b/Engine/Plugins/Runtime/HairStrands/Source/HairStrandsCore/Private/GroomAssetDataflow.cpp
--- a/Engine/Plugins/Runtime/HairStrands/Source/HairStrandsCore/Private/GroomAssetDataflow.cpp
+++ b/Engine/Plugins/Runtime/HairStrands/Source/HairStrandsCore/Private/GroomAssetDataflow.cpp
@@ -1,4 +1,4 @@
-﻿// Copyright Epic Games, Inc. All Rights Reserved.
+// Copyright Epic Games, Inc. All Rights Reserved.
 
 #include "GroomAssetDataflow.h"
 #include "GroomBindingAsset.h"
@@ -82,7 +82,7 @@ UAnimationAsset* FGroomDataflowSettings::GetPreviewAnimationAsset() const
 
 #endif
 
-FGroomDataflowSettings::FGroomDataflowSettings()
+FGroomDataflowSettings::FGroomDataflowSettings(UObject* InOwner, UDataflow* InDataflowAsset, FName InTerminalNodeName) : FDataflowInstance(InOwner, InDataflowAsset, InTerminalNodeName)
 {
 	RestCollection = MakeShared<FManagedArrayCollection>();
 }
diff --git a/Engine/Plugins/Runtime/HairStrands/Source/HairStrandsCore/Public/GroomAsset.h b/Engine/Plugins/Runtime/HairStrands/Source/HairStrandsCore/Public/GroomAsset.h
--- a/Engine/Plugins/Runtime/HairStrands/Source/HairStrandsCore/Public/GroomAsset.h
+++ b/Engine/Plugins/Runtime/HairStrands/Source/HairStrandsCore/Public/GroomAsset.h
@@ -455,7 +455,7 @@ struct FHairVertexFactoryTypesPerMaterialData
  * Implements an asset that can be used to store hair strands
  */
 UCLASS(BlueprintType, AutoExpandCategories = ("HairRendering", "HairPhysics", "HairInterpolation", "HairDataflow"), hidecategories = (Object, "Hidden"), MinimalAPI)
-class UGroomAsset : public UObject, public IInterface_AssetUserData, public IDataflowContentOwner
+class UGroomAsset : public UObject, public IInterface_AssetUserData, public IDataflowContentOwner, public IDataflowInstanceInterface
 {
 	GENERATED_UCLASS_BODY()
 
@@ -479,19 +479,19 @@ private:
 
 public:
 	UE_DEPRECATED(5.3, "Please do not access this member directly; use UGroomAsset accessor.")
-	UPROPERTY(EditAnywhere, EditFixedSize, EditFixedSize, Category = "HairInfo", meta = (DisplayName = "Group"))
+	UPROPERTY(EditAnywhere, EditFixedSize, EditFixedSize, Category = "HairInfo", meta = (DisplayName = "Group", TitleProperty = "{GroupName}"))
 	TArray<FHairGroupInfoWithVisibility> HairGroupsInfo;
 
 	UE_DEPRECATED(5.3, "Please do not access this member directly; use UGroomAsset accessor.")
-	UPROPERTY(EditAnywhere, EditFixedSize, BlueprintGetter = GetHairGroupsRendering, BlueprintSetter = SetHairGroupsRendering, Category = "HairRendering", meta = (DisplayName = "Group"))
+	UPROPERTY(EditAnywhere, EditFixedSize, BlueprintGetter = GetHairGroupsRendering, BlueprintSetter = SetHairGroupsRendering, Category = "HairRendering", meta = (DisplayName = "Group", TitleProperty = "{GroupName}"))
 	TArray<FHairGroupsRendering> HairGroupsRendering;
 
 	UE_DEPRECATED(5.3, "Please do not access this member directly; use UGroomAsset accessor.")
-	UPROPERTY(EditAnywhere, EditFixedSize, BlueprintGetter = GetHairGroupsPhysics, BlueprintSetter = SetHairGroupsPhysics, Category = "HairPhysics", meta = (DisplayName = "Group"))
+	UPROPERTY(EditAnywhere, EditFixedSize, BlueprintGetter = GetHairGroupsPhysics, BlueprintSetter = SetHairGroupsPhysics, Category = "HairPhysics", meta = (DisplayName = "Group", TitleProperty = "{GroupName}"))
 	TArray<FHairGroupsPhysics> HairGroupsPhysics;
 
 	UE_DEPRECATED(5.3, "Please do not access this member directly; use UGroomAsset accessor.")
-	UPROPERTY(EditAnywhere, EditFixedSize, BlueprintGetter = GetHairGroupsInterpolation, BlueprintSetter = SetHairGroupsInterpolation, Category = "HairInterpolation", meta = (DisplayName = "Group"))
+	UPROPERTY(EditAnywhere, EditFixedSize, BlueprintGetter = GetHairGroupsInterpolation, BlueprintSetter = SetHairGroupsInterpolation, Category = "HairInterpolation", meta = (DisplayName = "Group", TitleProperty = "{GroupName}"))
 	TArray<FHairGroupsInterpolation> HairGroupsInterpolation;
 
 	UE_DEPRECATED(5.3, "Please do not access this member directly; use UGroomAsset accessor.")
@@ -681,6 +681,11 @@ public:
 	UE_API virtual void ReadDataflowContent(const TObjectPtr<UDataflowBaseContent>& DataflowContent) override;
 	//~ End IDataflowContentOwner interface
 
+	/** Begin IDataflowInstanceInterface Interface */
+	UE_API virtual const FDataflowInstance& GetDataflowInstance() const override { return DataflowSettings; }
+	UE_API virtual FDataflowInstance& GetDataflowInstance() override { return DataflowSettings; }
+	/** End UAssetDefinition Interface */
+
 #if WITH_EDITOR
 	FOnGroomAssetChanged& GetOnGroomAssetChanged() { return OnGroomAssetChanged;  }
 	FOnGroomAssetResourcesChanged& GetOnGroomAssetResourcesChanged() { return OnGroomAssetResourcesChanged; }
diff --git a/Engine/Plugins/Runtime/HairStrands/Source/HairStrandsCore/Public/GroomAssetDataflow.h b/Engine/Plugins/Runtime/HairStrands/Source/HairStrandsCore/Public/GroomAssetDataflow.h
--- a/Engine/Plugins/Runtime/HairStrands/Source/HairStrandsCore/Public/GroomAssetDataflow.h
+++ b/Engine/Plugins/Runtime/HairStrands/Source/HairStrandsCore/Public/GroomAssetDataflow.h
@@ -3,6 +3,7 @@
 #pragma once
 
 #include "Dataflow/DataflowContent.h"
+#include "Dataflow/DataflowInstance.h"
 #include "UObject/Object.h"
 #include "UObject/ObjectMacros.h"
 #include "GroomAssetDataflow.generated.h"
@@ -46,11 +47,11 @@ protected:
 };
 
 USTRUCT()
-struct FGroomDataflowSettings
+struct FGroomDataflowSettings : public FDataflowInstance
 {
 	GENERATED_BODY()
 
-	UE_API FGroomDataflowSettings();
+	UE_API explicit FGroomDataflowSettings(UObject* InOwner = nullptr, UDataflow* InDataflowAsset = nullptr, FName InTerminalNodeName = NAME_None);
 
 	/** Return the Dataflow asset name. */
 	static UE_API FName GetDataflowAssetMemberName();
@@ -58,12 +59,6 @@ struct FGroomDataflowSettings
 	/** Return the Dataflow terminal name. */
 	static UE_API FName GetDataflowTerminalMemberName();
 
-	/** Return the Dataflow asset associated to this groom asset if any, const version. */
-	UDataflow* GetDataflowAsset() const { return DataflowAsset; }
-	
-	/** Return the Dataflow terminal associated to this groom asset if any. */
-	const FString& GetDataflowTerminal() const { return DataflowTerminal; }
-
 	/** Return the skeletal mesh associated to this groom asset if any. */
 	USkeletalMesh* GetSkeletalMesh(const int32 GroupIndex) const { return !SkeletalMeshes.IsValidIndex(GroupIndex) ? nullptr : SkeletalMeshes[GroupIndex]; }
 
@@ -73,12 +68,6 @@ struct FGroomDataflowSettings
 	/** Return the groom asset dataflow settings rest collection . */
 	const FManagedArrayCollection* GetRestCollection() const { return RestCollection.Get(); }
 	
-	/** Set the Dataflow asset associated to this groom asset if any */
-	void SetDataflowAsset(UDataflow* InDataflowAsset) { DataflowAsset = InDataflowAsset; }
-	
-	/** Set the Dataflow terminal associated to this groom asset if any. */
-	void SetDataflowTerminal(const FString& InDataflowTerminal) { DataflowTerminal = InDataflowTerminal; }
-	
 	/** Init the skeletal mesh associated to this groom asset if any. */
 	void InitSkeletalMeshes(const int32 NumGroups) { SkeletalMeshes.Init(nullptr, NumGroups); MeshLODs.Init(INDEX_NONE, NumGroups); }
 
@@ -116,19 +105,11 @@ struct FGroomDataflowSettings
 
 private :
 
-	/** Dataflow asset. */
-    UPROPERTY(EditAnywhere, Category = "HairDataflow")
-    TObjectPtr<UDataflow> DataflowAsset;
-
-    /** Dataflow terminal. */
-    UPROPERTY(EditAnywhere, Category = "HairDataflow")
-    FString DataflowTerminal = "GroomAssetTerminal";
-
     /** Hair geometry that could describe external cards and meshes. */
     UPROPERTY()
     TArray<TObjectPtr<USkeletalMesh>> SkeletalMeshes;
 
-	/** LOD indices of the hair geometry used to transfer the skinnign weights */
+	/** LOD indices of the hair geometry used to transfer the skinning weights */
 	UPROPERTY()
 	TArray<int32> MeshLODs;
 
diff --git a/Engine/Plugins/Runtime/HairStrands/Source/HairStrandsCore/Public/GroomAssetInterpolation.h b/Engine/Plugins/Runtime/HairStrands/Source/HairStrandsCore/Public/GroomAssetInterpolation.h
--- a/Engine/Plugins/Runtime/HairStrands/Source/HairStrandsCore/Public/GroomAssetInterpolation.h
+++ b/Engine/Plugins/Runtime/HairStrands/Source/HairStrandsCore/Public/GroomAssetInterpolation.h
@@ -212,6 +212,12 @@ struct FHairGroupsInterpolation
 	UPROPERTY()
 	FHairDeformationSettings RiggingSettings;
 
+#if WITH_EDITORONLY_DATA
+	/** Group name to be displayed in the array list */
+	UPROPERTY(VisibleDefaultsOnly, Transient, Category=AlwaysHidden, Meta=(EditCondition=False, EditConditionHides))
+	FName GroupName = NAME_None;
+#endif
+
 	UE_API bool operator==(const FHairGroupsInterpolation& A) const;
 
 	UE_API void BuildDDCKey(FArchive& Ar);
diff --git a/Engine/Plugins/Runtime/HairStrands/Source/HairStrandsCore/Public/GroomAssetPhysics.h b/Engine/Plugins/Runtime/HairStrands/Source/HairStrandsCore/Public/GroomAssetPhysics.h
--- a/Engine/Plugins/Runtime/HairStrands/Source/HairStrandsCore/Public/GroomAssetPhysics.h
+++ b/Engine/Plugins/Runtime/HairStrands/Source/HairStrandsCore/Public/GroomAssetPhysics.h
@@ -286,6 +286,12 @@ struct FHairGroupsPhysics
 	UPROPERTY(EditAnywhere, Category = "GroupsPhysics", meta = (ToolTip = "Strands Parameters for the hair physics"))
 	FHairStrandsParameters StrandsParameters;
 
+#if WITH_EDITORONLY_DATA
+	/** Group name to be displayed in the array list */
+	UPROPERTY(VisibleDefaultsOnly, Transient, Category=AlwaysHidden, Meta=(EditCondition=False, EditConditionHides))
+	FName GroupName = NAME_None;
+#endif
+
 	UE_API bool operator==(const FHairGroupsPhysics& A) const;
 
 	/** Build DDC key */
diff --git a/Engine/Plugins/Runtime/HairStrands/Source/HairStrandsCore/Public/GroomAssetRendering.h b/Engine/Plugins/Runtime/HairStrands/Source/HairStrandsCore/Public/GroomAssetRendering.h
--- a/Engine/Plugins/Runtime/HairStrands/Source/HairStrandsCore/Public/GroomAssetRendering.h
+++ b/Engine/Plugins/Runtime/HairStrands/Source/HairStrandsCore/Public/GroomAssetRendering.h
@@ -104,6 +104,12 @@ struct FHairGroupsRendering
 	UPROPERTY(EditAnywhere, Category = "MiscSettings", meta = (ToolTip = "Advanced rendering settings "))
 	FHairAdvancedRenderingSettings AdvancedSettings;
 
+#if WITH_EDITORONLY_DATA
+	/** Group name to be displayed in the array list */
+	UPROPERTY(VisibleDefaultsOnly, Transient, Category=AlwaysHidden, Meta=(EditCondition=False, EditConditionHides))
+	FName GroupName = NAME_None;
+#endif
+
 	UE_API bool operator==(const FHairGroupsRendering& A) const;
 
 };
diff --git a/Engine/Plugins/Runtime/HairStrands/Source/HairStrandsEditor/HairStrandsEditor.Build.cs b/Engine/Plugins/Runtime/HairStrands/Source/HairStrandsEditor/HairStrandsEditor.Build.cs
--- a/Engine/Plugins/Runtime/HairStrands/Source/HairStrandsEditor/HairStrandsEditor.Build.cs
+++ b/Engine/Plugins/Runtime/HairStrands/Source/HairStrandsEditor/HairStrandsEditor.Build.cs
@@ -15,6 +15,8 @@ namespace UnrealBuildTool.Rules
 					"Engine",
 					"GeometryCache",
 					"HairStrandsCore",
+					"DataflowCore",
+					"DataflowEngine",
 					"InputCore",
 					"MainFrame",
 					"Slate",
@@ -43,6 +45,7 @@ namespace UnrealBuildTool.Rules
 					"JsonUtilities",
 					"Json",
 					"ApplicationCore",
+					"ToolWidgets",
 				});
 
 			PrivateIncludePathModuleNames.AddRange(
