Groom construction from dataflow

#virtualized

diff --git a/Engine/Plugins/Experimental/HairCardGenerator/Source/HairCardGeneratorDataflow/Private/HairCardGeneratorDataflowModule.cpp b/Engine/Plugins/Experimental/HairCardGenerator/Source/HairCardGeneratorDataflow/Private/HairCardGeneratorDataflowModule.cpp
--- a/Engine/Plugins/Experimental/HairCardGenerator/Source/HairCardGeneratorDataflow/Private/HairCardGeneratorDataflowModule.cpp
+++ b/Engine/Plugins/Experimental/HairCardGenerator/Source/HairCardGeneratorDataflow/Private/HairCardGeneratorDataflowModule.cpp
@@ -8,6 +8,7 @@
 #include "CardsAssetTerminalNode.h"
 #include "HairCardDataflowRendering.h"
 #include "Modules/ModuleManager.h"
+#include "HairStrandsFactory.h"
 
 #define LOCTEXT_NAMESPACE "HairCardGeneratorDataflow"
 
@@ -20,11 +21,14 @@ void FHairCardGeneratorDataflowModule::StartupModule()
 	DATAFLOW_NODE_REGISTER_CREATION_FACTORY(FCardsAssetTerminalNode);
 
 	UE::CardGen::Private::RegisterRenderingCallbacks();
+	UE::Groom::RegisterGroomDataflowTemplate({ "SelectCardsDataflowTemplate", "Cards Dataflow", "Add a Dataflow with a cards generation.", "/HairCardGenerator/Dataflow/DF_GroomCardsTemplate.DF_GroomCardsTemplate", false });
 }
 
 void FHairCardGeneratorDataflowModule::ShutdownModule()
 {
 	UE::CardGen::Private::DeregisterRenderingCallbacks();
+
+	UE::Groom::UnregisterGroomDataflowTemplate("SelectCardsDataflowTemplate");
 }
 
 IMPLEMENT_MODULE(FHairCardGeneratorDataflowModule, HairCardGeneratorDataflow)
diff --git a/Engine/Plugins/Runtime/HairStrands/Source/HairStrandsEditor/Private/GroomFactory.cpp b/Engine/Plugins/Runtime/HairStrands/Source/HairStrandsEditor/Private/GroomFactory.cpp
--- a/Engine/Plugins/Runtime/HairStrands/Source/HairStrandsEditor/Private/GroomFactory.cpp
+++ b/Engine/Plugins/Runtime/HairStrands/Source/HairStrandsEditor/Private/GroomFactory.cpp
@@ -2,6 +2,7 @@
 
 #include "GroomFactory.h"
 #include "GroomAsset.h"
+#include "HairStrandsFactory.h"
 
 #include UE_INLINE_GENERATED_CPP_BY_NAME(GroomFactory)
 
@@ -32,6 +33,13 @@ UObject* UGroomFactory::FactoryCreateNew(UClass* InClass, UObject* InParent, FNa
 	Groom->UpdateCachedSettings();
 	Groom->InitResources();
 
+	UE::Groom::BuildGroomDataflowAsset(Groom);
+
 	return Groom;
 }
 
+FString UGroomFactory::GetDefaultNewAssetName() const
+{
+	return FString(TEXT("GA_NewGroomAsset"));
+}
+
diff --git a/Engine/Plugins/Runtime/HairStrands/Source/HairStrandsEditor/Private/HairStrandsFactory.cpp b/Engine/Plugins/Runtime/HairStrands/Source/HairStrandsEditor/Private/HairStrandsFactory.cpp
--- a/Engine/Plugins/Runtime/HairStrands/Source/HairStrandsEditor/Private/HairStrandsFactory.cpp
+++ b/Engine/Plugins/Runtime/HairStrands/Source/HairStrandsEditor/Private/HairStrandsFactory.cpp
@@ -13,18 +13,134 @@
 #include "GroomCacheImportOptions.h"
 #include "GroomBuilder.h"
 #include "GroomImportOptionsWindow.h"
+#include "Dataflow/DataflowObject.h"
+#include "Dialog/SMessageDialog.h"
 #include "HairDescription.h"
 #include "HairStrandsEditor.h"
 #include "HairStrandsImporter.h"
 #include "HairStrandsTranslator.h"
 #include "Misc/Paths.h"
 #include "Misc/ScopedSlowTask.h"
+#include "Misc/PackageName.h"
 #include "Modules/ModuleManager.h"
+#include "AssetRegistry/AssetRegistryModule.h"
 
 #include UE_INLINE_GENERATED_CPP_BY_NAME(HairStrandsFactory)
 
 #define LOCTEXT_NAMESPACE "HairStrandsFactory"
 
+namespace UE::Groom
+{ 
+	struct FGroomDataflowTemplateManager
+	{
+		/** Singleton dataflow template manager*/
+		static FGroomDataflowTemplateManager& GetInstance()
+		{
+			static FGroomDataflowTemplateManager TemplateManager;
+			return TemplateManager;
+		}
+
+		/** Build dataflow template buttons from data*/
+		TArray<SMessageDialog::FButton> BuildDataflowTemplateButtons() const
+		{
+			TArray<SMessageDialog::FButton> DataflowButtons;
+			DataflowButtons.Reserve(TemplateData.Num());
+
+			for(int32 Index = 0; Index < 10; ++Index)
+			{ 
+			for(const FGroomDataflowTemplateData& DataflowTemplate : TemplateData)
+			{
+				const FString DataflowTitleID = DataflowTemplate.TemplateName + FString("Title");
+				const FString DataflowTooltipID = DataflowTemplate.TemplateName + FString("Tooltip");
+				DataflowButtons.Add(SMessageDialog::FButton(FText::AsLocalizable_Advanced(TEXT(LOCTEXT_NAMESPACE), *DataflowTitleID, *DataflowTemplate.TemplateTitle))
+					.SetToolTipText(FText::AsLocalizable_Advanced(TEXT(LOCTEXT_NAMESPACE), *DataflowTooltipID, *DataflowTemplate.TemplateTooltip))
+					.SetPrimary(DataflowTemplate.bIsPrimaryTemplate));
+			}
+			}
+			return DataflowButtons;
+		}
+
+		/** Get the template path given an index */
+		FString GetDataflowTemplatePath(const int32 TemplateIndex) const
+		{
+			return TemplateData.IsValidIndex(TemplateIndex) ? TemplateData[TemplateIndex].TemplatePath : "";
+		}
+
+		/** Add dataflow template data to the manager */
+		void AddDataflowTemplateData(const FGroomDataflowTemplateData& DataflowTemplate)
+		{
+			TemplateData.Add(DataflowTemplate);
+		}
+
+		/** Add dataflow template data to the manager */
+		void RemoveDataflowTemplateData(const FString& TemplateName)
+		{
+			for(int32 TemplateIndex = TemplateData.Num()-1; TemplateIndex >= 0; --TemplateIndex)
+			{ 
+				if (TemplateData[TemplateIndex].TemplateName == TemplateName)
+				{
+					TemplateData.RemoveAt(TemplateIndex);
+				}
+			}
+		}
+
+	private:
+
+		/** List of dataflow template data*/
+		TArray<FGroomDataflowTemplateData>	TemplateData;
+	};
+
+	void RegisterGroomDataflowTemplate(const FGroomDataflowTemplateData& DataflowTemplate)
+	{
+		FGroomDataflowTemplateManager::GetInstance().AddDataflowTemplateData(DataflowTemplate);
+	}
+
+	void UnregisterGroomDataflowTemplate(const FString& TemplateName)
+	{
+		FGroomDataflowTemplateManager::GetInstance().RemoveDataflowTemplateData(TemplateName);
+	}
+
+	void BuildGroomDataflowAsset(UGroomAsset* GroomAsset)
+	{
+		// Create a new Dataflow asset
+		const FString DataflowPath = FPackageName::GetLongPackagePath(GroomAsset->GetOutermost()->GetName());
+		const FString GroomAssetName = GroomAsset->GetName();
+		FString DataflowName = FString(TEXT("DF_")) + (GroomAssetName.StartsWith(TEXT("GA_")) ? GroomAssetName.RightChop(3) : GroomAssetName);
+		FString DataflowPackageName = FPaths::Combine(DataflowPath, DataflowName);
+		if (FindPackage(nullptr, *DataflowPackageName))
+		{
+			// If a Dataflow asset already exists with this name, make a unique name from it to avoid clobbering it
+			MakeUniqueObjectName(nullptr, UPackage::StaticClass(), FName(DataflowPackageName)).ToString(DataflowPackageName);
+			DataflowName = FPaths::GetBaseFilename(DataflowPackageName);
+		}
+		UPackage* const DataflowPackage = CreatePackage(*DataflowPackageName);
+
+		// Select the template
+		TSharedRef<SMessageDialog> SelectTemplateMessageDialog = SNew(SMessageDialog)
+			.Title(FText(LOCTEXT("SelectTemplateTitle", "Select a Groom Dataflow Template")))
+			.Message(LOCTEXT("SelectTemplateMessage", "Select a template for this groom asset:"))
+			.Buttons(
+				{ 
+					FGroomDataflowTemplateManager::GetInstance().BuildDataflowTemplateButtons()
+				});
+
+		const FString GroomDataflowTemplate = FGroomDataflowTemplateManager::GetInstance().GetDataflowTemplatePath(SelectTemplateMessageDialog->ShowModal());
+
+		// Load the dataflow template into the groom asset
+		UDataflow* const Template = GroomDataflowTemplate.IsEmpty() ? nullptr : LoadObject<UDataflow>(DataflowPackage, *GroomDataflowTemplate);
+		if (UDataflow* const Dataflow = Template ? DuplicateObject(Template, DataflowPackage, FName(DataflowName)) : nullptr)
+		{
+			Dataflow->MarkPackageDirty();
+
+			// Notify the asset registry
+			FAssetRegistryModule::AssetCreated(Dataflow);
+
+			// Set the Dataflow to the groom asset
+			GroomAsset->GetDataflowSettings().SetDataflowAsset(Dataflow);
+		}
+	}
+}
+
 UHairStrandsFactory::UHairStrandsFactory(const FObjectInitializer& ObjectInitializer)
 	: Super(ObjectInitializer)
 {
@@ -267,6 +383,11 @@ UObject* UHairStrandsFactory::FactoryCreateFile(UClass* InClass, UObject* InPare
 		{
 			CurrentAsset = GroomCaches[0];
 		}
+	} 
+	
+	if(CurrentAsset)
+	{
+		UE::Groom::BuildGroomDataflowAsset(Cast<UGroomAsset>(CurrentAsset));
 	}
 
 	return CurrentAsset;
diff --git a/Engine/Plugins/Runtime/HairStrands/Source/HairStrandsEditor/Public/GroomActions.h b/Engine/Plugins/Runtime/HairStrands/Source/HairStrandsEditor/Public/GroomActions.h
--- a/Engine/Plugins/Runtime/HairStrands/Source/HairStrandsEditor/Public/GroomActions.h
+++ b/Engine/Plugins/Runtime/HairStrands/Source/HairStrandsEditor/Public/GroomActions.h
@@ -15,13 +15,13 @@ class UAssetDefinition_GroomAsset : public UAssetDefinitionDefault
 
 public:
 	// UAssetDefinition Begin
-	virtual FText GetAssetDisplayName() const override { return NSLOCTEXT("AssetTypeActions", "AssetTypeActions_Groom", "Groom"); }
+	virtual FText GetAssetDisplayName() const override { return NSLOCTEXT("AssetTypeActions", "AssetTypeActions_Groom", "Groom Asset"); }
 	virtual FLinearColor GetAssetColor() const override;
 	virtual TSoftClassPtr<UObject> GetAssetClass() const override { return UGroomAsset::StaticClass(); }
 	virtual bool CanImport() const override { return true; }
 	virtual TConstArrayView<FAssetCategoryPath> GetAssetCategories() const override
 	{
-		static const auto Categories = { EAssetCategoryPaths::Misc };
+		static const auto Categories = { EAssetCategoryPaths::Physics };
 		return Categories;
 	}
 	virtual UThumbnailInfo* LoadThumbnailInfo(const FAssetData& InAssetData) const override;
diff --git a/Engine/Plugins/Runtime/HairStrands/Source/HairStrandsEditor/Public/GroomFactory.h b/Engine/Plugins/Runtime/HairStrands/Source/HairStrandsEditor/Public/GroomFactory.h
--- a/Engine/Plugins/Runtime/HairStrands/Source/HairStrandsEditor/Public/GroomFactory.h
+++ b/Engine/Plugins/Runtime/HairStrands/Source/HairStrandsEditor/Public/GroomFactory.h
@@ -18,6 +18,7 @@ public:
 
 	//~ Begin UFactory Interface
 	virtual UObject* FactoryCreateNew(UClass* InClass, UObject* InParent, FName InName, EObjectFlags Flags, UObject* Context, FFeedbackContext* Warn) override;
+	virtual FString GetDefaultNewAssetName() const override;
 	//~ End UFactory Interface
 };
 
diff --git a/Engine/Plugins/Runtime/HairStrands/Source/HairStrandsEditor/Public/HairStrandsFactory.h b/Engine/Plugins/Runtime/HairStrands/Source/HairStrandsEditor/Public/HairStrandsFactory.h
--- a/Engine/Plugins/Runtime/HairStrands/Source/HairStrandsEditor/Public/HairStrandsFactory.h
+++ b/Engine/Plugins/Runtime/HairStrands/Source/HairStrandsEditor/Public/HairStrandsFactory.h
@@ -8,6 +8,38 @@
 #include "HairStrandsFactory.generated.h"
 
 class IGroomTranslator;
+class UGroomAsset;
+
+namespace UE::Groom
+{
+	HAIRSTRANDSEDITOR_API struct FGroomDataflowTemplateData
+	{
+		/** Dataflow template name */
+		FString TemplateName;
+
+		/** Dataflow template title */
+		FString TemplateTitle;
+
+		/** Dataflow template tooltip */
+		FString TemplateTooltip;
+
+		/** Dataflow template asset path */
+		FString TemplatePath;
+
+		/** Template primary flag */
+		bool bIsPrimaryTemplate = false;
+	};
+
+	/** Register dataflow template to the manager */
+	HAIRSTRANDSEDITOR_API void RegisterGroomDataflowTemplate(const FGroomDataflowTemplateData& DataflowTemplate);
+
+
+	/** Unregister dataflow template to the manager */
+	HAIRSTRANDSEDITOR_API void UnregisterGroomDataflowTemplate(const FString& TemplateName);
+
+	/** Build groom dataflow asset from templates */
+	HAIRSTRANDSEDITOR_API void BuildGroomDataflowAsset(UGroomAsset* GroomAsset);
+}
 
 /**
  * Implements a factory for UHairStrands objects.
@@ -15,16 +47,19 @@ class IGroomTranslator;
 UCLASS(hidecategories = Object)
 class UHairStrandsFactory : public UFactory
 {
-	GENERATED_UCLASS_BODY()
+	GENERATED_BODY()
 
 public:
 
-	//~ UFactory Interface
+	UHairStrandsFactory(const FObjectInitializer& ObjectInitializer);
+
+	//~ Begin UFactory interface
 	virtual UObject* FactoryCreateFile(UClass* InClass, UObject* InParent, FName InName, EObjectFlags Flags,
 		const FString& Filename, const TCHAR* Parms, FFeedbackContext* Warn, bool& bOutOperationCanceled) override;
 	virtual bool FactoryCanImport(const FString& Filename) override;
 	virtual void GetSupportedFileExtensions(TArray<FString>& OutExtensions) const override;
 	virtual void CleanUp() override;
+	//~ End UFactory interface
 
 protected:
 	void InitTranslators();
