Fix CIS in hair strands factory visibility

diff --git a/Engine/Plugins/Runtime/HairStrands/Source/HairStrandsEditor/Public/HairStrandsFactory.h b/Engine/Plugins/Runtime/HairStrands/Source/HairStrandsEditor/Public/HairStrandsFactory.h
--- a/Engine/Plugins/Runtime/HairStrands/Source/HairStrandsEditor/Public/HairStrandsFactory.h
+++ b/Engine/Plugins/Runtime/HairStrands/Source/HairStrandsEditor/Public/HairStrandsFactory.h
@@ -12,7 +12,7 @@ class UGroomAsset;
 
 namespace UE::Groom
 {
-	HAIRSTRANDSEDITOR_API struct FGroomDataflowTemplateData
+	struct HAIRSTRANDSEDITOR_API FGroomDataflowTemplateData
 	{
 		/** Dataflow template name */
 		FString TemplateName;
