Fix groom factory CIS

diff --git a/Engine/Plugins/Runtime/HairStrands/Source/HairStrandsEditor/Private/HairStrandsFactory.cpp b/Engine/Plugins/Runtime/HairStrands/Source/HairStrandsEditor/Private/HairStrandsFactory.cpp
--- a/Engine/Plugins/Runtime/HairStrands/Source/HairStrandsEditor/Private/HairStrandsFactory.cpp
+++ b/Engine/Plugins/Runtime/HairStrands/Source/HairStrandsEditor/Private/HairStrandsFactory.cpp
@@ -45,9 +45,7 @@ namespace UE::Groom
 		{
 			TArray<SMessageDialog::FButton> DataflowButtons;
 			DataflowButtons.Reserve(TemplateData.Num());
-
-			for(int32 Index = 0; Index < 10; ++Index)
-			{ 
+			
 			for(const FGroomDataflowTemplateData& DataflowTemplate : TemplateData)
 			{
 				const FString DataflowTitleID = DataflowTemplate.TemplateName + FString("Title");
@@ -56,7 +54,6 @@ namespace UE::Groom
 					.SetToolTipText(FText::AsLocalizable_Advanced(TEXT(LOCTEXT_NAMESPACE), *DataflowTooltipID, *DataflowTemplate.TemplateTooltip))
 					.SetPrimary(DataflowTemplate.bIsPrimaryTemplate));
 			}
-			}
 			return DataflowButtons;
 		}
 
@@ -110,7 +107,7 @@ namespace UE::Groom
 		if (FindPackage(nullptr, *DataflowPackageName))
 		{
 			// If a Dataflow asset already exists with this name, make a unique name from it to avoid clobbering it
-			MakeUniqueObjectName(nullptr, UPackage::StaticClass(), FName(DataflowPackageName)).ToString(DataflowPackageName);
+			DataflowPackageName = MakeUniqueObjectName(nullptr, UPackage::StaticClass(), FName(DataflowPackageName)).ToString();
 			DataflowName = FPaths::GetBaseFilename(DataflowPackageName);
 		}
 		UPackage* const DataflowPackage = CreatePackage(*DataflowPackageName);
@@ -127,7 +124,7 @@ namespace UE::Groom
 		const FString GroomDataflowTemplate = FGroomDataflowTemplateManager::GetInstance().GetDataflowTemplatePath(SelectTemplateMessageDialog->ShowModal());
 
 		// Load the dataflow template into the groom asset
-		UDataflow* const Template = GroomDataflowTemplate.IsEmpty() ? nullptr : LoadObject<UDataflow>(DataflowPackage, *GroomDataflowTemplate);
+		UDataflow* const Template = GroomDataflowTemplate.IsEmpty() ? nullptr : LoadObject<UDataflow>(nullptr, *GroomDataflowTemplate);
 		if (UDataflow* const Dataflow = Template ? DuplicateObject(Template, DataflowPackage, FName(DataflowName)) : nullptr)
 		{
 			Dataflow->MarkPackageDirty();
@@ -385,9 +382,9 @@ UObject* UHairStrandsFactory::FactoryCreateFile(UClass* InClass, UObject* InPare
 		}
 	} 
 	
-	if(CurrentAsset)
+	if(UGroomAsset* GroomAsset = Cast<UGroomAsset>(CurrentAsset))
 	{
-		UE::Groom::BuildGroomDataflowAsset(Cast<UGroomAsset>(CurrentAsset));
+		UE::Groom::BuildGroomDataflowAsset(GroomAsset);
 	}
 
 	return CurrentAsset;
