Fix rig groom generation to have invalide offset, causing all guides to generate the same bone chain.

diff --git a/Engine/Plugins/Runtime/HairStrands/Source/HairStrandsCore/Private/GroomBuilder.cpp b/Engine/Plugins/Runtime/HairStrands/Source/HairStrandsCore/Private/GroomBuilder.cpp
--- a/Engine/Plugins/Runtime/HairStrands/Source/HairStrandsCore/Private/GroomBuilder.cpp
+++ b/Engine/Plugins/Runtime/HairStrands/Source/HairStrandsCore/Private/GroomBuilder.cpp
@@ -55,7 +55,7 @@ bool DoesHairStrandsSupportCompressedPosition();
 
 FString FGroomBuilder::GetVersion()
 {
-	FString Out = TEXT("v20_debug4");
+	FString Out = TEXT("v21");
 	if (IsRenderCurveEnabled())
 	{
 		Out += TEXT("_Curve4a");
@@ -2814,7 +2814,7 @@ void Decimate(
 			FHairStrandsDatas::CopyPointLerp(InData, OutData, InAttributes, PointBegin, PointEnd, PointAlpha, PointCount);
 		}
 	}
-	HairStrandsBuilder::BuildInternalData(OutData);
+	OutData = HairStrandsBuilder::BuildInternalData(OutData);
 }
 
 void GetCurveShuffleIndices(TArray<uint32>& CurveRandomizedIndex, uint32 InCurveCount)
