Fix CIS issue.

diff --git a/Engine/Plugins/Runtime/HairStrands/Source/HairStrandsCore/Private/HairCardsBuilder.cpp b/Engine/Plugins/Runtime/HairStrands/Source/HairStrandsCore/Private/HairCardsBuilder.cpp
--- a/Engine/Plugins/Runtime/HairStrands/Source/HairStrandsCore/Private/HairCardsBuilder.cpp
+++ b/Engine/Plugins/Runtime/HairStrands/Source/HairStrandsCore/Private/HairCardsBuilder.cpp
@@ -71,6 +71,7 @@ FString GetVersion()
 
 #define HAIR_CARD_UV_ROOTUV 0
 #define HAIR_CARD_UV_CARDUV 1
+#define HAIR_CARD_UV_COUNT  2
 
 // Example of layout with 3 texture coords
 //  -----------------------------------------------  -----------------------------------------------
@@ -1263,8 +1264,9 @@ void BuildGeometry(
 	Out.Meshes.Normals.SetNum(TotalPointCount);
 	Out.Meshes.Tangents.SetNum(TotalPointCount);
 	Out.Meshes.Indices.SetNum(TotalIndexCount);
-	Out.Meshes.UVs.SetNum(1);
-	Out.Meshes.UVs[0].UV.SetNum(TotalPointCount);
+	Out.Meshes.UVs.SetNum(HAIR_CARD_UV_COUNT);
+	Out.Meshes.UVs[HAIR_CARD_UV_CARDUV].UV.SetNum(TotalPointCount);
+	Out.Meshes.UVs[HAIR_CARD_UV_ROOTUV].UV.Init(FVector2f(0,0), TotalPointCount);
 
 	Out.Meshes.Positions[0] = Center + FVector3f(-Extent.X, -Extent.Y, -Extent.Z);
 	Out.Meshes.Positions[1] = Center + FVector3f(+Extent.X, -Extent.Y, -Extent.Z);
