Fix hair card importer when using imported guides.

The UV channel count was incorrect, causing out of memory fetches.

diff --git a/Engine/Plugins/Runtime/HairStrands/Source/HairStrandsCore/Private/HairCardsBuilder.cpp b/Engine/Plugins/Runtime/HairStrands/Source/HairStrandsCore/Private/HairCardsBuilder.cpp
--- a/Engine/Plugins/Runtime/HairStrands/Source/HairStrandsCore/Private/HairCardsBuilder.cpp
+++ b/Engine/Plugins/Runtime/HairStrands/Source/HairStrandsCore/Private/HairCardsBuilder.cpp
@@ -975,8 +975,8 @@ static bool InternalImportGeometry_WithImportedGuides(
 	const TVertexInstanceAttributesRef<const FVector2f> VertexInstanceUVs		= MeshDescription->VertexInstanceAttributes().GetAttributesRef<FVector2f>(MeshAttribute::VertexInstance::TextureCoordinate);
 	const TVertexInstanceAttributesRef<const FVector4f> VertexColors			= MeshDescription->VertexInstanceAttributes().GetAttributesRef<FVector4f>(MeshAttribute::VertexInstance::Color);
 
-	const bool bHasVertexColor = VertexColors.GetNumElements() > 0;	
-	const uint32 NumSourceUVs = StaticMesh->GetNumTexCoords(0);
+	const bool bHasVertexColor = VertexColors.GetNumElements() > 0;
+	const uint32 NumSourceUVs = VertexInstanceUVs.GetNumChannels();
 	const uint32 NumUVs = NumSourceUVs + 1u /*RootUV*/;
 
 	// Write out BulkData
@@ -1015,9 +1015,7 @@ static bool InternalImportGeometry_WithImportedGuides(
 		// * UVs[1-N] -> Card UV & else
 		for (uint32 UVIndex=0;UVIndex<NumSourceUVs;++UVIndex)
 		{
-			FVector2f UV = FVector2f::ZeroVector;
-			UV.X = VertexInstanceUVs[VertexInstanceId0].Component(UVIndex*2 + 0);
-			UV.Y = VertexInstanceUVs[VertexInstanceId0].Component(UVIndex*2 + 1);
+			const FVector2f UV = VertexInstanceUVs.Get(VertexInstanceId0, UVIndex);
 
 			SetPackedUV(VertexIndex, HAIR_CARD_UV_CARDUV + UVIndex, NumUVs, UV, OutBulk.UVs);
 		}
