Fix potential crash in GetGroomAsset node

diff --git a/Engine/Plugins/Runtime/HairStrands/Source/HairStrandsCore/Private/GroomEdit.cpp b/Engine/Plugins/Runtime/HairStrands/Source/HairStrandsCore/Private/GroomEdit.cpp
--- a/Engine/Plugins/Runtime/HairStrands/Source/HairStrandsCore/Private/GroomEdit.cpp
+++ b/Engine/Plugins/Runtime/HairStrands/Source/HairStrandsCore/Private/GroomEdit.cpp
@@ -222,17 +222,20 @@ void ConvertFromGroomAsset(UGroomAsset* InGroom, FEditableGroom* OutGroom, const
 
 	for (uint32 GroupIndex = 0; GroupIndex < GroupCount; ++GroupIndex)
 	{
-		FHairGroupInfo GroupInfo;
-		FHairStrandsDatas StrandsData;
-		FHairStrandsDatas GuidesData;
-		FGroomBuilder::BuildData(HairDescriptionGroups.HairGroups[GroupIndex], In.GetHairGroupsInterpolation()[GroupIndex], GroupInfo, StrandsData, GuidesData, bAllowCurveReordering, bApplyDecimation);
+		if (HairDescriptionGroups.HairGroups.IsValidIndex(GroupIndex))
+		{
+			FHairGroupInfo GroupInfo;
+			FHairStrandsDatas StrandsData;
+			FHairStrandsDatas GuidesData;
+			FGroomBuilder::BuildData(HairDescriptionGroups.HairGroups[GroupIndex], In.GetHairGroupsInterpolation()[GroupIndex], GroupInfo, StrandsData, GuidesData, bAllowCurveReordering, bApplyDecimation);
 
-		Out.Groups[GroupIndex].GroupIndex = GroupIndex;
-		Out.Groups[GroupIndex].GroupID = GroupInfo.GroupID;
-		Out.Groups[GroupIndex].GroupName = GroupInfo.GroupName;
-		HairStrandsDataToEditableHairGuides(GuidesData, Out.Groups[GroupIndex].Guides);
+			Out.Groups[GroupIndex].GroupIndex = GroupIndex;
+			Out.Groups[GroupIndex].GroupID = GroupInfo.GroupID;
+			Out.Groups[GroupIndex].GroupName = GroupInfo.GroupName;
+			HairStrandsDataToEditableHairGuides(GuidesData, Out.Groups[GroupIndex].Guides);
 
-		HairStrandsDataToEditableHairStrands(StrandsData, Out.Groups[GroupIndex].Strands);
+			HairStrandsDataToEditableHairStrands(StrandsData, Out.Groups[GroupIndex].Strands);
+		}
 	}
 #endif
 }
