Add default groom templates

diff --git a/Engine/Plugins/Runtime/HairStrands/Source/HairStrandsEditor/Private/HairStrandsEditor.cpp b/Engine/Plugins/Runtime/HairStrands/Source/HairStrandsEditor/Private/HairStrandsEditor.cpp
--- a/Engine/Plugins/Runtime/HairStrands/Source/HairStrandsEditor/Private/HairStrandsEditor.cpp
+++ b/Engine/Plugins/Runtime/HairStrands/Source/HairStrandsEditor/Private/HairStrandsEditor.cpp
@@ -35,6 +35,7 @@
 #include "Delegates/DelegateCombinations.h"
 #include "ThumbnailRendering/ThumbnailManager.h"
 #include "UObject/UObjectIterator.h"
+#include "HairStrandsFactory.h"
 
 IMPLEMENT_MODULE(FGroomEditor, HairStrandsEditor);
 
@@ -154,6 +155,10 @@ void FGroomEditor::StartupModule()
 
 	UThumbnailManager::Get().RegisterCustomRenderer(UGroomAsset::StaticClass(), UGroomAssetThumbnailRenderer::StaticClass());
 	UThumbnailManager::Get().RegisterCustomRenderer(UGroomBindingAsset::StaticClass(), UGroomBindingAssetThumbnailRenderer::StaticClass());
+
+	UE::Groom::RegisterGroomDataflowTemplate({"SelectNoDataflowTemplate", "No Dataflow", "Don't add a Dataflow to this Groom Asset. Useful when a Dataflow already exists.", "", false});
+	UE::Groom::RegisterGroomDataflowTemplate({"SelectEmptyDataflowTemplate", "Empty Dataflow", "Add an empty Dataflow with a Groom Terminal node.", "/HairStrands/Dataflow/DF_GroomEmptyTemplate.DF_GroomDataflowTemplate", false});
+	UE::Groom::RegisterGroomDataflowTemplate({"SelectGuidesDataflowTemplate", "Guides Dataflow", "Add a Dataflow with basic nodes guides generation.", "/HairStrands/Dataflow/DF_GroomGuidesTemplate.DF_GroomGuidesTemplate", true});
 }
 
 void FGroomEditor::ShutdownModule()
@@ -197,6 +202,10 @@ void FGroomEditor::ShutdownModule()
 		ensure(StyleSet.IsUnique());
 		StyleSet.Reset();
 	}
+
+	UE::Groom::UnregisterGroomDataflowTemplate("SelectNoDataflowTemplate");
+	UE::Groom::UnregisterGroomDataflowTemplate("SelectEmptyDataflowTemplate");
+	UE::Groom::UnregisterGroomDataflowTemplate("SelectGuidesDataflowTemplate");
 }
 
 void FGroomEditor::RegisterMenus()
