Mark cached data as mutable to allow hair strand data access from const GroomAssets

#rnx

diff --git a/Engine/Plugins/Runtime/HairStrands/Source/HairStrandsCore/Private/GroomAsset.cpp b/Engine/Plugins/Runtime/HairStrands/Source/HairStrandsCore/Private/GroomAsset.cpp
--- a/Engine/Plugins/Runtime/HairStrands/Source/HairStrandsCore/Private/GroomAsset.cpp
+++ b/Engine/Plugins/Runtime/HairStrands/Source/HairStrandsCore/Private/GroomAsset.cpp
@@ -2486,7 +2486,7 @@ void UGroomAsset::CommitHairDescription(FHairDescription&& InHairDescription, EH
 	HairDescriptionBulkData[HairDescriptionType]->SaveHairDescription(*CachedHairDescription[HairDescriptionType]);
 }
 
-const FHairDescriptionGroups& UGroomAsset::GetHairDescriptionGroups()
+const FHairDescriptionGroups& UGroomAsset::GetHairDescriptionGroups() const
 {
 	TRACE_CPUPROFILER_EVENT_SCOPE(UGroomAsset::GetHairDescriptionGroups);
 
@@ -3832,7 +3832,7 @@ bool UGroomAsset::IsMaterialUsed(int32 MaterialIndex) const
 bool UGroomAsset::GetHairStrandsDatas(
 	const int32 GroupIndex,
 	FHairStrandsDatas& OutStrandsData,
-	FHairStrandsDatas& OutGuidesData)
+	FHairStrandsDatas& OutGuidesData) const
 {
 	if (!CanRebuildFromDescription())
 	{
diff --git a/Engine/Plugins/Runtime/HairStrands/Source/HairStrandsCore/Public/GroomAsset.h b/Engine/Plugins/Runtime/HairStrands/Source/HairStrandsCore/Public/GroomAsset.h
--- a/Engine/Plugins/Runtime/HairStrands/Source/HairStrandsCore/Public/GroomAsset.h
+++ b/Engine/Plugins/Runtime/HairStrands/Source/HairStrandsCore/Public/GroomAsset.h
@@ -865,7 +865,7 @@ public:
 	UE_API FHairDescription GetHairDescription() const;
 
 	/** Get/Build render & guides data based on the hair description and interpolation settings */
-	UE_API bool GetHairStrandsDatas(const int32 GroupIndex, FHairStrandsDatas& OutStrandsData, FHairStrandsDatas& OutGuidesData);
+	UE_API bool GetHairStrandsDatas(const int32 GroupIndex, FHairStrandsDatas& OutStrandsData, FHairStrandsDatas& OutGuidesData) const;
 	UE_API bool GetHairCardsGuidesDatas(const int32 GroupIndex, const int32 LODIndex, FHairStrandsDatas& OutCardsGuidesData);
 
 	/** Caches the computed (group) groom data with the given build settings from/to the Derived Data Cache, building it if needed.
@@ -882,7 +882,7 @@ public:
 	UE_API FString GetDerivedDataKeyForStrands(uint32 GroupIndex);
 	UE_API FString GetDerivedDataKeyForMeshes(uint32 GroupIndex);
 
-	UE_API const FHairDescriptionGroups& GetHairDescriptionGroups();
+	UE_API const FHairDescriptionGroups& GetHairDescriptionGroups() const;
 private:
 	
 	bool IsFullyCached();
@@ -892,9 +892,9 @@ private:
 	// Transient HairDescription & HairDescriptionGroups, which are built from HairDescriptionBulkData.
 	// All these data (bulk/desc/groups) needs to be in sync. I.e., when the HairDescription is updated, 
 	// HairDescriptionGroups needs to also be updated
-	UE::FRecursiveMutex InternalLock;
-	TUniquePtr<FHairDescription> CachedHairDescription[EHairDescriptionType::Count];
-	TUniquePtr<FHairDescriptionGroups> CachedHairDescriptionGroups[EHairDescriptionType::Count];
+	mutable UE::FRecursiveMutex InternalLock;
+	mutable TUniquePtr<FHairDescription> CachedHairDescription[EHairDescriptionType::Count];
+	mutable TUniquePtr<FHairDescriptionGroups> CachedHairDescriptionGroups[EHairDescriptionType::Count];
 
 	TArray<FString> StrandsDerivedDataKey;
 	TArray<FString> CardsDerivedDataKey;
