Groom dataflow creation from already imported asset + reopen editors

diff --git a/Engine/Plugins/Runtime/HairStrands/Source/HairStrandsCore/Public/GroomAsset.h b/Engine/Plugins/Runtime/HairStrands/Source/HairStrandsCore/Public/GroomAsset.h
--- a/Engine/Plugins/Runtime/HairStrands/Source/HairStrandsCore/Public/GroomAsset.h
+++ b/Engine/Plugins/Runtime/HairStrands/Source/HairStrandsCore/Public/GroomAsset.h
@@ -464,6 +464,7 @@ class UGroomAsset : public UObject, public IInterface_AssetUserData, public IDat
 	DECLARE_MULTICAST_DELEGATE(FOnGroomAssetChanged);
 	DECLARE_MULTICAST_DELEGATE(FOnGroomAssetResourcesChanged);
 	DECLARE_MULTICAST_DELEGATE(FOnGroomAsyncLoadFinished);
+	DECLARE_MULTICAST_DELEGATE_OneParam(FOnCreateGroomDataflow, UGroomAsset*);
 #endif
 
 private:
@@ -478,6 +479,18 @@ private:
 	FGroomDataflowSettings DataflowSettings;
 
 public:
+	/** Return the builder category name for layout */
+	static FName GetDataflowBuilderCategoryName() { return FName("HairDataflow"); }
+
+	/** Create the dataflow from the groom asset */
+	UFUNCTION(CallInEditor, BlueprintCallable, Category = "HairDataflow")
+	void CreateGroomDataflow() 
+	{ 
+#if WITH_EDITOR
+		OnCreateGroomDataflow.Broadcast(this); 
+#endif
+	}
+	
 	UE_DEPRECATED(5.3, "Please do not access this member directly; use UGroomAsset accessor.")
 	UPROPERTY(EditAnywhere, EditFixedSize, EditFixedSize, Category = "HairInfo", meta = (DisplayName = "Group", TitleProperty = "{GroupName}"))
 	TArray<FHairGroupInfoWithVisibility> HairGroupsInfo;
@@ -690,6 +703,7 @@ public:
 	FOnGroomAssetChanged& GetOnGroomAssetChanged() { return OnGroomAssetChanged;  }
 	FOnGroomAssetResourcesChanged& GetOnGroomAssetResourcesChanged() { return OnGroomAssetResourcesChanged; }
 	FOnGroomAsyncLoadFinished& GetOnGroomAsyncLoadFinished() { return OnGroomAsyncLoadFinished; }
+	FOnCreateGroomDataflow& GetOnCreateGroomDataflow() { return OnCreateGroomDataflow; }
 
 	/**  Part of Uobject interface  */
 	UE_API virtual bool Modify(bool bAlwaysMarkDirty = true) override;
@@ -790,6 +804,7 @@ public:
 	FOnGroomAssetChanged OnGroomAssetChanged;
 	FOnGroomAssetResourcesChanged OnGroomAssetResourcesChanged;
 	FOnGroomAsyncLoadFinished OnGroomAsyncLoadFinished;
+	FOnCreateGroomDataflow OnCreateGroomDataflow;
 
 	UE_API void MarkMaterialsHasChanged();
 
diff --git a/Engine/Plugins/Runtime/HairStrands/Source/HairStrandsEditor/Private/GroomActions.cpp b/Engine/Plugins/Runtime/HairStrands/Source/HairStrandsEditor/Private/GroomActions.cpp
--- a/Engine/Plugins/Runtime/HairStrands/Source/HairStrandsEditor/Private/GroomActions.cpp
+++ b/Engine/Plugins/Runtime/HairStrands/Source/HairStrandsEditor/Private/GroomActions.cpp
@@ -37,6 +37,7 @@
 
 #include "AssetToolsModule.h"
 #include "ContentBrowserMenuContexts.h"
+#include "HairStrandsFactory.h"
 #include "ToolMenus.h"
 #include "Dataflow/DataflowEditor.h"
 
@@ -63,6 +64,22 @@ EAssetCommandResult UAssetDefinition_GroomAsset::OpenAssets(const FAssetOpenArgs
 	{
 		if (GroomAsset != nullptr)
 		{
+			if(!GroomAsset->GetOnCreateGroomDataflow().IsBound())
+			{
+				GroomAsset->GetOnCreateGroomDataflow().AddLambda(
+					[](UGroomAsset* GroomAsset)
+					{
+						UE::Groom::BuildGroomDataflowAsset(GroomAsset);
+						if (UAssetEditorSubsystem* AssetEditorSubsystem = GEditor->GetEditorSubsystem<UAssetEditorSubsystem>())
+						{
+							// Close all the opened editors
+							AssetEditorSubsystem->CloseAllEditorsForAsset(GroomAsset);
+							
+							// Open a new editor to be able to edit the groom asset with dataflow
+							AssetEditorSubsystem->OpenEditorForAsset(GroomAsset);
+						}
+					} );
+			}
 			if(!GroomAsset->GetDataflowSettings().GetDataflowAsset())
 			{
 				TSharedRef<FGroomCustomAssetEditorToolkit> NewCustomAssetEditor(new FGroomCustomAssetEditorToolkit());
diff --git a/Engine/Plugins/Runtime/HairStrands/Source/HairStrandsEditor/Private/GroomAssetDetails.cpp b/Engine/Plugins/Runtime/HairStrands/Source/HairStrandsEditor/Private/GroomAssetDetails.cpp
--- a/Engine/Plugins/Runtime/HairStrands/Source/HairStrandsEditor/Private/GroomAssetDetails.cpp
+++ b/Engine/Plugins/Runtime/HairStrands/Source/HairStrandsEditor/Private/GroomAssetDetails.cpp
@@ -713,6 +713,8 @@ void FGroomRenderingDetails::CustomizeStrandsGroupProperties(IDetailLayoutBuilde
 		DetailLayout.HideProperty(DetailLayout.GetProperty(UGroomAsset::GetLODModeMemberName(), UGroomAsset::StaticClass()));
 		DetailLayout.HideProperty(DetailLayout.GetProperty(UGroomAsset::GetAutoLODBiasMemberName(), UGroomAsset::StaticClass()));
 		DetailLayout.HideProperty(DetailLayout.GetProperty(UGroomAsset::GetDataflowSettingsMemberName(), UGroomAsset::StaticClass()));
+		DetailLayout.HideProperty(DetailLayout.GetProperty(UGroomAsset::GetDataflowSettingsMemberName(), UGroomAsset::StaticClass()));
+		DetailLayout.HideCategory(UGroomAsset::GetDataflowBuilderCategoryName());
 	}
 	break;
 	case EMaterialPanelType::Meshes:
@@ -734,6 +736,7 @@ void FGroomRenderingDetails::CustomizeStrandsGroupProperties(IDetailLayoutBuilde
 		DetailLayout.HideProperty(DetailLayout.GetProperty(UGroomAsset::GetLODModeMemberName(), UGroomAsset::StaticClass()));
 		DetailLayout.HideProperty(DetailLayout.GetProperty(UGroomAsset::GetAutoLODBiasMemberName(), UGroomAsset::StaticClass()));
 		DetailLayout.HideProperty(DetailLayout.GetProperty(UGroomAsset::GetDataflowSettingsMemberName(), UGroomAsset::StaticClass()));
+		DetailLayout.HideCategory(UGroomAsset::GetDataflowBuilderCategoryName());
 	}
 	break;
 	case EMaterialPanelType::Strands:
@@ -755,6 +758,7 @@ void FGroomRenderingDetails::CustomizeStrandsGroupProperties(IDetailLayoutBuilde
 		DetailLayout.HideProperty(DetailLayout.GetProperty(UGroomAsset::GetLODModeMemberName(), UGroomAsset::StaticClass()));
 		DetailLayout.HideProperty(DetailLayout.GetProperty(UGroomAsset::GetAutoLODBiasMemberName(), UGroomAsset::StaticClass()));
 		DetailLayout.HideProperty(DetailLayout.GetProperty(UGroomAsset::GetDataflowSettingsMemberName(), UGroomAsset::StaticClass()));
+		DetailLayout.HideCategory(UGroomAsset::GetDataflowBuilderCategoryName());
 	}
 	break;
 	case EMaterialPanelType::Physics:
@@ -776,6 +780,7 @@ void FGroomRenderingDetails::CustomizeStrandsGroupProperties(IDetailLayoutBuilde
 		DetailLayout.HideProperty(DetailLayout.GetProperty(UGroomAsset::GetLODModeMemberName(), UGroomAsset::StaticClass()));
 		DetailLayout.HideProperty(DetailLayout.GetProperty(UGroomAsset::GetAutoLODBiasMemberName(), UGroomAsset::StaticClass()));
 		DetailLayout.HideProperty(DetailLayout.GetProperty(UGroomAsset::GetDataflowSettingsMemberName(), UGroomAsset::StaticClass()));
+		DetailLayout.HideCategory(UGroomAsset::GetDataflowBuilderCategoryName());
 	}
 	break;
 	case EMaterialPanelType::Interpolation:
@@ -797,6 +802,7 @@ void FGroomRenderingDetails::CustomizeStrandsGroupProperties(IDetailLayoutBuilde
 		DetailLayout.HideProperty(DetailLayout.GetProperty(UGroomAsset::GetLODModeMemberName(), UGroomAsset::StaticClass()));
 		DetailLayout.HideProperty(DetailLayout.GetProperty(UGroomAsset::GetAutoLODBiasMemberName(), UGroomAsset::StaticClass()));
 		DetailLayout.HideProperty(DetailLayout.GetProperty(UGroomAsset::GetDataflowSettingsMemberName(), UGroomAsset::StaticClass()));
+		DetailLayout.HideCategory(UGroomAsset::GetDataflowBuilderCategoryName());
 	}
 	break;
 	case EMaterialPanelType::LODs:
@@ -818,6 +824,7 @@ void FGroomRenderingDetails::CustomizeStrandsGroupProperties(IDetailLayoutBuilde
 //		DetailLayout.HideProperty(DetailLayout.GetProperty(UGroomAsset::GetLODModeMemberName(), UGroomAsset::StaticClass()));
 //		DetailLayout.HideProperty(DetailLayout.GetProperty(UGroomAsset::GetAutoLODBiasMemberName(), UGroomAsset::StaticClass()));
 		DetailLayout.HideProperty(DetailLayout.GetProperty(UGroomAsset::GetDataflowSettingsMemberName(), UGroomAsset::StaticClass()));
+		DetailLayout.HideCategory(UGroomAsset::GetDataflowBuilderCategoryName());
 	}
 	break;
 	case EMaterialPanelType::Dataflow:
@@ -839,6 +846,7 @@ void FGroomRenderingDetails::CustomizeStrandsGroupProperties(IDetailLayoutBuilde
 		DetailLayout.HideProperty(DetailLayout.GetProperty(UGroomAsset::GetLODModeMemberName(), UGroomAsset::StaticClass()));
 		DetailLayout.HideProperty(DetailLayout.GetProperty(UGroomAsset::GetAutoLODBiasMemberName(), UGroomAsset::StaticClass()));
 		//DetailLayout.HideProperty(DetailLayout.GetProperty(UGroomAsset::GetDataflowSettingsMemberName(), UGroomAsset::StaticClass()));
+		//DetailLayout.HideCategory(UGroomAsset::GetDataflowBuilderCategoryName());
 	}
 	break;
 	}
