Fix crash when using a groom binding with a skel. mesh with disabled sections.

diff --git a/Engine/Plugins/Runtime/HairStrands/Source/HairStrandsCore/Private/GroomBindingBuilder.cpp b/Engine/Plugins/Runtime/HairStrands/Source/HairStrandsCore/Private/GroomBindingBuilder.cpp
--- a/Engine/Plugins/Runtime/HairStrands/Source/HairStrandsCore/Private/GroomBindingBuilder.cpp
+++ b/Engine/Plugins/Runtime/HairStrands/Source/HairStrandsCore/Private/GroomBindingBuilder.cpp
@@ -92,7 +92,7 @@ FGroomBindingBuilder::FInput::FInput(UGroomBindingAsset* BindingAsset, const ITa
 FString FGroomBindingBuilder::GetVersion()
 {
 	// Important to update the version when groom building changes
-	return TEXT("6c");
+	return TEXT("6d");
 }
 
 ///////////////////////////////////////////////////////////////////////////////////////////////////
@@ -233,6 +233,7 @@ public:
 	virtual const TArray<uint32>& GetIndexBuffer() const = 0;
 	virtual const int32 GetNumSections() const = 0;
 	virtual const IMeshSectionData& GetSection(uint32 SectionIndex) const = 0;
+	virtual bool IsSectionValid(uint32 SectionIndex) const = 0;
 	virtual const FVector3f& GetVertexPosition(uint32 VertexIndex) const = 0;
 	virtual FVector2f GetVertexUV(uint32 VertexIndex, uint32 ChannelIndex) const = 0;
 	virtual FVector3f GetVertexNormal(uint32 VertexIndex) const = 0;
@@ -247,7 +248,7 @@ public:
 	{
 		for (int32 SectionIt = 0, NumSections = GetNumSections(); SectionIt < NumSections; ++SectionIt)
 		{
-			if (uint32(InMaterialIndex) == GetSection(SectionIt).GetMaterialIndex())
+			if (IsSectionValid(SectionIt) && uint32(InMaterialIndex) == GetSection(SectionIt).GetMaterialIndex())
 			{
 				return true;
 			}
@@ -262,7 +263,7 @@ public:
 		Out.Reserve(GetNumSections());
 		for (int32 SectionIndex=0, NumSections = GetNumSections(); SectionIndex<NumSections; ++SectionIndex)
 		{
-			if (GetSection(SectionIndex).GetMaterialIndex() == InMaterialIndex)
+			if (IsSectionValid(SectionIndex) && GetSection(SectionIndex).GetMaterialIndex() == InMaterialIndex)
 			{
 				Out.Add(SectionIndex);
 			}
@@ -394,6 +395,14 @@ public:
 		return Sections[InSectionIndex];
 	}
 
+	virtual bool IsSectionValid(uint32 SectionIndex) const override
+	{
+		check(MeshData);
+		check(MeshData->LODRenderData.IsValidIndex(LODIndex));
+		check(MeshData->LODRenderData[LODIndex].RenderSections.IsValidIndex(SectionIndex));
+		return !MeshData->LODRenderData[LODIndex].RenderSections[SectionIndex].bDisabled;
+	}
+
 	virtual const FVector3f& GetVertexPosition(uint32 InVertexIndex) const override
 	{
 		check(MeshData);
@@ -715,6 +724,11 @@ public:
 		return Sections[SectionIndex];
 	}
 
+	virtual bool IsSectionValid(uint32 SectionIndex) const override
+	{
+		return true;
+	}
+
 	virtual const FVector3f& GetVertexPosition(uint32 VertexIndex) const override
 	{
 		return MeshData.Positions[VertexIndex];
@@ -1502,6 +1516,11 @@ namespace GroomBinding_RootProjection
 		check(SectionCount > 0);
 		for (uint32 SectionIt = 0; SectionIt < SectionCount; ++SectionIt)
 		{
+			if (!MeshLODData.IsSectionValid(SectionIt))
+			{
+				continue;
+			}
+
 			// 2.2.1 Compute the bounding box of the skeletal mesh
 			const GroomBinding_Mesh::IMeshSectionData& Section = MeshLODData.GetSection(SectionIt);
 			const uint32 TriangleCount = Section.GetNumTriangles();
@@ -1576,6 +1595,11 @@ namespace GroomBinding_RootProjection
 		bool bIsGridPopulated = false;
 		for (uint32 SectionIt = 0; SectionIt < SectionCount; ++SectionIt)
 		{
+			if (!MeshLODData.IsSectionValid(SectionIt))
+			{
+				continue;
+			}
+
 			// 2.2.2 Insert all triangle within the grid
 			const GroomBinding_Mesh::IMeshSectionData& Section = MeshLODData.GetSection(SectionIt);
 			const uint32 TriangleCount = Section.GetNumTriangles();
@@ -1717,6 +1741,10 @@ namespace GroomBinding_RootProjection
 			const uint32 SectionCount = MeshLODData.GetNumSections();
 			for (uint32 SectionIt = 0; SectionIt < SectionCount; ++SectionIt)
 			{
+				if (!MeshLODData.IsSectionValid(SectionIt))
+				{
+					continue;
+				}
 				// 2.2.2 Insert all triangle within the grid
 				const GroomBinding_Mesh::IMeshSectionData& Section = MeshLODData.GetSection(SectionIt);
 				const uint32 VertexCount = Section.GetNumVertices();
@@ -2395,6 +2423,7 @@ namespace GroomBinding_Transfer
 			for (uint32 TargetVertexIt = 0; TargetVertexIt < TargetVertexCount; ++TargetVertexIt)
 #endif
 			{
+				// This code is only about transferring position from source to target. There is no need to remove 'disabled' sections
 				const int32 SectionIt = TargetMeshLODData.GetSectionFromVertexIndex(TargetVertexIt);
 				const int32 MaterialIndex = TargetMeshLODData.GetSection(SectionIt).GetMaterialIndex();
 				if (MaterialIndex != LocalTargetMaterialIndex)
