Fix typo reported by BugHawk in HairStrands code.

diff --git a/Engine/Plugins/Runtime/HairStrands/Source/HairStrandsCore/Private/HairStrandsMeshProjection.cpp b/Engine/Plugins/Runtime/HairStrands/Source/HairStrandsCore/Private/HairStrandsMeshProjection.cpp
--- a/Engine/Plugins/Runtime/HairStrands/Source/HairStrandsCore/Private/HairStrandsMeshProjection.cpp
+++ b/Engine/Plugins/Runtime/HairStrands/Source/HairStrandsCore/Private/HairStrandsMeshProjection.cpp
@@ -377,7 +377,7 @@ void AddHairStrandUpdateMeshTrianglesPass(
 	const bool bComputePreviousDeformedPosition = IsHairStrandContinuousDecimationReorderingEnabled();
 	if (bComputePreviousDeformedPosition)
 	{
-		Register(GraphBuilder, DeformedResources->GetLOD(MeshLODIndex)->GetDeformedUniqueTrianglePositionBuffer(FHairStrandsLODDeformedRootResource::Previous), ERDGImportedBufferFlags::CreateUAV, ERDGUnorderedAccessViewFlags::SkipBarrier);
+		OutputPrevBuffer = Register(GraphBuilder, DeformedResources->GetLOD(MeshLODIndex)->GetDeformedUniqueTrianglePositionBuffer(FHairStrandsLODDeformedRootResource::Previous), ERDGImportedBufferFlags::CreateUAV, ERDGUnorderedAccessViewFlags::SkipBarrier);
 	}
 
 	if (AddHairStrandUpdateMeshTrianglesPass(
@@ -1012,7 +1012,7 @@ void AddComputeMipsPass(
 	{
 		const uint32 SourceMipIndex = MipIt;
 		const uint32 TargetMipIndex = MipIt + 1;
-		const uint32 TargetResolution = Resolution << TargetMipIndex;
+		const uint32 TargetResolution = Resolution >> TargetMipIndex;
 
 		FGenerateMipCS::FParameters* Parameters = GraphBuilder.AllocParameters<FGenerateMipCS::FParameters>();
 		Parameters->InTexture = GraphBuilder.CreateSRV(FRDGTextureSRVDesc::CreateForMipLevel(OutTexture, SourceMipIndex));
