Fix PSO crash on platform not supporting wave 32.

diff --git a/Engine/Plugins/Runtime/HairStrands/Source/HairStrandsCore/Private/HairStrandsInterpolation.cpp b/Engine/Plugins/Runtime/HairStrands/Source/HairStrandsCore/Private/HairStrandsInterpolation.cpp
--- a/Engine/Plugins/Runtime/HairStrands/Source/HairStrandsCore/Private/HairStrandsInterpolation.cpp
+++ b/Engine/Plugins/Runtime/HairStrands/Source/HairStrandsCore/Private/HairStrandsInterpolation.cpp
@@ -562,6 +562,10 @@ public:
 
 		return FDataDrivenShaderPlatformInfo::GetSupportsWaveOperations(InPlatform);
 	}
+	static bool DoesRuntimeSupportsWaveOps32()
+	{
+		return GRHISupportsWaveOperations && GRHIMinimumWaveSize <= 32 && GRHIMaximumWaveSize >= 32;
+	}
 	static bool ShouldCompilePermutation(const FGlobalShaderPermutationParameters& Parameters) 
 	{ 
 		FPermutationDomain PermutationVector(Parameters.PermutationId);
@@ -571,6 +575,16 @@ public:
 		}
 		return IsHairStrandsSupported(EHairStrandsShaderType::All, Parameters.Platform); 
 	}
+	static EShaderPermutationPrecacheRequest ShouldPrecachePermutation(const FGlobalShaderPermutationParameters& Parameters) 
+	{ 
+		FPermutationDomain PermutationVector(Parameters.PermutationId);
+		if (PermutationVector.Get<FWaveOps>())
+		{
+			return (DoesSupportsWaveOps(Parameters.Platform, PermutationVector.Get<FPointPerCurve>()) != ERHIFeatureSupport::Unsupported && DoesRuntimeSupportsWaveOps32()) ? EShaderPermutationPrecacheRequest::Precached : EShaderPermutationPrecacheRequest::NotPrecached;
+		}		
+		return EShaderPermutationPrecacheRequest::Precached; 
+	}
+
 	static void ModifyCompilationEnvironment(const FGlobalShaderPermutationParameters& Parameters, FShaderCompilerEnvironment& OutEnvironment)
 	{
 		FGlobalShader::ModifyCompilationEnvironment(Parameters, OutEnvironment);
@@ -729,7 +743,7 @@ void AddHairStrandsInterpolationPass(
 	Parameters->DispatchCountX = DispatchInfo.DispatchCount.X;
 
 	const bool bSingleGuidePermutation = bUseSingleGuide || GetHairStrandsForceSingleGuideInterpolation();
-	const bool bWaveOps = GRHISupportsWaveOperations && GRHIMinimumWaveSize <= 32 && GRHIMaximumWaveSize >= 32 && FHairInterpolationCS::DoesSupportsWaveOps(InPlatform, DispatchInfo.PointPerCurve) != ERHIFeatureSupport::Unsupported;
+	const bool bWaveOps = FHairInterpolationCS::DoesRuntimeSupportsWaveOps32() && FHairInterpolationCS::DoesSupportsWaveOps(InPlatform, DispatchInfo.PointPerCurve) != ERHIFeatureSupport::Unsupported;
 
 	// Deformer type : 0 no deformer, 1 add deformer displacement, 2 overwrite deformed position with the deformer one
 	const int32 DeformerType = !bSupportDeformer ? 0 : GVarHairStrands_AddDeformerStrandsDisplacement ? 1 : 2;
