Fix platform specific wave support query.

diff --git a/Engine/Plugins/Runtime/HairStrands/Source/HairStrandsCore/Private/HairStrandsInterpolation.cpp b/Engine/Plugins/Runtime/HairStrands/Source/HairStrandsCore/Private/HairStrandsInterpolation.cpp
--- a/Engine/Plugins/Runtime/HairStrands/Source/HairStrandsCore/Private/HairStrandsInterpolation.cpp
+++ b/Engine/Plugins/Runtime/HairStrands/Source/HairStrandsCore/Private/HairStrandsInterpolation.cpp
@@ -552,24 +552,19 @@ class FHairInterpolationCS : public FGlobalShader
 
 public:
 	static uint32 GetGroupSize() { return 32;}
-	static ERHIFeatureSupport DoesSupportsWaveOps(EShaderPlatform InPlatform, uint32 InPointPerCurve)
+	static bool DoesSupportsWaveOps(EShaderPlatform InPlatform, uint32 InPointPerCurve)
 	{
 		// D3D11 / SM5 or preview do not support, or work well with, wave-ops by default (or SM5 preview has issues with wave intrinsics too)
 		if (InPointPerCurve != 32 || InPlatform == SP_PCD3D_SM5 || FDataDrivenShaderPlatformInfo::GetIsPreviewPlatform(InPlatform))
 		{
-			return ERHIFeatureSupport::Unsupported;
+			return false;
 		}
-
-		return FDataDrivenShaderPlatformInfo::GetSupportsWaveOperations(InPlatform);
-	}
-	static bool DoesRuntimeSupportsWaveOps32()
-	{
-		return GRHISupportsWaveOperations && GRHIMinimumWaveSize <= 32 && GRHIMaximumWaveSize >= 32;
+		return FDataDrivenShaderPlatformInfo::GetSupportsWaveOperations(InPlatform) != ERHIFeatureSupport::Unsupported && FDataDrivenShaderPlatformInfo::GetMinimumWaveSize(InPlatform) <= 32 && FDataDrivenShaderPlatformInfo::GetMaximumWaveSize(InPlatform) >= 32;
 	}
 	static bool ShouldCompilePermutation(const FGlobalShaderPermutationParameters& Parameters) 
 	{ 
 		FPermutationDomain PermutationVector(Parameters.PermutationId);
-		if (PermutationVector.Get<FWaveOps>() && DoesSupportsWaveOps(Parameters.Platform, PermutationVector.Get<FPointPerCurve>()) == ERHIFeatureSupport::Unsupported)
+		if (PermutationVector.Get<FWaveOps>() && !DoesSupportsWaveOps(Parameters.Platform, PermutationVector.Get<FPointPerCurve>()))
 		{
 			return false;
 		}
@@ -580,7 +575,7 @@ public:
 		FPermutationDomain PermutationVector(Parameters.PermutationId);
 		if (PermutationVector.Get<FWaveOps>())
 		{
-			return (DoesSupportsWaveOps(Parameters.Platform, PermutationVector.Get<FPointPerCurve>()) != ERHIFeatureSupport::Unsupported && DoesRuntimeSupportsWaveOps32()) ? EShaderPermutationPrecacheRequest::Precached : EShaderPermutationPrecacheRequest::NotPrecached;
+			return DoesSupportsWaveOps(Parameters.Platform, PermutationVector.Get<FPointPerCurve>()) ? EShaderPermutationPrecacheRequest::Precached : EShaderPermutationPrecacheRequest::NotPrecached;
 		}		
 		return EShaderPermutationPrecacheRequest::Precached; 
 	}
@@ -743,7 +738,7 @@ void AddHairStrandsInterpolationPass(
 	Parameters->DispatchCountX = DispatchInfo.DispatchCount.X;
 
 	const bool bSingleGuidePermutation = bUseSingleGuide || GetHairStrandsForceSingleGuideInterpolation();
-	const bool bWaveOps = FHairInterpolationCS::DoesRuntimeSupportsWaveOps32() && FHairInterpolationCS::DoesSupportsWaveOps(InPlatform, DispatchInfo.PointPerCurve) != ERHIFeatureSupport::Unsupported;
+	const bool bWaveOps = FHairInterpolationCS::DoesSupportsWaveOps(InPlatform, DispatchInfo.PointPerCurve);
 
 	// Deformer type : 0 no deformer, 1 add deformer displacement, 2 overwrite deformed position with the deformer one
 	const int32 DeformerType = !bSupportDeformer ? 0 : GVarHairStrands_AddDeformerStrandsDisplacement ? 1 : 2;
