Fix hair cards/meshes UVs on platform not supporting vertex fetches.

diff --git a/Engine/Plugins/Runtime/HairStrands/Source/HairStrandsCore/Private/HairCardsVertexFactory.cpp b/Engine/Plugins/Runtime/HairStrands/Source/HairStrandsCore/Private/HairCardsVertexFactory.cpp
--- a/Engine/Plugins/Runtime/HairStrands/Source/HairStrandsCore/Private/HairCardsVertexFactory.cpp
+++ b/Engine/Plugins/Runtime/HairStrands/Source/HairStrandsCore/Private/HairCardsVertexFactory.cpp
@@ -403,12 +403,12 @@ void FHairCardsVertexFactory::InitResources(FRHICommandListBase& RHICmdList)
 		{
 			// Add additional TexCoord
 			const uint32 UVSize = sizeof(FVector2f);
-			const EVertexElementType UVVertexElementType = EVertexElementType::VET_Float4;
+			const uint32 NumUVComponentCount = FMath::DivideAndRoundUp(NumUVs, 2u); // Two UVs are packed per Float4
 
 			// UV[0] is rootUV which is not part of the regular UVs, hence UV starts at index 1
 			for (uint32 UVIt = 1; UVIt < NumUVs; UVIt +=2)
 			{
-				Elements.Add(this->AccessStreamComponent(FVertexStreamComponent(UVsBuffer, 0, UVIt * UVSize, UVSize * NumUVs, UVVertexElementType, EVertexStreamUsage::Default), StartStreamIndex++));
+				Elements.Add(this->AccessStreamComponent(FVertexStreamComponent(UVsBuffer, 0, UVIt * UVSize, FHairCardsUVFormat::SizeInByte * NumUVComponentCount, FHairCardsUVFormat::VertexElementType, EVertexStreamUsage::Default), StartStreamIndex++));
 			}
 		};
 
@@ -416,6 +416,7 @@ void FHairCardsVertexFactory::InitResources(FRHICommandListBase& RHICmdList)
 		{
 			const FHairGroupInstance::FCards::FLOD& LOD = Data.Instance->Cards.LODs[Data.LODIndex];
 			const uint32 NumUVs = LOD.RestResource->BulkData.GetNumUVs();
+			const uint32 NumUVComponentCount = FMath::DivideAndRoundUp(NumUVs, 2u); // Two UVs are packed per Float4
 			const bool bHasVertexColor = LOD.RestResource->BulkData.Header.bVertexColor;
 			const bool bDynamic = LOD.DeformedResource != nullptr;
 			const FHairCardsVertexBuffer* VertexColorsBuffer = bHasVertexColor ? &LOD.RestResource->VertexColorsBuffer : &LOD.RestResource->NormalsBuffer;
@@ -446,7 +447,7 @@ void FHairCardsVertexFactory::InitResources(FRHICommandListBase& RHICmdList)
 			}
 			Elements.Add(AccessStreamComponent(FVertexStreamComponent(&LOD.RestResource->MaterialsBuffer,			0, 0,									FHairCardsMaterialFormat::SizeInByte,												FHairCardsMaterialFormat::VertexElementType,	EVertexStreamUsage::Default), 4));
 			Elements.Add(AccessStreamComponent(FVertexStreamComponent(VertexColorsBuffer,							0, 0,									FHairCardsMaterialFormat::SizeInByte,												FHairCardsMaterialFormat::VertexElementType,	EVertexStreamUsage::Default), 5));
-			Elements.Add(AccessStreamComponent(FVertexStreamComponent(&LOD.RestResource->UVsBuffer,					0, 0,									FHairCardsUVFormat::SizeInByte,														FHairCardsUVFormat::VertexElementType,			EVertexStreamUsage::Default), 6));
+			Elements.Add(AccessStreamComponent(FVertexStreamComponent(&LOD.RestResource->UVsBuffer,					0, 0,									FHairCardsUVFormat::SizeInByte * NumUVComponentCount,														FHairCardsUVFormat::VertexElementType,			EVertexStreamUsage::Default), 6));
 
 			// Add additional UVs
 			AddUVChannels(&LOD.RestResource->UVsBuffer, LOD.RestResource->BulkData.GetNumUVs(), 7);
@@ -469,6 +470,8 @@ void FHairCardsVertexFactory::InitResources(FRHICommandListBase& RHICmdList)
 		else if (Data.GeometryType == EHairGeometryType::Meshes)
 		{
 			const FHairGroupInstance::FMeshes::FLOD& LOD = Data.Instance->Meshes.LODs[Data.LODIndex];
+			const uint32 NumUVs = LOD.RestResource->BulkData.GetNumUVs();
+			const uint32 NumUVComponentCount = FMath::DivideAndRoundUp(NumUVs, 2u); // Two UVs are packed per Float4
 			const bool bHasVertexColor = LOD.RestResource->BulkData.Header.bVertexColor;
 			const FHairCardsVertexBuffer* VertexColorsBuffer = bHasVertexColor ? &LOD.RestResource->VertexColorsBuffer : &LOD.RestResource->NormalsBuffer;
 
@@ -495,7 +498,7 @@ void FHairCardsVertexFactory::InitResources(FRHICommandListBase& RHICmdList)
 
 			Elements.Add(AccessStreamComponent(FVertexStreamComponent(&LOD.RestResource->NormalsBuffer,				0, 0,									FHairCardsMaterialFormat::SizeInByte,												FHairCardsMaterialFormat::VertexElementType,	EVertexStreamUsage::Default), 4)); // Dummy
 			Elements.Add(AccessStreamComponent(FVertexStreamComponent(&LOD.RestResource->VertexColorsBuffer,		0, 0,									FHairCardsMaterialFormat::SizeInByte,												FHairCardsMaterialFormat::VertexElementType,	EVertexStreamUsage::Default), 5));
-			Elements.Add(AccessStreamComponent(FVertexStreamComponent(&LOD.RestResource->UVsBuffer,					0, 0,									FHairCardsUVFormat::SizeInByte,														FHairCardsUVFormat::VertexElementType,			EVertexStreamUsage::Default), 6));
+			Elements.Add(AccessStreamComponent(FVertexStreamComponent(&LOD.RestResource->UVsBuffer,					0, 0,									FHairCardsUVFormat::SizeInByte * NumUVComponentCount,								FHairCardsUVFormat::VertexElementType,			EVertexStreamUsage::Default), 6));
 
 			// Add additional UVs
 			AddUVChannels(&LOD.RestResource->UVsBuffer, LOD.RestResource->BulkData.GetNumUVs(), 7);
