Groom Dataflow terminal node : allow renaming of the attribute key inputs
- Also add support for drag and dropping a connection to create a new input

diff --git a/Engine/Plugins/Runtime/HairStrands/Source/HairStrandsDataflow/Private/GroomAssetTerminalNode.cpp b/Engine/Plugins/Runtime/HairStrands/Source/HairStrandsDataflow/Private/GroomAssetTerminalNode.cpp
--- a/Engine/Plugins/Runtime/HairStrands/Source/HairStrandsDataflow/Private/GroomAssetTerminalNode.cpp
+++ b/Engine/Plugins/Runtime/HairStrands/Source/HairStrandsDataflow/Private/GroomAssetTerminalNode.cpp
@@ -12,6 +12,8 @@
 
 namespace UE::Groom::Private
 {
+	static const FName AttributeKeyTypeName(TEXT("FCollectionAttributeKey"));
+
 	FORCEINLINE void BuildEditableGuides(const UE::Groom::FGroomGuidesFacade& GuidesFacade, FEditableGroom& EditGroom)
 	{
 		int32 ObjectIndex = 0, CurveIndex = 0;
@@ -202,6 +204,16 @@ namespace UE::Groom::Private
 	}
 }
 
+FGroomAssetTerminalDataflowNode::FGroomAssetTerminalDataflowNode(const UE::Dataflow::FNodeParameters& InParam, FGuid InGuid)
+	: FDataflowTerminalNode(InParam, InGuid)
+{
+	DataflowAssetWeakPtr = Cast<UDataflow>(InParam.OwningObject);
+
+	RegisterInputConnection(&Collection);
+	RegisterOutputConnection(&Collection, &Collection);
+}
+
+
 void FGroomAssetTerminalDataflowNode::SetAssetValue(TObjectPtr<UObject> Asset, UE::Dataflow::FContext& Context) const
 {
 	if (UGroomAsset* GroomAsset = Cast<UGroomAsset>(Asset.Get()))
@@ -275,7 +287,11 @@ UE::Dataflow::TConnectionReference<FCollectionAttributeKey> FGroomAssetTerminalD
 TArray<UE::Dataflow::FPin> FGroomAssetTerminalDataflowNode::AddPins()
 {
 	const int32 Index = AttributeKeys.AddDefaulted();
-	const FDataflowInput& Input = RegisterInputArrayConnection(GetConnectionReference(Index));
+	FDataflowInput& Input = RegisterInputArrayConnection(GetConnectionReference(Index));
+
+	// Make sure we have the right number of names
+	AttributeNames.SetNum(AttributeKeys.Num());
+
 	return { { UE::Dataflow::FPin::EDirection::INPUT, Input.GetType(), Input.GetName() } };
 }
 
@@ -302,6 +318,9 @@ void FGroomAssetTerminalDataflowNode::OnPinRemoved(const UE::Dataflow::FPin& Pin
 #endif
 	AttributeKeys.SetNum(Index);
 
+	// Make sure we have the right number of names
+	AttributeNames.SetNum(AttributeKeys.Num());
+
 	return Super::OnPinRemoved(Pin);
 }
 
@@ -339,9 +358,107 @@ void FGroomAssetTerminalDataflowNode::PostSerialize(const FArchive& Ar)
 		{
 			ensureAlways(AttributeKeys.Num() + NumOtherInputs == GetNumInputs());
 		}
+		// make sure the number of names is the same as the current number of key inputs
+		AttributeNames.SetNum(AttributeKeys.Num());
+		SyncInputNames();
+	}
+}
+
+bool FGroomAssetTerminalDataflowNode::ShouldInvalidateOnPropertyChanged(const FPropertyChangedEvent& InPropertyChangedEvent) const
+{
+	if (InPropertyChangedEvent.GetPropertyName() == GET_MEMBER_NAME_CHECKED(FGroomAssetTerminalDataflowNode, AttributeNames))
+	{
+		// we only changing the name of the inputs, no need to inavlidate the node
+		return false;
+	}
+	return true;
+}
+
+void FGroomAssetTerminalDataflowNode::OnPropertyChanged(UE::Dataflow::FContext& Context, const FPropertyChangedEvent& InPropertyChangedEvent)
+{
+	if (InPropertyChangedEvent.GetPropertyName() == GET_MEMBER_NAME_CHECKED(FGroomAssetTerminalDataflowNode, AttributeNames))
+	{
+		SyncInputNames();
+	}
+}
+
+bool FGroomAssetTerminalDataflowNode::SupportsDropConnectionOnNode(FName TypeName, UE::Dataflow::FPin::EDirection Direction) const
+{
+	if (TypeName == UE::Groom::Private::AttributeKeyTypeName && Direction == UE::Dataflow::FPin::EDirection::OUTPUT)
+	{
+		return true;
+	}
+	return false;
+}
+
+const FDataflowConnection* FGroomAssetTerminalDataflowNode::OnDropConnectionOnNode(const FDataflowConnection& DroppedConnection)
+{
+	if (DroppedConnection.GetType() == UE::Groom::Private::AttributeKeyTypeName && DroppedConnection.GetDirection() == UE::Dataflow::FPin::EDirection::OUTPUT)
+	{
+		const int32 Index = AttributeKeys.AddDefaulted();
+		FDataflowInput& NewInput = RegisterInputArrayConnection(GetConnectionReference(Index));
+
+		// Make sure we have the right number of names
+		AttributeNames.SetNum(AttributeKeys.Num());
+
+		// infer the name for the new entry from the dropped connection
+		const FProperty* DroppedProperty = DroppedConnection.GetProperty();
+		const FName DroppedDisplayName = DroppedProperty 
+			? FName(DroppedProperty->GetDisplayNameText().ToString()) 
+			: DroppedConnection.GetName();
+
+		const FName NewName = GenerateUniqueInputName(DroppedDisplayName);
+		if (!NewName.IsNone())
+		{
+			AttributeNames.Add(NewName);
+			NewInput.SetName(NewName);
+		}
+		return &NewInput;
 	}
+	return nullptr;
 }
 
+void FGroomAssetTerminalDataflowNode::SyncInputNames()
+{
+	bool bChanged = false;
+	for (int32 KeyIndex = 0; KeyIndex < AttributeKeys.Num(); ++KeyIndex)
+	{
+		if (FDataflowInput* Input = FindInput(GetConnectionReference(KeyIndex)))
+		{
+			if (AttributeNames.IsValidIndex(KeyIndex) && Input->GetName() != AttributeNames[KeyIndex])
+			{
+				const FName UniqueName = GenerateUniqueInputName(AttributeNames[KeyIndex]);
+				if (!UniqueName.IsNone())
+				{
+					Input->SetName(UniqueName);
+					bChanged = true;
+				}
+			}
+		}
+	}
 
+	// refresh the Ed Node is necessary
+	if (bChanged)
+	{
+		if (TStrongObjectPtr<UDataflow> DataflowAsset = DataflowAssetWeakPtr.Pin())
+		{
+			DataflowAsset->RefreshEdNodeByGuid(GetGuid());
+		}
+	}
+}
 
+FName FGroomAssetTerminalDataflowNode::GenerateUniqueInputName(FName BaseName) const
+{
+	if (BaseName.IsNone())
+	{
+		return BaseName;
+	}
+	FName NewName = BaseName;
+	int32 SuffixNumber = 1;
+	while (FindInput(NewName) != nullptr)
+	{
+		NewName.SetNumber(SuffixNumber++);
+	}
+	return NewName;
+}
 
diff --git a/Engine/Plugins/Runtime/HairStrands/Source/HairStrandsDataflow/Public/GroomAssetTerminalNode.h b/Engine/Plugins/Runtime/HairStrands/Source/HairStrandsDataflow/Public/GroomAssetTerminalNode.h
--- a/Engine/Plugins/Runtime/HairStrands/Source/HairStrandsDataflow/Public/GroomAssetTerminalNode.h
+++ b/Engine/Plugins/Runtime/HairStrands/Source/HairStrandsDataflow/Public/GroomAssetTerminalNode.h
@@ -1,4 +1,4 @@
-﻿// Copyright Epic Games, Inc. All Rights Reserved.
+// Copyright Epic Games, Inc. All Rights Reserved.
 
 #pragma once 
 
@@ -8,9 +8,12 @@
 #include "Dataflow/DataflowCore.h"
 #include "Dataflow/DataflowEngine.h"
 #include "Dataflow/DataflowConnectionTypes.h"
+#include "UObject/WeakObjectPtrTemplatesFwd.h"
 
 #include "GroomAssetTerminalNode.generated.h"
 
+class UDataflow;
+
 USTRUCT(meta = (Experimental, DataflowGroom, DataflowTerminal))
 struct FGroomAssetTerminalDataflowNode : public FDataflowTerminalNode
 {
@@ -19,12 +22,7 @@ struct FGroomAssetTerminalDataflowNode : public FDataflowTerminalNode
 
 public:
 
-	FGroomAssetTerminalDataflowNode(const UE::Dataflow::FNodeParameters& InParam, FGuid InGuid = FGuid::NewGuid())
-	: FDataflowTerminalNode(InParam, InGuid)
-	{
-		RegisterInputConnection(&Collection);
-		RegisterOutputConnection(&Collection, &Collection);
-	}
+	FGroomAssetTerminalDataflowNode(const UE::Dataflow::FNodeParameters& InParam, FGuid InGuid = FGuid::NewGuid());
 
 	//~ Begin FDataflowTerminalNode interface
 	virtual void SetAssetValue(TObjectPtr<UObject> Asset, UE::Dataflow::FContext& Context) const override;
@@ -45,10 +43,24 @@ public:
 	UPROPERTY()
 	TArray<FCollectionAttributeKey> AttributeKeys;
 
+	UPROPERTY(EditAnyWhere, Category = Attributes, EditFixedSize)
+	TArray<FName> AttributeNames;
+
 private :
 	/** Get the connection reference matching an attribute key index */
 	UE::Dataflow::TConnectionReference<FCollectionAttributeKey> GetConnectionReference(int32 Index) const;
 
+	virtual bool ShouldInvalidateOnPropertyChanged(const FPropertyChangedEvent& InPropertyChangedEvent) const override;
+	virtual void OnPropertyChanged(UE::Dataflow::FContext& Context, const FPropertyChangedEvent& InPropertyChangedEvent) override;
+
+	virtual bool SupportsDropConnectionOnNode(FName TypeName, UE::Dataflow::FPin::EDirection Direction) const override;
+	virtual const FDataflowConnection* OnDropConnectionOnNode(const FDataflowConnection& DroppedConnection) override;
+
+	void SyncInputNames();
+	FName GenerateUniqueInputName(FName BaseName) const;
+
+	TWeakObjectPtr<UDataflow> DataflowAssetWeakPtr;
+
 	static constexpr int32 NumOtherInputs = 1;
 };
 
