Groom : add support for loading templates dynamically at time of request

diff --git a/Engine/Plugins/Runtime/HairStrands/Source/HairStrandsEditor/Private/HairStrandsEditor.cpp b/Engine/Plugins/Runtime/HairStrands/Source/HairStrandsEditor/Private/HairStrandsEditor.cpp
--- a/Engine/Plugins/Runtime/HairStrands/Source/HairStrandsEditor/Private/HairStrandsEditor.cpp
+++ b/Engine/Plugins/Runtime/HairStrands/Source/HairStrandsEditor/Private/HairStrandsEditor.cpp
@@ -156,9 +156,9 @@ void FGroomEditor::StartupModule()
 	UThumbnailManager::Get().RegisterCustomRenderer(UGroomAsset::StaticClass(), UGroomAssetThumbnailRenderer::StaticClass());
 	UThumbnailManager::Get().RegisterCustomRenderer(UGroomBindingAsset::StaticClass(), UGroomBindingAssetThumbnailRenderer::StaticClass());
 
-	UE::Groom::RegisterGroomDataflowTemplate({"SelectNoDataflowTemplate", "No Dataflow", "Don't add a Dataflow to this Groom Asset. Useful when a Dataflow already exists.", "", false});
-	UE::Groom::RegisterGroomDataflowTemplate({"SelectEmptyDataflowTemplate", "Empty Dataflow", "Add an empty Dataflow with a Groom Terminal node.", "/HairStrands/Dataflow/DF_GroomEmptyTemplate.DF_GroomDataflowTemplate", false});
-	UE::Groom::RegisterGroomDataflowTemplate({"SelectGuidesDataflowTemplate", "Guides Dataflow", "Add a Dataflow with basic nodes guides generation.", "/HairStrands/Dataflow/DF_GroomGuidesTemplate.DF_GroomGuidesTemplate", true});
+	// default no dataflow template
+	UE::Groom::RegisterGroomDataflowTemplate({ "SelectNoDataflowTemplate", "No Dataflow", "Don't add a Dataflow to this Groom Asset. Useful when a Dataflow already exists.", "", false });
+	UE::Groom::RegisterGroomDataflowTemplatePath(TEXT("/HairStrands/Dataflow"));
 }
 
 void FGroomEditor::ShutdownModule()
@@ -204,8 +204,7 @@ void FGroomEditor::ShutdownModule()
 	}
 
 	UE::Groom::UnregisterGroomDataflowTemplate("SelectNoDataflowTemplate");
-	UE::Groom::UnregisterGroomDataflowTemplate("SelectEmptyDataflowTemplate");
-	UE::Groom::UnregisterGroomDataflowTemplate("SelectGuidesDataflowTemplate");
+	UE::Groom::RegisterGroomDataflowTemplatePath(TEXT("/HairStrands/Dataflow"));
 }
 
 void FGroomEditor::RegisterMenus()
diff --git a/Engine/Plugins/Runtime/HairStrands/Source/HairStrandsEditor/Private/HairStrandsFactory.cpp b/Engine/Plugins/Runtime/HairStrands/Source/HairStrandsEditor/Private/HairStrandsFactory.cpp
--- a/Engine/Plugins/Runtime/HairStrands/Source/HairStrandsEditor/Private/HairStrandsFactory.cpp
+++ b/Engine/Plugins/Runtime/HairStrands/Source/HairStrandsEditor/Private/HairStrandsFactory.cpp
@@ -23,6 +23,8 @@
 #include "Misc/ScopedSlowTask.h"
 #include "Misc/PackageName.h"
 #include "Modules/ModuleManager.h"
+
+#include "AssetRegistry/AssetData.h"
 #include "AssetRegistry/AssetRegistryModule.h"
 
 #include UE_INLINE_GENERATED_CPP_BY_NAME(HairStrandsFactory)
@@ -41,18 +43,25 @@ namespace UE::Groom
 		}
 
 		/** Build dataflow template buttons from data*/
-		TArray<SMessageDialog::FButton> BuildDataflowTemplateButtons() const
+		TArray<SMessageDialog::FButton> BuildDataflowTemplateButtons(TFunctionRef<void(FString)> OnButtonClicked) const
 		{
 			TArray<SMessageDialog::FButton> DataflowButtons;
 			DataflowButtons.Reserve(TemplateData.Num());
 			
-			for(const FGroomDataflowTemplateData& DataflowTemplate : TemplateData)
+			TArray<FGroomDataflowTemplateData> OutAllTemplateData;
+			GetFixedAndDynamicallyLoadedTemplateData(OutAllTemplateData);
+
+			for(const FGroomDataflowTemplateData& DataflowTemplate : OutAllTemplateData)
 			{
 				const FString DataflowTitleID = DataflowTemplate.TemplateName + FString("Title");
 				const FString DataflowTooltipID = DataflowTemplate.TemplateName + FString("Tooltip");
-				DataflowButtons.Add(SMessageDialog::FButton(FText::AsLocalizable_Advanced(TEXT(LOCTEXT_NAMESPACE), *DataflowTitleID, *DataflowTemplate.TemplateTitle))
+				DataflowButtons.Add(
+					SMessageDialog::FButton(FText::AsLocalizable_Advanced(TEXT(LOCTEXT_NAMESPACE), *DataflowTitleID, *DataflowTemplate.TemplateTitle))
 					.SetToolTipText(FText::AsLocalizable_Advanced(TEXT(LOCTEXT_NAMESPACE), *DataflowTooltipID, *DataflowTemplate.TemplateTooltip))
-					.SetPrimary(DataflowTemplate.bIsPrimaryTemplate));
+					.SetPrimary(DataflowTemplate.bIsPrimaryTemplate)
+					.SetOnClicked(FSimpleDelegate::CreateLambda(OnButtonClicked, DataflowTemplate.TemplatePath))
+				);
+				
 			}
 			return DataflowButtons;
 		}
@@ -81,10 +90,56 @@ namespace UE::Groom
 			}
 		}
 
-	private:
+		void AddDataflowTemplatePath(const FString& DataflowTemplatePath)
+		{
+			DynamicallyLoadedTemplatePaths.Add(DataflowTemplatePath);
+		}
+
+		void RemoveDataflowTemplatePath(const FString& DataflowTemplatePath)
+		{
+			DynamicallyLoadedTemplatePaths.Remove(DataflowTemplatePath);
+		}
+
+		void GetFixedAndDynamicallyLoadedTemplateData(TArray<FGroomDataflowTemplateData>& OutAllTemplateData) const
+		{
+			// first fixed ones
+			OutAllTemplateData = TemplateData;
+
+			const FAssetRegistryModule& AssetRegistryModule = FModuleManager::LoadModuleChecked<FAssetRegistryModule>("AssetRegistry");
+
+			// now load on the spot the ones from dynamically loaded paths
+			TArray<FAssetData> AssetDataArray;
+			for (const FString& TemplatePath : DynamicallyLoadedTemplatePaths)
+			{
+				AssetDataArray.Reset();
+				AssetRegistryModule.Get().GetAssetsByPath(FName(TemplatePath), AssetDataArray, /*bRecursive*/true, /*bIncludeOnlyOnDiskAssets*/true);
+
+				for (const FAssetData& AssetData : AssetDataArray)
+				{
+					const FString AssetNameStr =
+						AssetData.AssetName.ToString()
+						.Replace(TEXT("DF_Groom"), TEXT(""))
+						.Replace(TEXT("Template"), TEXT(""))
+						;
+
+
+					UE::Groom::FGroomDataflowTemplateData LoadedTemplateData
+					{
+						.TemplateName = AssetNameStr,
+						.TemplateTitle = AssetNameStr,
+						.TemplateTooltip = FString::Format(TEXT("Use {0} Dataflow Template"), { AssetData.AssetName.ToString() }),
+						.TemplatePath = AssetData.GetObjectPathString(),
+						.bIsPrimaryTemplate = false,
+					};
+					OutAllTemplateData.Emplace(MoveTemp(LoadedTemplateData));
+				}
+			}
+		}
 
+	private:
 		/** List of dataflow template data*/
 		TArray<FGroomDataflowTemplateData>	TemplateData;
+		TArray<FString> DynamicallyLoadedTemplatePaths;
 	};
 
 	void RegisterGroomDataflowTemplate(const FGroomDataflowTemplateData& DataflowTemplate)
@@ -97,6 +152,16 @@ namespace UE::Groom
 		FGroomDataflowTemplateManager::GetInstance().RemoveDataflowTemplateData(TemplateName);
 	}
 
+	void RegisterGroomDataflowTemplatePath(const FString& DataflowTemplatePath)
+	{
+		FGroomDataflowTemplateManager::GetInstance().AddDataflowTemplatePath(DataflowTemplatePath);
+	}
+
+	void UnregisterGroomDataflowTemplatePath(const FString& DataflowTemplatePath)
+	{
+		FGroomDataflowTemplateManager::GetInstance().RemoveDataflowTemplatePath(DataflowTemplatePath);
+	}
+
 	void BuildGroomDataflowAsset(UGroomAsset* GroomAsset)
 	{
 		// Create a new Dataflow asset
@@ -112,16 +177,22 @@ namespace UE::Groom
 		}
 		UPackage* const DataflowPackage = CreatePackage(*DataflowPackageName);
 
+		FString GroomDataflowTemplate;
+		auto OnButtonClicked = [&GroomDataflowTemplate](FString TemplatePath)
+			{
+				GroomDataflowTemplate = TemplatePath;
+			};
+
+		TArray<SMessageDialog::FButton> Buttons = FGroomDataflowTemplateManager::GetInstance().BuildDataflowTemplateButtons(OnButtonClicked);
+
 		// Select the template
 		TSharedRef<SMessageDialog> SelectTemplateMessageDialog = SNew(SMessageDialog)
 			.Title(FText(LOCTEXT("SelectTemplateTitle", "Select a Groom Dataflow Template")))
 			.Message(LOCTEXT("SelectTemplateMessage", "Select a template for this groom asset:"))
-			.Buttons(
-				{ 
-					FGroomDataflowTemplateManager::GetInstance().BuildDataflowTemplateButtons()
-				});
+			.Buttons(Buttons);
 
-		const FString GroomDataflowTemplate = FGroomDataflowTemplateManager::GetInstance().GetDataflowTemplatePath(SelectTemplateMessageDialog->ShowModal());
+		// result will be set in GroomDataflowTemplate
+		SelectTemplateMessageDialog->ShowModal();
 
 		// Load the dataflow template into the groom asset
 		UDataflow* const Template = GroomDataflowTemplate.IsEmpty() ? nullptr : LoadObject<UDataflow>(nullptr, *GroomDataflowTemplate);
diff --git a/Engine/Plugins/Runtime/HairStrands/Source/HairStrandsEditor/Public/HairStrandsFactory.h b/Engine/Plugins/Runtime/HairStrands/Source/HairStrandsEditor/Public/HairStrandsFactory.h
--- a/Engine/Plugins/Runtime/HairStrands/Source/HairStrandsEditor/Public/HairStrandsFactory.h
+++ b/Engine/Plugins/Runtime/HairStrands/Source/HairStrandsEditor/Public/HairStrandsFactory.h
@@ -33,10 +33,15 @@ namespace UE::Groom
 	/** Register dataflow template to the manager */
 	HAIRSTRANDSEDITOR_API void RegisterGroomDataflowTemplate(const FGroomDataflowTemplateData& DataflowTemplate);
 
-
 	/** Unregister dataflow template to the manager */
 	HAIRSTRANDSEDITOR_API void UnregisterGroomDataflowTemplate(const FString& TemplateName);
 
+	/** Register a dataflow template path to the manager. this will be used to dynamically load template from it whenever the list of templates is requested */
+	HAIRSTRANDSEDITOR_API void RegisterGroomDataflowTemplatePath(const FString& DataflowTemplatePath);
+
+	/** Unregister a dataflow template path from the manager */
+	HAIRSTRANDSEDITOR_API void UnregisterGroomDataflowTemplatePath(const FString& DataflowTemplatePath);
+
 	/** Build groom dataflow asset from templates */
 	HAIRSTRANDSEDITOR_API void BuildGroomDataflowAsset(UGroomAsset* GroomAsset);
 }
