Groom : fix few things with the template dynamic scanning code

diff --git a/Engine/Plugins/Runtime/HairStrands/Source/HairStrandsEditor/Private/HairStrandsEditor.cpp b/Engine/Plugins/Runtime/HairStrands/Source/HairStrandsEditor/Private/HairStrandsEditor.cpp
--- a/Engine/Plugins/Runtime/HairStrands/Source/HairStrandsEditor/Private/HairStrandsEditor.cpp
+++ b/Engine/Plugins/Runtime/HairStrands/Source/HairStrandsEditor/Private/HairStrandsEditor.cpp
@@ -204,7 +204,7 @@ void FGroomEditor::ShutdownModule()
 	}
 
 	UE::Groom::UnregisterGroomDataflowTemplate("SelectNoDataflowTemplate");
-	UE::Groom::RegisterGroomDataflowTemplatePath(TEXT("/HairStrands/Dataflow"));
+	UE::Groom::UnregisterGroomDataflowTemplatePath(TEXT("/HairStrands/Dataflow"));
 }
 
 void FGroomEditor::RegisterMenus()
diff --git a/Engine/Plugins/Runtime/HairStrands/Source/HairStrandsEditor/Private/HairStrandsFactory.cpp b/Engine/Plugins/Runtime/HairStrands/Source/HairStrandsEditor/Private/HairStrandsFactory.cpp
--- a/Engine/Plugins/Runtime/HairStrands/Source/HairStrandsEditor/Private/HairStrandsFactory.cpp
+++ b/Engine/Plugins/Runtime/HairStrands/Source/HairStrandsEditor/Private/HairStrandsFactory.cpp
@@ -43,7 +43,7 @@ namespace UE::Groom
 		}
 
 		/** Build dataflow template buttons from data*/
-		TArray<SMessageDialog::FButton> BuildDataflowTemplateButtons(TFunctionRef<void(FString)> OnButtonClicked) const
+		TArray<SMessageDialog::FButton> BuildDataflowTemplateButtons(TFunction<void(FString)> OnButtonClicked) const
 		{
 			TArray<SMessageDialog::FButton> DataflowButtons;
 			DataflowButtons.Reserve(TemplateData.Num());
