Hair spline skinning

#virtualized

diff --git a/Engine/Plugins/Runtime/HairStrands/Shaders/Private/Deformers/DeformerGroomMeshesRead.ush b/Engine/Plugins/Runtime/HairStrands/Shaders/Private/Deformers/DeformerGroomMeshesRead.ush
--- a/Engine/Plugins/Runtime/HairStrands/Shaders/Private/Deformers/DeformerGroomMeshesRead.ush
+++ b/Engine/Plugins/Runtime/HairStrands/Shaders/Private/Deformers/DeformerGroomMeshesRead.ush
@@ -1,4 +1,4 @@
-﻿// Copyright Epic Games, Inc. All Rights Reserved.
+// Copyright Epic Games, Inc. All Rights Reserved.
 
 #include "/Plugin/ComputeFramework/Private/ComputeKernelCommon.ush"
 
@@ -19,5 +19,14 @@ float3x4 ReadBoneTransformMatrix_{DataInterfaceName}(uint BoneIndex)
 		float3x4({DataInterfaceName}_BoneTransformMatrices[3*BoneIndex], {DataInterfaceName}_BoneTransformMatrices[3*BoneIndex+1], {DataInterfaceName}_BoneTransformMatrices[3*BoneIndex+2]) : float3x4(1.0f,0.0f,0.0f,0.0f,0.0f,1.0f,0.0f,0.0f,0.0f,0.0f,1.0f,0.0f);
 }
 
+float3x4 ReadBindTransformMatrix_{DataInterfaceName}(uint BoneIndex)
+{
+	return (BoneIndex < {DataInterfaceName}_NumBones) ?
+		float3x4({DataInterfaceName}_BoneTransformMatrices[({DataInterfaceName}_NumBones + BoneIndex) * 3], 
+				 {DataInterfaceName}_BoneTransformMatrices[({DataInterfaceName}_NumBones + BoneIndex) * 3 + 1], 
+				 {DataInterfaceName}_BoneTransformMatrices[({DataInterfaceName}_NumBones + BoneIndex) * 3 + 2]) : 
+		float3x4(1.0f,0.0f,0.0f,0.0f,0.0f,1.0f,0.0f,0.0f,0.0f,0.0f,1.0f,0.0f);
+}
+
 
 
diff --git a/Engine/Plugins/Runtime/HairStrands/Source/HairStrandsDataflow/Private/BuildGroomSplineSkinningNode.cpp b/Engine/Plugins/Runtime/HairStrands/Source/HairStrandsDataflow/Private/BuildGroomSplineSkinningNode.cpp
new file mode 100644
--- /dev/null
+++ b/Engine/Plugins/Runtime/HairStrands/Source/HairStrandsDataflow/Private/BuildGroomSplineSkinningNode.cpp
@@ -0,0 +1,734 @@
+// Copyright Epic Games, Inc. All Rights Reserved.
+
+#include "BuildGroomSplineSkinningNode.h"
+#include "Algo/MaxElement.h"
+#include "Algo/Count.h"
+
+#include UE_INLINE_GENERATED_CPP_BY_NAME(BuildGroomSplineSkinningNode)
+
+namespace UE::Groom::Private
+{
+	struct FCatmullRomSpline
+	{
+		TArray<FVector> CVs;
+		TArray<float> Params;
+		FIntVector2 Bones;
+
+		void ClosestParam(const FVector& Position, float& Param, float& ClosestDistance) const;
+		void Init(FReferenceSkeleton& Skeleton, int32 RootIndex, int32 NumSamplesPerSegment);
+	};
+
+
+	void FCatmullRomSpline::Init(FReferenceSkeleton& Skeleton, int32 RootIndex, int32 NumSamplesPerSegment)
+	{
+		// Spline evaluation code from DrawCentripetalCatmullRomSpline (DrawDebugHelpers.cpp)
+		const float Alpha = 0.5f; // Centripedal Catmull-Rom
+
+		auto GetT = [](float T, float Alpha, const FVector& P0, const FVector& P1)
+			{
+				const FVector P1P0 = P1 - P0;
+				const float Dot = P1P0 | P1P0;
+				const float Pow = FMath::Pow(Dot, Alpha * .5f);
+				return Pow + T;
+			};
+
+		TArray<int32> BoneIndices;
+		BoneIndices.SetNum(0);
+		BoneIndices.Add(RootIndex);
+		TArray<int32> Children;
+		int32 ParentIndex = RootIndex;
+		while (Skeleton.GetDirectChildBones(BoneIndices.Last(), Children))
+		{
+			if (Children.Num() > 0)
+			{
+				BoneIndices.Add(Children[0]);
+			}
+			else
+			{
+				// TODO: Add support for branching splines
+			}
+		}
+
+		Bones[0] = RootIndex;
+		Bones[1] = BoneIndices.Last();
+
+		TArray<FTransform> BoneTransforms;
+		Skeleton.GetRawBoneAbsoluteTransforms(BoneTransforms);
+
+		const int32 NumSegments = BoneIndices.Num() - 1;
+
+		CVs.SetNum(NumSegments * NumSamplesPerSegment + 1);
+		Params.SetNum(NumSegments * NumSamplesPerSegment + 1);
+
+		for (int32 BoneIndex = 0; BoneIndex < NumSegments; ++BoneIndex)
+		{
+			const FVector P0 = BoneTransforms[BoneIndices[FMath::Max(BoneIndex - 1, 0)			]].GetTranslation();
+			const FVector P1 = BoneTransforms[BoneIndices[BoneIndex								]].GetTranslation();
+			const FVector P2 = BoneTransforms[BoneIndices[BoneIndex + 1							]].GetTranslation();
+			const FVector P3 = BoneTransforms[BoneIndices[FMath::Min(BoneIndex + 2, NumSegments)]].GetTranslation();
+
+			const float T0 = 0.0f;
+			const float T1 = GetT(T0, Alpha, P0, P1);
+			const float T2 = GetT(T1, Alpha, P1, P2);
+			const float T3 = GetT(T2, Alpha, P2, P3);
+
+			const float T1T0 = T1 - T0;
+			const float T2T1 = T2 - T1;
+			const float T3T2 = T3 - T2;
+			const float T2T0 = T2 - T0;
+			const float T3T1 = T3 - T1;
+
+			const bool bIsNearlyZeroT1T0 = FMath::IsNearlyZero(T1T0, UE_KINDA_SMALL_NUMBER);
+			const bool bIsNearlyZeroT2T1 = FMath::IsNearlyZero(T2T1, UE_KINDA_SMALL_NUMBER);
+			const bool bIsNearlyZeroT3T2 = FMath::IsNearlyZero(T3T2, UE_KINDA_SMALL_NUMBER);
+			const bool bIsNearlyZeroT2T0 = FMath::IsNearlyZero(T2T0, UE_KINDA_SMALL_NUMBER);
+			const bool bIsNearlyZeroT3T1 = FMath::IsNearlyZero(T3T1, UE_KINDA_SMALL_NUMBER);
+
+			for (int SampleIndex = 0; SampleIndex < NumSamplesPerSegment; ++SampleIndex)
+			{
+				const float ParametricDistance = float(SampleIndex) / float(NumSamplesPerSegment);
+
+				const float T = FMath::Lerp(T1, T2, ParametricDistance);
+
+				const FVector A1 = bIsNearlyZeroT1T0 ? P0 : (T1 - T) / T1T0 * P0 + (T - T0) / T1T0 * P1;
+				const FVector A2 = bIsNearlyZeroT2T1 ? P1 : (T2 - T) / T2T1 * P1 + (T - T1) / T2T1 * P2;
+				const FVector A3 = bIsNearlyZeroT3T2 ? P2 : (T3 - T) / T3T2 * P2 + (T - T2) / T3T2 * P3;
+				const FVector B1 = bIsNearlyZeroT2T0 ? A1 : (T2 - T) / T2T0 * A1 + (T - T0) / T2T0 * A2;
+				const FVector B2 = bIsNearlyZeroT3T1 ? A2 : (T3 - T) / T3T1 * A2 + (T - T1) / T3T1 * A3;
+
+				CVs[BoneIndex * NumSamplesPerSegment + SampleIndex] = bIsNearlyZeroT2T1 ? B1 : (T2 - T) / T2T1 * B1 + (T - T1) / T2T1 * B2;
+				Params[BoneIndex * NumSamplesPerSegment + SampleIndex] = float(BoneIndices[BoneIndex]) + ParametricDistance;
+			}
+		}
+
+		CVs.Last() = BoneTransforms[BoneIndices.Last()].GetTranslation();
+		Params.Last() = BoneIndices.Last();
+	}
+
+	void FCatmullRomSpline::ClosestParam(const FVector& Position, float& Param, float& ClosestDistance) const
+	{
+		ClosestDistance = UE_MAX_FLT;
+		Param = 0;
+
+		for (int32 SampleIndex = 0; SampleIndex < CVs.Num(); ++SampleIndex)
+		{
+			double Distance = FVector::DistSquared(Position, CVs[SampleIndex]);
+
+			if (Distance < ClosestDistance)
+			{
+				Param = Params[SampleIndex];
+				ClosestDistance = Distance;
+			}
+		}
+	}
+
+	void BuildSplines(const TObjectPtr<USkeletalMesh> SkeletalMesh, const TArray<FString>& RootBones, int32 SamplesPerSegment, TArray<UE::Groom::Private::FCatmullRomSpline>& Splines)
+	{
+		if (SkeletalMesh)
+		{
+			if (RootBones.Num() > 0)
+			{
+				for (int32 BoneIndex = 0; BoneIndex < RootBones.Num(); ++BoneIndex)
+				{
+					int32 RootBoneIndex = SkeletalMesh->GetRefSkeleton().FindBoneIndex(FName(*RootBones[BoneIndex]));
+					if (RootBoneIndex != INDEX_NONE)
+					{
+						UE::Groom::Private::FCatmullRomSpline Spline;
+						Spline.Init(SkeletalMesh->GetRefSkeleton(), RootBoneIndex, SamplesPerSegment);
+						Splines.Add(Spline);
+					}
+				}
+			}
+
+			// Use skeleton root if no root bones specified or found
+			if (Splines.IsEmpty())
+			{
+				UE::Groom::Private::FCatmullRomSpline Spline;
+				Spline.Init(SkeletalMesh->GetRefSkeleton(), 0, SamplesPerSegment);
+				Splines.Add(Spline);
+			}
+		}
+	}
+	
+	template<typename FacadeType>
+	static void BuildSplineSkinningWeights(FManagedArrayCollection& GroomCollection, int GroupIndex, const TObjectPtr<USkeletalMesh> SkeletalMesh, int LODIndex,
+		const TArray<UE::Groom::Private::FCatmullRomSpline>& Splines, FCollectionAttributeKey SplineParamKey, FCollectionAttributeKey SplineBoneKey)
+	{
+		FacadeType GroomFacade(GroomCollection);
+
+		if (GroomFacade.IsValid() && (GroupIndex == INDEX_NONE || GroupIndex < GroomFacade.GetNumObjects()))
+		{
+			// Set groom meshes
+			TManagedArray<TObjectPtr<UObject>>& ObjectSkeletalMeshes =
+				GroomCollection.AddAttribute<TObjectPtr<UObject>>(UE::Groom::FGroomGuidesFacade::ObjectSkeletalMeshesAttribute, FacadeType::ObjectsGroup);
+
+			TManagedArray<int32>& ObjectMeshLODs =
+				GroomCollection.AddAttribute<int32>(UE::Groom::FGroomGuidesFacade::ObjectMeshLODsAttribute, FacadeType::ObjectsGroup);
+
+			if (GroupIndex == INDEX_NONE)
+			{
+				for (int32 LocalIndex = 0; LocalIndex < ObjectSkeletalMeshes.Num(); ++LocalIndex)
+				{
+					ObjectSkeletalMeshes[LocalIndex] = SkeletalMesh;
+					ObjectMeshLODs[LocalIndex] = LODIndex;
+				}
+			}
+			else if (ObjectSkeletalMeshes.IsValidIndex(GroupIndex) && ObjectMeshLODs.IsValidIndex(GroupIndex))
+			{
+				ObjectSkeletalMeshes[GroupIndex] = SkeletalMesh;
+				ObjectMeshLODs[GroupIndex] = LODIndex;
+			}
+
+			// Compute skinning weights
+			TManagedArray<float>& VerticesSplineParams = GroomCollection.AddAttribute<float>(FName(SplineParamKey.Attribute), FName(SplineParamKey.Group));
+			TManagedArray<FIntVector2>& VerticesSplineBones = GroomCollection.AddAttribute<FIntVector2>(FName(SplineBoneKey.Attribute), FName(SplineBoneKey.Group));
+			TManagedArray<float>& VerticesSplineSegments = GroomCollection.AddAttribute<float>(FName("SplineSegment_Build"), FName(SplineParamKey.Group));
+			TManagedArray<float>& VerticesSplineSegmentParams = GroomCollection.AddAttribute<float>(FName("SplineSegmentParam_Build"), FName(SplineParamKey.Group));
+
+			const TArray<FVector3f>& Positions = GroomFacade.GetPointRestPositions();
+			const TArray<int32>& CurvePointOffsets = GroomFacade.GetCurvePointOffsets();
+
+			const int FirstCurveIndex = GroupIndex > 0 ? GroomFacade.GetObjectCurveOffsets()[GroupIndex - 1] : 0;
+			const int LastCurveIndex = GroupIndex != INDEX_NONE ? GroomFacade.GetObjectCurveOffsets()[GroupIndex] : GroomFacade.GetNumCurves();
+			const int32 NumCurves = LastCurveIndex - FirstCurveIndex;
+
+			ParallelFor(NumCurves, [&FirstCurveIndex, &Positions, &CurvePointOffsets, &Splines, &VerticesSplineParams, &VerticesSplineBones, &VerticesSplineSegments, &VerticesSplineSegmentParams](int32 CurveIndex)
+			{
+				TArray<int32> SplineCounts;
+				SplineCounts.Init(0, Splines.Num());
+					
+				const int32 FirstPoint = FirstCurveIndex + CurveIndex > 0 ? CurvePointOffsets[FirstCurveIndex + CurveIndex - 1] : 0;
+				const int32 LastPoint = CurvePointOffsets[FirstCurveIndex + CurveIndex];
+
+				for (int32 PointIndex = FirstPoint; PointIndex < LastPoint; ++PointIndex)
+				{
+					const FVector Position(Positions[PointIndex]);
+					int32 ClosestSplineIndex = 0;
+					float MinParam=0, MinDistance = UE_MAX_FLT;
+
+					for (int32 SplineIndex = 0; SplineIndex < Splines.Num(); ++SplineIndex)
+					{
+						float Param, Distance;
+						Splines[SplineIndex].ClosestParam(Position, Param, Distance);
+
+						if (Distance < MinDistance)
+						{
+							MinParam = Param;
+							MinDistance = Distance;
+							ClosestSplineIndex = SplineIndex;
+
+						}
+					}
+					
+					VerticesSplineParams[PointIndex * 2] = MinParam;
+					VerticesSplineParams[PointIndex * 2 + 1] = MinParam;
+					VerticesSplineBones[PointIndex * 2] = Splines[ClosestSplineIndex].Bones;
+					VerticesSplineBones[PointIndex * 2 + 1] = Splines[ClosestSplineIndex].Bones;
+
+					++SplineCounts[ClosestSplineIndex];
+				}
+				
+				// Ensure that all CVs on a strand are bound to the same spline
+				if (Algo::Count(SplineCounts, 0) < SplineCounts.Num() - 1)
+				{
+					// Use the most commonly bound spline
+					int32 SplineIndex = Algo::MaxElement(SplineCounts) - &SplineCounts[0];
+
+					for (int32 PointIndex = FirstPoint; PointIndex < LastPoint; ++PointIndex)
+					{
+						if (VerticesSplineBones[PointIndex * 2] != Splines[SplineIndex].Bones)
+						{
+							const FVector Position(Positions[PointIndex]);
+							float MinDistance = UE_MAX_FLT;
+							float Param;
+							Splines[SplineIndex].ClosestParam(Position, Param, MinDistance);
+							
+							VerticesSplineParams[PointIndex * 2] = Param;
+							VerticesSplineParams[PointIndex * 2 + 1] = Param;
+							VerticesSplineBones[PointIndex * 2] = Splines[SplineIndex].Bones;
+							VerticesSplineBones[PointIndex * 2 + 1] = Splines[SplineIndex].Bones;
+						}
+					}
+				}
+			});
+		}
+	}
+
+	template<typename FacadeType>
+	static void ConvertFromSplineToLinearWeights(FManagedArrayCollection & GroomCollection, int32 GroomGroupIndex, const FCollectionAttributeKey & SplineParamKey, const FCollectionAttributeKey& SplineBonesKey,
+		const FCollectionAttributeKey & BoneIndicesKey, const FCollectionAttributeKey & BoneWeightsKey)
+	{
+		const TManagedArray<float>* SplineParams = GroomCollection.FindAttributeTyped<float>(FName(SplineParamKey.Attribute), FName(SplineParamKey.Group));
+		const TManagedArray<FIntVector2>* SplineBones = GroomCollection.FindAttributeTyped<FIntVector2>(FName(SplineBonesKey.Attribute), FName(SplineBonesKey.Group));
+
+		if (SplineParams && SplineBones)
+		{
+			FacadeType GroomFacade(GroomCollection);
+
+			if (GroomFacade.IsValid() && (GroomGroupIndex == INDEX_NONE || GroomGroupIndex < GroomFacade.GetNumObjects()))
+			{
+				bool ExistingBoneIndices = !GroomCollection.HasAttribute(FName(BoneIndicesKey.Attribute), FName(BoneIndicesKey.Group));
+				TManagedArray<FIntVector4>& BoneIndices = GroomCollection.AddAttribute<FIntVector4>(FName(BoneIndicesKey.Attribute), FName(BoneIndicesKey.Group));
+				TManagedArray<FVector4f>& BoneWeights = GroomCollection.AddAttribute<FVector4f>(FName(BoneWeightsKey.Attribute), FName(BoneWeightsKey.Group));
+
+				if (!ExistingBoneIndices)
+				{
+					for (int32 VertexIndex = 0; VertexIndex < GroomFacade.GetNumPoints() * 2; ++VertexIndex)
+					{
+						BoneIndices[VertexIndex] = FIntVector4(-1, -1, -1, -1);
+						BoneWeights[VertexIndex] = FVector4f(0, 0, 0, 0);
+					}
+				}
+
+				TArray<int32> GroupIndices;
+				if (GroomGroupIndex == INDEX_NONE)
+				{
+					for (int32 LocalIndex = 0; LocalIndex < GroomFacade.GetNumObjects(); ++LocalIndex)
+					{
+						GroupIndices.Add(LocalIndex);
+					}
+				}
+				else
+				{
+					GroupIndices.Add(GroomGroupIndex);
+				}
+
+				for (int32 GroupIndex : GroupIndices)
+				{
+					const int FirstCurveIndex = GroupIndex > 0 ? GroomFacade.GetObjectCurveOffsets()[GroupIndex - 1] : 0;
+					const int LastCurveIndex = GroomFacade.GetObjectCurveOffsets()[GroupIndex];
+					const int32 FirstVertex = FirstCurveIndex > 0 ? GroomFacade.GetCurvePointOffsets()[FirstCurveIndex - 1] * 2 : 0;
+					const int32 LastVertex = GroomFacade.GetCurvePointOffsets()[LastCurveIndex-1] * 2;
+
+					for (int32 VertexIndex = FirstVertex; VertexIndex < LastVertex; ++VertexIndex)
+					{
+						float SplineParam = (*SplineParams)[VertexIndex];
+
+						int Segment = trunc(SplineParam);
+						float SegmentParam = SplineParam - Segment;
+
+						if (Segment < (*SplineBones)[VertexIndex][1])
+						{
+							BoneIndices[VertexIndex] = FIntVector4(Segment, Segment + 1, -1 , -1);
+							BoneWeights[VertexIndex] = FVector4f(1.0 - SegmentParam, SegmentParam, 0, 0);
+						}
+						else
+						{
+							BoneIndices[VertexIndex] = FIntVector4(Segment, -1, -1, -1);
+							BoneWeights[VertexIndex] = FVector4f(1, 0, 0, 0);
+						}
+					}
+				}
+			}
+		}
+	}
+
+	template<typename FacadeType>
+	static void ConvertFromLinearToSplineWeights(FManagedArrayCollection & GroomCollection, int32 GroomGroupIndex, const FCollectionAttributeKey & SplineParamKey, const FCollectionAttributeKey & SplineBonesKey)
+	{
+		const TManagedArray<FIntVector4>* BoneIndices = GroomCollection.FindAttributeTyped<FIntVector4>(FName(UE::Groom::FGroomGuidesFacade::PointBoneIndicesAttribute), FacadeType::VerticesGroup);
+		const TManagedArray<FVector4f>* BoneWeights = GroomCollection.FindAttributeTyped<FVector4f>(FName(UE::Groom::FGroomGuidesFacade::PointBoneWeightsAttribute), FacadeType::VerticesGroup);
+		
+		TManagedArray<TObjectPtr<UObject>>* ObjectSkeletalMeshes =
+			GroomCollection.FindAttributeTyped<TObjectPtr<UObject>>(UE::Groom::FGroomGuidesFacade::ObjectSkeletalMeshesAttribute, FacadeType::ObjectsGroup);
+		TManagedArray<int32>* ObjectMeshLODs =
+			GroomCollection.FindAttributeTyped<int32>(UE::Groom::FGroomGuidesFacade::ObjectMeshLODsAttribute, FacadeType::ObjectsGroup);
+
+		if (BoneIndices && BoneWeights && ObjectSkeletalMeshes && ObjectMeshLODs && !SplineParamKey.Attribute.IsEmpty() && !SplineBonesKey.Attribute.IsEmpty())
+		{
+			FacadeType GroomFacade(GroomCollection);
+
+			if (GroomFacade.IsValid() && (GroomGroupIndex == INDEX_NONE || GroomGroupIndex < GroomFacade.GetNumObjects()))
+			{
+				TManagedArray<float>& VertexSplineParams = GroomCollection.AddAttribute<float>(FName(SplineParamKey.Attribute), FName(SplineParamKey.Group));
+				TManagedArray<FIntVector2>& VertexSplineBones = GroomCollection.AddAttribute<FIntVector2>(FName(SplineBonesKey.Attribute), FName(SplineBonesKey.Group));
+
+				TArray<int32> GroupIndices;
+				if (GroomGroupIndex == INDEX_NONE)
+				{
+					for (int32 LocalIndex = 0; LocalIndex < GroomFacade.GetNumObjects(); ++LocalIndex)
+					{
+						GroupIndices.Add(LocalIndex);
+					}
+				}
+				else
+				{
+					GroupIndices.Add(GroomGroupIndex);
+				}
+
+				for (int32 GroupIndex : GroupIndices)
+				{
+					if (ObjectSkeletalMeshes->IsValidIndex(GroupIndex) && ObjectMeshLODs->IsValidIndex(GroupIndex))
+					{
+						if (TObjectPtr<USkeletalMesh> SkeletalMesh = Cast<USkeletalMesh>((*ObjectSkeletalMeshes)[GroupIndex]))
+						{
+							const int FirstCurveIndex = GroupIndex > 0 ? GroomFacade.GetObjectCurveOffsets()[GroupIndex - 1] : 0;
+							const int LastCurveIndex = GroomFacade.GetObjectCurveOffsets()[GroupIndex];
+							const int32 FirstVertex = FirstCurveIndex > 0 ? GroomFacade.GetCurvePointOffsets()[FirstCurveIndex - 1] * 2 : 0;
+							const int32 LastVertex = GroomFacade.GetCurvePointOffsets()[LastCurveIndex-1]*2;
+
+							// Find all used bones for this group						
+							TSet<int32> UsedBoneIndices;
+							for (int32 VertexIndex = FirstVertex; VertexIndex < LastVertex; ++VertexIndex)
+							{
+								for (int32 InfluenceIndex = 0; InfluenceIndex < 4; ++InfluenceIndex)
+								{
+									FIntVector4 DebugInfluences = (*BoneIndices)[VertexIndex];
+									if ((*BoneIndices)[VertexIndex][InfluenceIndex] != INDEX_NONE)
+									{
+										UsedBoneIndices.Add((*BoneIndices)[VertexIndex][InfluenceIndex]);
+									}
+								}
+							}
+							
+							// Find root bones
+							TArray<int32> RootBones;
+							TMap<int32, int32> BoneSplineMap;
+							
+							const FReferenceSkeleton& Skeleton = SkeletalMesh->GetRefSkeleton();
+							for (int32 BoneIndex : UsedBoneIndices)
+							{
+								int32 ParentIndex = Skeleton.GetParentIndex(BoneIndex);
+								if (ParentIndex == INDEX_NONE || !UsedBoneIndices.Contains(ParentIndex))
+								{
+									RootBones.Add(BoneIndex);
+								}
+							}
+							
+							// Find end bones
+							TArray<int32> EndBones;						
+							TArray<int32> Children;
+
+							for (int32 SplineIndex = 0; SplineIndex < RootBones.Num(); ++SplineIndex)
+							{
+								BoneSplineMap.Add(RootBones[SplineIndex], SplineIndex);
+								int32 BoneIndex = RootBones[SplineIndex];
+
+								while (Skeleton.GetDirectChildBones(BoneIndex, Children))
+								{
+									BoneSplineMap.Add(Children[0], SplineIndex);
+									BoneIndex = Children[0];
+									if (Children.Num() > 1)
+									{
+										// TODO: Add support for branching splines
+									}
+								}
+
+								EndBones.Add(BoneIndex);
+							}
+							
+							
+							for (int32 VertexIndex = FirstVertex; VertexIndex < LastVertex; ++VertexIndex)
+							{
+								const FIntVector4& Influences = (*BoneIndices)[VertexIndex];
+								const FVector4f& Weights = (*BoneWeights)[VertexIndex];
+
+								int32 MaxWeightIndex = INDEX_NONE;
+								float MaxWeight = 0;
+								for (int32 BoneIndex = 0; BoneIndex < 4; ++BoneIndex)
+								{
+									if (Weights[BoneIndex] > MaxWeight)
+									{
+										MaxWeight = Weights[BoneIndex];
+										MaxWeightIndex = BoneIndex;
+									}
+								}
+								
+								if (MaxWeight > 0)
+								{
+									float ParentWeight = 0, ChildWeight = 0;
+									
+									for (int32 InfluenceIndex = 0; InfluenceIndex < 4; ++InfluenceIndex)
+									{
+										if (Influences[MaxWeightIndex] > VertexSplineBones[VertexIndex][0] && Influences[InfluenceIndex] == Influences[MaxWeightIndex]-1)
+										{
+											ParentWeight = Weights[InfluenceIndex];
+										}
+
+										if (Influences[MaxWeightIndex] < VertexSplineBones[VertexIndex][1] && Influences[InfluenceIndex] == Influences[MaxWeightIndex] + 1)
+										{
+											ChildWeight = Weights[InfluenceIndex];
+										}
+									}
+									
+									if (ParentWeight > ChildWeight)
+									{
+										int Segment = Influences[MaxWeightIndex] - 1; 
+										float SegmentParam = MaxWeight / (MaxWeight + ParentWeight);
+										VertexSplineParams[VertexIndex] = Segment + SegmentParam;
+									}
+									else
+									{
+										int Segment = Influences[MaxWeightIndex];
+										float SegmentParam = (1 - MaxWeight) / (MaxWeight + ChildWeight);
+										VertexSplineParams[VertexIndex] = Segment + SegmentParam;
+									}		
+									
+									const int32 SplineIndex = BoneSplineMap[Influences[MaxWeightIndex]];
+									VertexSplineBones[VertexIndex] = FIntVector2(RootBones[SplineIndex], EndBones[SplineIndex]);
+								}
+								else if (!RootBones.IsEmpty())
+								{
+									// Rigidly attach to first spline root if there are no weights for this vertex
+									VertexSplineBones[VertexIndex] = FIntVector2(RootBones[0], EndBones[0]);
+									VertexSplineParams[VertexIndex] = 0;
+								}
+							}
+						}
+					}
+				}
+			}
+		}
+	}
+}	
+
+
+FBuildGroomSplineSkinWeightsNode::FBuildGroomSplineSkinWeightsNode(const UE::Dataflow::FNodeParameters& InParam, FGuid InGuid)
+	: FDataflowNode(InParam, InGuid)
+{
+	RegisterInputConnection(&Collection);
+	RegisterInputConnection(&SkeletalMesh);
+
+	RegisterOutputConnection(&Collection, &Collection);
+	RegisterOutputConnection(&SplineParamKey);
+	RegisterOutputConnection(&SplineBoneKey);
+}
+
+void FBuildGroomSplineSkinWeightsNode::Evaluate(UE::Dataflow::FContext& Context, const FDataflowOutput* Out) const
+{
+	if (Out->IsA<FManagedArrayCollection>(&Collection))
+	{
+		FManagedArrayCollection GroomCollection = GetValue<FManagedArrayCollection>(Context, &Collection);
+		FCollectionAttributeKey ParamKey = GetSplineParamKey();
+		FCollectionAttributeKey BoneKey = GetSplineBonesKey();
+
+		if (!ParamKey.Attribute.IsEmpty() && !BoneKey.Attribute.IsEmpty() && SkeletalMesh && SkeletalMesh->IsValidLODIndex(LODIndex))
+		{
+			TArray<UE::Groom::Private::FCatmullRomSpline> Splines;
+			UE::Groom::Private::BuildSplines(SkeletalMesh, RootBones, SamplesPerSegment, Splines);
+
+			if (CurvesType == EGroomCollectionType::Guides)
+			{
+				UE::Groom::Private::BuildSplineSkinningWeights<UE::Groom::FGroomGuidesFacade>(GroomCollection, GroupIndex, SkeletalMesh, LODIndex, Splines, ParamKey, BoneKey);
+			}
+			else
+			{
+				UE::Groom::Private::BuildSplineSkinningWeights<UE::Groom::FGroomStrandsFacade>(GroomCollection, GroupIndex, SkeletalMesh, LODIndex, Splines, ParamKey, BoneKey);
+			}
+
+		}
+
+		SetValue(Context, MoveTemp(GroomCollection), &Collection);
+	}
+	else if (Out->IsA<FCollectionAttributeKey>(&SplineParamKey))
+	{
+		SetValue(Context, GetSplineParamKey(), &SplineParamKey);
+	}
+	else if (Out->IsA<FCollectionAttributeKey>(&SplineBoneKey))
+	{
+		SetValue(Context, GetSplineBonesKey(), &SplineBoneKey);
+	}
+}
+
+FCollectionAttributeKey FBuildGroomSplineSkinWeightsNode::GetSplineParamKey() const
+{
+	FCollectionAttributeKey Key;
+	Key.Group = (CurvesType == EGroomCollectionType::Guides) ? UE::Groom::FGroomGuidesFacade::VerticesGroup.ToString() :
+		UE::Groom::FGroomStrandsFacade::VerticesGroup.ToString();
+	Key.Attribute = SplineParamName;
+	return Key;
+}
+
+FCollectionAttributeKey FBuildGroomSplineSkinWeightsNode::GetSplineBonesKey() const
+{
+	FCollectionAttributeKey Key;
+	Key.Group = (CurvesType == EGroomCollectionType::Guides) ? UE::Groom::FGroomGuidesFacade::VerticesGroup.ToString() :
+		UE::Groom::FGroomStrandsFacade::VerticesGroup.ToString();
+	Key.Attribute = SplineBonesParamName;
+	return Key;
+}
+
+TArray<UE::Dataflow::FRenderingParameter> FBuildGroomSplineSkinWeightsNode::GetRenderParametersImpl() const
+{
+	if (CurvesType == EGroomCollectionType::Guides)
+	{
+		return { {TEXT("GuidesRender"), FName("FGroomCollection"), {TEXT("Collection")}} };
+	}
+	else
+	{
+		return { {TEXT("StrandsRender"), FName("FGroomCollection"), {TEXT("Collection")}} };
+	}
+}
+
+FConvertSplineToLinearSkinWeightsNode::FConvertSplineToLinearSkinWeightsNode(const UE::Dataflow::FNodeParameters& InParam, FGuid InGuid)
+	: FDataflowNode(InParam, InGuid)
+{
+	RegisterInputConnection(&Collection);
+	RegisterInputConnection(&SplineParamKey);
+	RegisterInputConnection(&SplineBonesKey);
+	RegisterOutputConnection(&Collection, &Collection);
+	RegisterOutputConnection(&BoneIndicesKey);
+	RegisterOutputConnection(&BoneWeightsKey);
+}
+
+void FConvertSplineToLinearSkinWeightsNode::Evaluate(UE::Dataflow::FContext& Context, const FDataflowOutput* Out) const
+{
+	if (Out->IsA<FManagedArrayCollection>(&Collection))
+	{
+		FManagedArrayCollection GroomCollection = GetValue<FManagedArrayCollection>(Context, &Collection);
+		FCollectionAttributeKey SplineParam = GetSplineParamKey(Context);
+		FCollectionAttributeKey SplineBones = GetSplineBonesKey(Context);
+
+		if (!SplineParam.Attribute.IsEmpty() && !SplineBones.Attribute.IsEmpty())
+		{
+			if (CurvesType == EGroomCollectionType::Guides)
+			{
+				UE::Groom::Private::ConvertFromSplineToLinearWeights<UE::Groom::FGroomGuidesFacade>(GroomCollection, GroupIndex, SplineParam, SplineBones, GetBoneIndicesKey(), GetBoneWeightsKey());
+			}
+			else
+			{
+				UE::Groom::Private::ConvertFromSplineToLinearWeights<UE::Groom::FGroomStrandsFacade>(GroomCollection, GroupIndex, SplineParam, SplineBones, GetBoneIndicesKey(), GetBoneWeightsKey());
+			}
+		}
+		SetValue(Context, MoveTemp(GroomCollection), &Collection);
+	}
+	else if (Out->IsA<FCollectionAttributeKey>(&BoneIndicesKey))
+	{
+		SetValue(Context, GetBoneIndicesKey(), &BoneIndicesKey);
+	}
+	else if (Out->IsA<FCollectionAttributeKey>(&BoneWeightsKey))
+	{
+		SetValue(Context, GetBoneWeightsKey(), &BoneWeightsKey);
+	}
+}
+
+FCollectionAttributeKey FConvertSplineToLinearSkinWeightsNode::GetBoneIndicesKey() const
+{
+	FCollectionAttributeKey Key;
+	Key.Group = (CurvesType == EGroomCollectionType::Guides) ? UE::Groom::FGroomGuidesFacade::VerticesGroup.ToString() :
+		UE::Groom::FGroomStrandsFacade::VerticesGroup.ToString();
+	Key.Attribute = UE::Groom::FGroomGuidesFacade::PointBoneIndicesAttribute.ToString();
+	return Key;
+}
+
+FCollectionAttributeKey FConvertSplineToLinearSkinWeightsNode::GetBoneWeightsKey() const
+{
+	FCollectionAttributeKey Key;
+	Key.Group = (CurvesType == EGroomCollectionType::Guides) ? UE::Groom::FGroomGuidesFacade::VerticesGroup.ToString() :
+		UE::Groom::FGroomStrandsFacade::VerticesGroup.ToString();
+	Key.Attribute = UE::Groom::FGroomGuidesFacade::PointBoneWeightsAttribute.ToString();
+	return Key;
+}
+
+
+
+FCollectionAttributeKey FConvertSplineToLinearSkinWeightsNode::GetSplineParamKey(UE::Dataflow::FContext& Context) const
+{
+	FCollectionAttributeKey Key = GetValue(Context, &SplineParamKey, SplineParamKey);
+
+	// If nothing set use default value
+	if (Key.Attribute.IsEmpty() && Key.Group.IsEmpty())
+	{
+		Key.Group = (CurvesType == EGroomCollectionType::Guides) ? UE::Groom::FGroomGuidesFacade::VerticesGroup.ToString() :
+			UE::Groom::FGroomStrandsFacade::VerticesGroup.ToString();
+		Key.Attribute = SplineParamName;
+	}
+
+	return Key;
+}
+
+FCollectionAttributeKey FConvertSplineToLinearSkinWeightsNode::GetSplineBonesKey(UE::Dataflow::FContext& Context) const
+{
+	FCollectionAttributeKey Key = GetValue(Context, &SplineBonesKey, SplineBonesKey);
+
+	if (Key.Attribute.IsEmpty() && Key.Group.IsEmpty())
+	{
+		Key.Group = (CurvesType == EGroomCollectionType::Guides) ? UE::Groom::FGroomGuidesFacade::VerticesGroup.ToString() :
+			UE::Groom::FGroomStrandsFacade::VerticesGroup.ToString();
+		Key.Attribute = SplineBonesName;
+	}
+
+	return Key;
+}
+
+TArray<UE::Dataflow::FRenderingParameter> FConvertSplineToLinearSkinWeightsNode::GetRenderParametersImpl() const
+{
+	if (CurvesType == EGroomCollectionType::Guides)
+	{
+		return { {TEXT("GuidesRender"), FName("FGroomCollection"), {TEXT("Collection")}} };
+	}
+	else
+	{
+		return { {TEXT("StrandsRender"), FName("FGroomCollection"), {TEXT("Collection")}} };
+	}
+}
+
+FConvertLinearToSplineSkinWeightsNode::FConvertLinearToSplineSkinWeightsNode(const UE::Dataflow::FNodeParameters& InParam, FGuid InGuid)
+	: FDataflowNode(InParam, InGuid)
+{
+	RegisterInputConnection(&Collection);
+	RegisterOutputConnection(&Collection, &Collection);
+	RegisterOutputConnection(&SplineParamKey);
+	RegisterOutputConnection(&SplineBoneKey);
+}
+
+void FConvertLinearToSplineSkinWeightsNode::Evaluate(UE::Dataflow::FContext& Context, const FDataflowOutput* Out) const
+{
+	using namespace UE::Dataflow;
+
+	if (Out->IsA<FManagedArrayCollection>(&Collection))
+	{
+		FManagedArrayCollection GroomCollection = GetValue<FManagedArrayCollection>(Context, &Collection);
+
+		if (CurvesType == EGroomCollectionType::Guides)
+		{
+			UE::Groom::Private::ConvertFromLinearToSplineWeights<UE::Groom::FGroomGuidesFacade>(GroomCollection, GroupIndex, GetSplineParamKey(), GetSplineBonesKey());
+		}
+		else
+		{
+			UE::Groom::Private::ConvertFromLinearToSplineWeights<UE::Groom::FGroomStrandsFacade>(GroomCollection, GroupIndex, GetSplineParamKey(), GetSplineBonesKey());
+		}
+
+		SetValue(Context, MoveTemp(GroomCollection), &Collection);
+	}
+	else if (Out->IsA<FCollectionAttributeKey>(&SplineParamKey))
+	{
+		SetValue(Context, GetSplineParamKey(), &SplineParamKey);
+	}
+	else if (Out->IsA<FCollectionAttributeKey>(&SplineBoneKey))
+	{
+		SetValue(Context, GetSplineBonesKey(), &SplineBoneKey);
+	}
+
+}
+
+TArray<UE::Dataflow::FRenderingParameter> FConvertLinearToSplineSkinWeightsNode::GetRenderParametersImpl() const
+{
+	if (CurvesType == EGroomCollectionType::Guides)
+	{
+		return { {TEXT("GuidesRender"), FName("FGroomCollection"), {TEXT("Collection")}} };
+	}
+	else
+	{
+		return { {TEXT("StrandsRender"), FName("FGroomCollection"), {TEXT("Collection")}} };
+	}
+}
+
+FCollectionAttributeKey FConvertLinearToSplineSkinWeightsNode::GetSplineParamKey() const
+{
+	FCollectionAttributeKey Key;
+	Key.Group = (CurvesType == EGroomCollectionType::Guides) ? UE::Groom::FGroomGuidesFacade::VerticesGroup.ToString() :
+		UE::Groom::FGroomStrandsFacade::VerticesGroup.ToString();
+	Key.Attribute = SplineParamName;
+	return Key;
+}
+
+FCollectionAttributeKey FConvertLinearToSplineSkinWeightsNode::GetSplineBonesKey() const
+{
+	FCollectionAttributeKey Key;
+	Key.Group = (CurvesType == EGroomCollectionType::Guides) ? UE::Groom::FGroomGuidesFacade::VerticesGroup.ToString() :
+		UE::Groom::FGroomStrandsFacade::VerticesGroup.ToString();
+	Key.Attribute = SplineBonesName;
+	return Key;
+}
+
diff --git a/Engine/Plugins/Runtime/HairStrands/Source/HairStrandsDataflow/Private/HairStrandsDataflowModule.cpp b/Engine/Plugins/Runtime/HairStrands/Source/HairStrandsDataflow/Private/HairStrandsDataflowModule.cpp
--- a/Engine/Plugins/Runtime/HairStrands/Source/HairStrandsDataflow/Private/HairStrandsDataflowModule.cpp
+++ b/Engine/Plugins/Runtime/HairStrands/Source/HairStrandsDataflow/Private/HairStrandsDataflowModule.cpp
@@ -10,6 +10,7 @@
 #include "GetGroomAttributesNodes.h"
 #include "ResampleGuidesPointsNode.h"
 #include "BuildGroomSkinningNodes.h"
+#include "BuildGroomSplineSkinningNode.h"
 #include "GenerateGuidesCurvesNode.h"
 #include "SmoothGuidesCurvesNode.h"
 #include "GroomDataflowRendering.h"
@@ -58,7 +59,10 @@ void FHairStrandsDataflowModule::StartupModule()
 	DATAFLOW_NODE_REGISTER_CREATION_FACTORY(FSmoothGuidesCurvesDataflowNode);
 	DATAFLOW_NODE_REGISTER_CREATION_FACTORY(FAttachGuidesRootsDataflowNode);
 	DATAFLOW_NODE_REGISTER_CREATION_FACTORY(FBuildGuidesLODsDataflowNode);
-	
+	DATAFLOW_NODE_REGISTER_CREATION_FACTORY(FBuildGroomSplineSkinWeightsNode);
+	DATAFLOW_NODE_REGISTER_CREATION_FACTORY(FConvertLinearToSplineSkinWeightsNode);
+	DATAFLOW_NODE_REGISTER_CREATION_FACTORY(FConvertSplineToLinearSkinWeightsNode);
+
 	UE::Groom::RegisterRenderingCallbacks();
 	
 	FDataflowAddScalarVertexPropertyCallbackRegistry::Get().RegisterCallbacks(
diff --git a/Engine/Plugins/Runtime/HairStrands/Source/HairStrandsDataflow/Public/BuildGroomSplineSkinningNode.h b/Engine/Plugins/Runtime/HairStrands/Source/HairStrandsDataflow/Public/BuildGroomSplineSkinningNode.h
new file mode 100644
--- /dev/null
+++ b/Engine/Plugins/Runtime/HairStrands/Source/HairStrandsDataflow/Public/BuildGroomSplineSkinningNode.h
@@ -0,0 +1,197 @@
+// Copyright Epic Games, Inc. All Rights Reserved.
+
+#pragma once 
+
+#include "CoreMinimal.h"
+#include "GroomCollectionFacades.h"
+#include "Dataflow/DataflowCore.h"
+#include "Dataflow/DataflowConnectionTypes.h"
+#include "Engine/SkeletalMesh.h"
+
+#include "BuildGroomSplineSkinningNode.generated.h"
+
+
+/** Build spline skinning data from skeletal mesh */
+USTRUCT(meta = (Experimental, DataflowGroom))
+struct FBuildGroomSplineSkinWeightsNode : public FDataflowNode
+{
+	GENERATED_BODY()
+
+	DATAFLOW_NODE_DEFINE_INTERNAL(FBuildGroomSplineSkinWeightsNode, "BuildSplineSkinWeights", "Groom", "")
+
+public:
+	FBuildGroomSplineSkinWeightsNode(const UE::Dataflow::FNodeParameters& InParam, FGuid InGuid = FGuid::NewGuid());
+
+	//~ Begin FDataflowNode interface
+	virtual void Evaluate(UE::Dataflow::FContext& Context, const FDataflowOutput* Out) const override;
+	virtual TArray<UE::Dataflow::FRenderingParameter> GetRenderParametersImpl() const override;
+	//~ End FDataflowNode interface
+
+	/** Managed array collection to be used to store datas */
+	UPROPERTY(meta = (DataflowInput, DataflowOutput, DataflowPassthrough  = "Collection"))
+	FManagedArrayCollection Collection;
+
+	/** Spline Skinning Parameter key */
+	UPROPERTY(meta = (DisplayName = "Spline Parameter Key", DataflowOutput))
+	FCollectionAttributeKey SplineParamKey;
+
+	/** Spline Bones key containing root and end spline bone. To be used in other nodes if necessary */
+	UPROPERTY(meta = (DataflowOutput))
+	FCollectionAttributeKey SplineBoneKey;
+
+
+	/** SkeletalMesh used for spline skinning. Will be stored onto the groom asset */
+	UPROPERTY(EditAnywhere, Category="Skeletal Mesh", meta = (DataFlowInput))
+	TObjectPtr<USkeletalMesh> SkeletalMesh = nullptr;
+
+	/** Root bones to be used for spline skinning. Uses all bones if empty. Note that branches in the skeleton is currently not supported. */
+	UPROPERTY(EditAnywhere, Category = "Skeletal Mesh")
+	TArray<FString> RootBones;
+
+	/** LOD used to build skinning weights */
+	UPROPERTY(EditAnywhere, Category="Skeletal Mesh", meta = (DisplayName = "LOD Index"))
+	int32 LODIndex = 0;
+
+	/** Number of spline samples per bone segment. */
+	UPROPERTY(EditAnywhere, Category = "Skeletal Mesh")
+	int32 SamplesPerSegment = 64;
+
+
+	/** Groom group index to build skinning weights for. -1 will build all groups */
+	UPROPERTY(EditAnywhere, Category="Groom Groups")
+	int32 GroupIndex = INDEX_NONE;
+
+	/** Type of curves to build skinning for (guides/strands) */
+	UPROPERTY(EditAnywhere, Category = "Groom Groups")
+	EGroomCollectionType CurvesType = EGroomCollectionType::Guides;
+
+	/** Spline skinning parameter attribute name */
+	UPROPERTY(EditAnywhere, Category = "Vertex Attributes", meta = (DisplayName = "Spline Parameter"))
+	FString SplineParamName = "SplineParam";
+
+	/** Spline bones attribute name. Used for storing root and end spline bone for each vertex. */
+	UPROPERTY(EditAnywhere, Category = "Vertex Attributes")
+	FString SplineBonesParamName = "SplineBones";
+
+//private:
+	FCollectionAttributeKey GetSplineParamKey() const;
+	FCollectionAttributeKey GetSplineBonesKey() const;
+};
+
+/** Convert spline skinning data to linear data */
+USTRUCT(meta = (Experimental, DataflowGroom))
+struct FConvertSplineToLinearSkinWeightsNode : public FDataflowNode
+{
+	GENERATED_BODY()
+
+	DATAFLOW_NODE_DEFINE_INTERNAL(FConvertSplineToLinearSkinWeightsNode, "SplineToLinearSkinWeights", "Groom", "")
+
+public:
+	FConvertSplineToLinearSkinWeightsNode(const UE::Dataflow::FNodeParameters& InParam, FGuid InGuid = FGuid::NewGuid());
+
+	//~ Begin FDataflowNode interface
+	virtual void Evaluate(UE::Dataflow::FContext& Context, const FDataflowOutput* Out) const override;
+	virtual TArray<UE::Dataflow::FRenderingParameter> GetRenderParametersImpl() const override;
+	//~ End FDataflowNode interface
+
+	/** Managed array collection to be used to store datas */
+	UPROPERTY(meta = (DataflowInput, DataflowOutput, DataflowPassthrough = "Collection"))
+	FManagedArrayCollection Collection;
+
+	/** Spline skinning parameter attribute name */
+	UPROPERTY(EditAnywhere, Category = "Vertex Attributes", meta = (DisplayName = "Spline Parameter"))
+	FString SplineParamName = "SplineParam";
+
+	/** Spline bones attribute name. Used for storing root and end spline bone for each vertex. */
+	UPROPERTY(EditAnywhere, Category = "Vertex Attributes", meta = (DisplayName = "Spline Bones"))
+	FString SplineBonesName = "SplineBones";
+
+	/**Spline skinning parameter key */
+	UPROPERTY(meta = (DataflowInput, DisplayName = "Spline Parameter Key"))
+	FCollectionAttributeKey SplineParamKey;
+
+	/** Spline bones key. Used for storing root and end spline bone for each vertex. */
+	UPROPERTY(meta = (DataflowInput))
+	FCollectionAttributeKey SplineBonesKey;
+
+	/** Groom group index for converting skinning weight. -1 will convert all the groups */
+	UPROPERTY(EditAnywhere, Category = "Groom Groups")
+	int32 GroupIndex = INDEX_NONE;
+
+
+	/** Linear skinning bone indices key */
+	UPROPERTY(meta = (DataflowOutput))
+	FCollectionAttributeKey BoneIndicesKey;
+
+	/** Linear skinning bone weights key */
+	UPROPERTY(meta = (DataflowOutput))
+	FCollectionAttributeKey BoneWeightsKey;
+
+	/** Type of curves to convert skinning for (guides/strands) */
+	UPROPERTY(EditAnywhere, Category = "Vertex Attributes")
+	EGroomCollectionType CurvesType = EGroomCollectionType::Guides;
+
+private:
+	/** Get bone indices key */
+	FCollectionAttributeKey GetBoneIndicesKey() const;
+
+	/** Get bone weights key */
+	FCollectionAttributeKey GetBoneWeightsKey() const;
+
+	/** Get spline parameter key */
+	FCollectionAttributeKey GetSplineParamKey(UE::Dataflow::FContext& Context) const;
+
+	/** Get spline bones key */
+	FCollectionAttributeKey GetSplineBonesKey(UE::Dataflow::FContext& Context) const;
+};
+
+
+/** Convert linear skinning data to spline skinning data */
+USTRUCT(meta = (Experimental, DataflowGroom))
+struct FConvertLinearToSplineSkinWeightsNode : public FDataflowNode
+{
+	GENERATED_BODY()
+
+	DATAFLOW_NODE_DEFINE_INTERNAL(FConvertLinearToSplineSkinWeightsNode, "LinearToSplineSkinWeights", "Groom", "")
+
+public:
+	FConvertLinearToSplineSkinWeightsNode(const UE::Dataflow::FNodeParameters& InParam, FGuid InGuid = FGuid::NewGuid());
+
+	//~ Begin FDataflowNode interface
+	virtual void Evaluate(UE::Dataflow::FContext& Context, const FDataflowOutput* Out) const override;
+	virtual TArray<UE::Dataflow::FRenderingParameter> GetRenderParametersImpl() const override;
+	//~ End FDataflowNode interface
+
+	/** Managed array collection to be used to store datas */
+	UPROPERTY(meta = (DataflowInput, DataflowOutput, DataflowPassthrough = "Collection"))
+	FManagedArrayCollection Collection;
+
+	/** Groom group index for converting skinning weights. -1 will convert all groups */
+	UPROPERTY(EditAnywhere, Category = "Groom Groups")
+	int32 GroupIndex = INDEX_NONE;
+
+	/** Spline Skinning Parameter key */
+	UPROPERTY(meta = (DisplayName = "Spline Parameter Key", DataflowOutput))
+	FCollectionAttributeKey SplineParamKey;
+
+	/** Spline Bones key containing root and end spline bone. */
+	UPROPERTY(meta = (DataflowOutput))
+	FCollectionAttributeKey SplineBoneKey;
+
+	/** Spline skinning parameter attribute name. */
+	UPROPERTY(EditAnywhere, Category = "Vertex Attributes", meta = (DisplayName = "Spline Parameter"))
+	FString SplineParamName = "SplineParam";
+
+	/** Spline bones parameter attribute name. Contains root and end spline bone for each vertex. */
+	UPROPERTY(EditAnywhere, Category = "Vertex Attributes", meta = (DisplayName = "Spline Bones"))
+	FString SplineBonesName = "SplineBones";
+
+	/** Type of curves to convert skinning for (guides/strands) */
+	UPROPERTY(EditAnywhere, Category = "Vertex Attributes")
+	EGroomCollectionType CurvesType = EGroomCollectionType::Guides;
+
+private:
+	FCollectionAttributeKey GetSplineParamKey() const;
+	FCollectionAttributeKey GetSplineBonesKey() const;
+};
+
diff --git a/Engine/Plugins/Runtime/HairStrands/Source/HairStrandsDeformer/Private/DeformerGroomInterfaceUtils.h b/Engine/Plugins/Runtime/HairStrands/Source/HairStrandsDeformer/Private/DeformerGroomInterfaceUtils.h
--- a/Engine/Plugins/Runtime/HairStrands/Source/HairStrandsDeformer/Private/DeformerGroomInterfaceUtils.h
+++ b/Engine/Plugins/Runtime/HairStrands/Source/HairStrandsDeformer/Private/DeformerGroomInterfaceUtils.h
@@ -20,18 +20,18 @@ namespace UE::Groom::Private
 
 	FORCEINLINE USkeletalMeshComponent* GetGroupSkelMesh(const UGroomComponent* GroomComponent, const int32 GroupIndex, int32& MeshLOD)
 	{
-		if(GroomComponent)
+		if (GroomComponent)
 		{
 			TArray<UActorComponent*> ActorComponents;
-			if(AActor* RootActor = GroomComponent->GetAttachmentRootActor())
+			if (AActor* RootActor = GroomComponent->GetAttachmentRootActor())
 			{
 				RootActor->GetComponents(USkeletalMeshComponent::StaticClass(), ActorComponents);
 			}
-			for(UActorComponent* ActorComponent : ActorComponents)
+			for (UActorComponent* ActorComponent : ActorComponents)
 			{
-				if(USkeletalMeshComponent* SkeletalMeshComponent = Cast<USkeletalMeshComponent>(ActorComponent))
+				if (USkeletalMeshComponent* SkeletalMeshComponent = Cast<USkeletalMeshComponent>(ActorComponent))
 				{
-					if(GroomComponent->GroomAsset && (SkeletalMeshComponent->GetSkeletalMeshAsset() ==
+					if (GroomComponent->GroomAsset && (SkeletalMeshComponent->GetSkeletalMeshAsset() ==
 						GroomComponent->GroomAsset->GetDataflowSettings().GetSkeletalMesh(GroupIndex)))
 					{
 						MeshLOD = GroomComponent->GroomAsset->GetDataflowSettings().GetMeshLOD(GroupIndex);
@@ -47,9 +47,9 @@ namespace UE::Groom::Private
 	FORCEINLINE void GroomComponentsToInstances(const TArray<const UGroomComponent*>& GroomComponents, TArray<const FHairGroupInstance*>& GroupInstances)
  	{
  		GroupInstances.Reset();
- 		for(const UGroomComponent* GroomComponent : GroomComponents)
+ 		for (const UGroomComponent* GroomComponent : GroomComponents)
  		{
- 			if(GroomComponent)
+ 			if (GroomComponent)
  			{
  				const uint32 NumGroups = GroomComponent->GetGroupCount();
  				for (uint32 GroupIndex = 0; GroupIndex < NumGroups; ++GroupIndex)
@@ -63,32 +63,34 @@ namespace UE::Groom::Private
  		}
  	}
 	FORCEINLINE void GroomComponentsToSkelmeshes(const TArray<const UGroomComponent*>& GroomComponents, TArray<const FSkeletalMeshObject*>& SkeletalMeshes, TArray<FMatrix44f>& SkeletalTransforms,
-		TArray<TArray<FMatrix44f>>& BonesRefToLocals, TArray<const FHairGroupInstance*>& GroupInstances)
+		TArray<TArray<FMatrix44f>>& BonesRefToLocals, TArray<TArray<FMatrix44f>>& BindTransforms, TArray<const FHairGroupInstance*>& GroupInstances)
 	{
 		SkeletalMeshes.Reset();
 		SkeletalTransforms.Reset();
 		BonesRefToLocals.Reset();
+		BindTransforms.Reset();
 		GroupInstances.Reset();
 
 		TArray<FTransform> BonesTransforms;
 		TArray<FMatrix44f> RefToLocals;
+		TArray<FMatrix44f> RefBases;
 		int32 MeshLOD = INDEX_NONE;
 		
-		for(const UGroomComponent* GroomComponent : GroomComponents)
+		for (const UGroomComponent* GroomComponent : GroomComponents)
 		{
-			if(GroomComponent)
+			if (GroomComponent)
 			{
 				const uint32 NumGroups = GroomComponent->GetGroupCount();
 				for (uint32 GroupIndex = 0; GroupIndex < NumGroups; ++GroupIndex)
 				{
-					if(GroomComponent->IsDeformationEnable(GroupIndex) && GroomComponent->GetGroupInstance(GroupIndex))
+					if (GroomComponent->IsDeformationEnable(GroupIndex) && GroomComponent->GetGroupInstance(GroupIndex))
 					{
-						if(USkeletalMeshComponent* SkelMesh =  GetGroupSkelMesh(GroomComponent, GroupIndex, MeshLOD))
+						if (USkeletalMeshComponent* SkelMesh =  GetGroupSkelMesh(GroomComponent, GroupIndex, MeshLOD))
 						{
-							if(USkinnedAsset* SkinnedAsset = SkelMesh->GetSkinnedAsset())
+							if (USkinnedAsset* SkinnedAsset = SkelMesh->GetSkinnedAsset())
 							{
 								const FSkeletalMeshRenderData* const RenderData = SkinnedAsset->GetResourceForRendering();
-								if(RenderData && RenderData->LODRenderData.IsValidIndex(MeshLOD))
+								if (RenderData && RenderData->LODRenderData.IsValidIndex(MeshLOD))
 								{
 									// Get the matching active indices used for skinning
 									TArray<FBoneIndexType> GuidesBones = RenderData->LODRenderData[MeshLOD].ActiveBoneIndices;
@@ -102,20 +104,22 @@ namespace UE::Groom::Private
 									SkelMesh->GetSkinnedAsset()->FillComponentSpaceTransforms(
 										SkelMesh->GetBoneSpaceTransforms(), GuidesBones, BonesTransforms);
 
-									// Fill ref to local matrices
+									// Fill ref to local and bind matrices
 									RefToLocals.Init(FMatrix44f::Identity, RefBasesInvMatrix->Num());
+									RefBases.Init(FMatrix44f::Identity, RefBasesInvMatrix->Num());
 								
 									for (int32 BoneIndex = 0; BoneIndex < RefToLocals.Num(); ++BoneIndex)
 									{
 										RefToLocals[BoneIndex] = (*RefBasesInvMatrix)[BoneIndex] * (FMatrix44f)BonesTransforms[BoneIndex].ToMatrixWithScale();
+										RefBases[BoneIndex] = (*RefBasesInvMatrix)[BoneIndex].Inverse();
 									}
 							
 									const FTransform& BonesTransform = SkelMesh->GetComponentTransform();
 									const FTransform& GroupTransform = GroomComponent->GetGroupInstance(GroupIndex)->GetCurrentLocalToWorld();
-									//SkelMesh->GetCurrentRefToLocalMatrices(RefToLocals, SkelMesh->MeshObject->GetLOD(), &RequiredExtraBoneIndices);//MeshLOD);
 		
 									SkeletalTransforms.Add(FMatrix44f(BonesTransform.ToMatrixWithScale() * GroupTransform.ToInverseMatrixWithScale()));
 									BonesRefToLocals.Add(RefToLocals);
+									BindTransforms.Add(RefBases);
 									SkeletalMeshes.Add(SkelMesh->MeshObject);
 									GroupInstances.Add(GroomComponent->GetGroupInstance(GroupIndex));
 								}
@@ -130,13 +134,13 @@ namespace UE::Groom::Private
 	FORCEINLINE void GatherGroomComponents(const UActorComponent* ActorComponent, TArray<const UGroomComponent*>& GroomComponents)
 	{
 		GroomComponents.Reset();
-		if(const UGroomComponent* GroomComponent = Cast<UGroomComponent>(ActorComponent))
+		if (const UGroomComponent* GroomComponent = Cast<UGroomComponent>(ActorComponent))
 		{
 			GroomComponents.Add(GroomComponent);
 		}
-		else if(const UGroomSolverComponent* GroomSolver = Cast<UGroomSolverComponent>(ActorComponent))
+		else if (const UGroomSolverComponent* GroomSolver = Cast<UGroomSolverComponent>(ActorComponent))
 		{
-			for(const TObjectPtr<UGroomComponent>& SolverGroom : GroomSolver->GetGroomComponents())
+			for (const TObjectPtr<UGroomComponent>& SolverGroom : GroomSolver->GetGroomComponents())
 			{
 				GroomComponents.Add(SolverGroom);
 			}
@@ -151,12 +155,13 @@ namespace UE::Groom::Private
 		GroomComponentsToInstances(GroomComponents, GroupInstances);
 	}
 
-	FORCEINLINE void GatherGroupSkelmeshes(const UActorComponent* ActorComponent, TArray<const FSkeletalMeshObject*>& SkeletalMeshes, TArray<FMatrix44f>& SkeletalTransforms, TArray<TArray<FMatrix44f>>& BonesRefToLocals, TArray<const FHairGroupInstance*>& GroupInstances)
+	FORCEINLINE void GatherGroupSkelmeshes(const UActorComponent* ActorComponent, TArray<const FSkeletalMeshObject*>& SkeletalMeshes, TArray<FMatrix44f>& SkeletalTransforms,
+		TArray<TArray<FMatrix44f>>& BonesRefToLocals, TArray<TArray<FMatrix44f>>& BindTransforms, TArray<const FHairGroupInstance*>& GroupInstances)
 	{
 		TArray<const UGroomComponent*> GroomComponents;
 		GatherGroomComponents(ActorComponent, GroomComponents);
 
-		GroomComponentsToSkelmeshes(GroomComponents, SkeletalMeshes, SkeletalTransforms, BonesRefToLocals, GroupInstances);
+		GroomComponentsToSkelmeshes(GroomComponents, SkeletalMeshes, SkeletalTransforms, BonesRefToLocals, BindTransforms, GroupInstances);
 	}
 
 	template<typename DataType>
diff --git a/Engine/Plugins/Runtime/HairStrands/Source/HairStrandsDeformer/Private/DeformerGroomMeshesRead.cpp b/Engine/Plugins/Runtime/HairStrands/Source/HairStrandsDeformer/Private/DeformerGroomMeshesRead.cpp
--- a/Engine/Plugins/Runtime/HairStrands/Source/HairStrandsDeformer/Private/DeformerGroomMeshesRead.cpp
+++ b/Engine/Plugins/Runtime/HairStrands/Source/HairStrandsDeformer/Private/DeformerGroomMeshesRead.cpp
@@ -29,6 +29,7 @@ TArray<FOptimusCDIPinDefinition> UOptimusGroomMeshesReadDataInterface::GetPinDef
 	// Bones buffers
 	Defs.Add({ "NumMeshesBones",   "ReadNumBones", false, "NumMeshesBones" });
 	Defs.Add({ "BoneTransformMatrix", "ReadBoneTransformMatrix", MeshesBones,   "ReadNumBones", true, "BoneTransformMatrix"  });
+	Defs.Add({ "BindTransformMatrix", "ReadBindTransformMatrix", MeshesBones,   "ReadNumBones", true, "BindTransformMatrix" });
 
 	return Defs;
 }
@@ -51,6 +52,11 @@ void UOptimusGroomMeshesReadDataInterface::GetSupportedInputs(TArray<FShaderFunc
 		.SetName(TEXT("ReadBoneTransformMatrix"))
 		.AddReturnType(EShaderFundamentalType::Float, 3, 4)
 		.AddParam(EShaderFundamentalType::Uint);
+
+	OutFunctions.AddDefaulted_GetRef()
+		.SetName(TEXT("ReadBindTransformMatrix"))
+		.AddReturnType(EShaderFundamentalType::Float, 3, 4)
+		.AddParam(EShaderFundamentalType::Uint);
 }
 
 BEGIN_SHADER_PARAMETER_STRUCT(FOptimusGroomMeshesReadParameters, )
@@ -101,7 +107,7 @@ FComputeDataProviderRenderProxy* UOptimusGroomMeshesReadDataProvider::GetRenderP
 
 FOptimusGroomMeshesReadDataProviderProxy::FOptimusGroomMeshesReadDataProviderProxy(UMeshComponent* MeshComponent)
 {
-	UE::Groom::Private::GatherGroupSkelmeshes(MeshComponent, SkeletalMeshObjects, SkeletalMeshTransforms, BonesRefToLocals, GroupInstances);
+	UE::Groom::Private::GatherGroupSkelmeshes(MeshComponent, SkeletalMeshObjects, SkeletalMeshTransforms, BonesRefToLocals, BindTransforms, GroupInstances);
 }
 
 bool FOptimusGroomMeshesReadDataProviderProxy::IsValid(FValidationData const& InValidationData) const
@@ -118,12 +124,11 @@ bool FOptimusGroomMeshesReadDataProviderProxy::IsValid(FValidationData const& In
 	return true;
 }
 
-static void CreateInternalBuffers(const FSkeletalMeshObject* SkeletalMeshObject, const FMatrix44f& SkeletalMeshTransform, const TArray<FMatrix44f>& BoneRefToLocals,
+static void CreateInternalBuffers(const FSkeletalMeshObject* SkeletalMeshObject, const FMatrix44f& SkeletalMeshTransform, const TArray<FMatrix44f>& BoneRefToLocals, const TArray<FMatrix44f>& BindTransforms,
 	const FHairGroupInstance* GroupInstance, FRDGBuilder& GraphBuilder, TArray<FRDGBufferSRVRef>& AttributeResources)
 {
 	if(SkeletalMeshObject && GroupInstance)
 	{
-		//const TArray<FMatrix44f>& BoneRefToLocals = SkeletalMeshObject->GetReferenceToLocalMatrices();
 		const FString BufferName("Hair.Deformer.Strands.BoneTransformMatrices");
 		if(FHairStrandsDeformedResource* DeformedResource = GroupInstance->Strands.DeformedResource)
 		{
@@ -139,7 +144,7 @@ static void CreateInternalBuffers(const FSkeletalMeshObject* SkeletalMeshObject,
 			else
 			{
 				TransientBuffer = GraphBuilder.CreateBuffer(FRDGBufferDesc::CreateStructuredDesc(
-					sizeof(FVector4f), BoneRefToLocals.Num() * 3 ), *BufferName);
+					sizeof(FVector4f), BoneRefToLocals.Num() * 3 * 2), *BufferName);
 				ExternalBuffer = &DeformedResource->ExternalBuffers.Add(BufferName);
 				
 				ExternalBuffer->Buffer = GraphBuilder.ConvertToExternalBuffer(TransientBuffer);
@@ -149,8 +154,9 @@ static void CreateInternalBuffers(const FSkeletalMeshObject* SkeletalMeshObject,
 				const int32 NumBones = BoneRefToLocals.Num();
 	
 				TArray<FVector4f> BoneVectors;
-				BoneVectors.Init(FVector4f::Zero(), NumBones * 3);
+				BoneVectors.Init(FVector4f::Zero(), NumBones * 3 * 2);
 
+				// Bone transforms
 				int32 VectorIndex = 0;
 				for(int32 BoneIndex = 0; BoneIndex < NumBones; ++BoneIndex)
 				{
@@ -159,6 +165,16 @@ static void CreateInternalBuffers(const FSkeletalMeshObject* SkeletalMeshObject,
 					 
 					VectorIndex += 3;
 				}
+
+				// Bind transforms
+				for (int32 BoneIndex = 0; BoneIndex < NumBones; ++BoneIndex)
+				{
+					const FMatrix44f LocalMatrix = BindTransforms[BoneIndex] * SkeletalMeshTransform;
+					LocalMatrix.To3x4MatrixTranspose(static_cast<float*>(&BoneVectors[VectorIndex].X));
+
+					VectorIndex += 3;
+				}
+
 				GraphBuilder.QueueBufferUpload(TransientBuffer, BoneVectors.GetData(),
 					sizeof(FVector4f) * BoneVectors.Num(), ERDGInitialDataFlags::None);
 			}
@@ -176,7 +192,7 @@ void FOptimusGroomMeshesReadDataProviderProxy::AllocateResources(FRDGBuilder& Gr
 	int32 SkelIndex = 0;
 	for (const FSkeletalMeshObject* SkeletalMeshObject : SkeletalMeshObjects)
 	{
-		CreateInternalBuffers(SkeletalMeshObject, SkeletalMeshTransforms[SkelIndex], BonesRefToLocals[SkelIndex], GroupInstances[SkelIndex], GraphBuilder, BoneMatricesResources);
+		CreateInternalBuffers(SkeletalMeshObject, SkeletalMeshTransforms[SkelIndex], BonesRefToLocals[SkelIndex], BindTransforms[SkelIndex], GroupInstances[SkelIndex], GraphBuilder, BoneMatricesResources);
 		++SkelIndex;
 	}
 }
diff --git a/Engine/Plugins/Runtime/HairStrands/Source/HairStrandsDeformer/Private/DeformerGroomMeshesRead.h b/Engine/Plugins/Runtime/HairStrands/Source/HairStrandsDeformer/Private/DeformerGroomMeshesRead.h
--- a/Engine/Plugins/Runtime/HairStrands/Source/HairStrandsDeformer/Private/DeformerGroomMeshesRead.h
+++ b/Engine/Plugins/Runtime/HairStrands/Source/HairStrandsDeformer/Private/DeformerGroomMeshesRead.h
@@ -85,6 +85,9 @@ private:
 	/** Bones refs to locals matrices */
 	TArray<TArray<FMatrix44f>> BonesRefToLocals;
 
+	/** Bind Transforms */
+	TArray<TArray<FMatrix44f>> BindTransforms;
+
 	/** List of instances (invocations) used in that data interface */
 	TArray<const FHairGroupInstance*> GroupInstances;
 	
