Fix groom guide-cache import

diff --git a/Engine/Plugins/Runtime/HairStrands/Source/HairStrandsEditor/Private/GroomCacheImporter.cpp b/Engine/Plugins/Runtime/HairStrands/Source/HairStrandsEditor/Private/GroomCacheImporter.cpp
--- a/Engine/Plugins/Runtime/HairStrands/Source/HairStrandsEditor/Private/GroomCacheImporter.cpp
+++ b/Engine/Plugins/Runtime/HairStrands/Source/HairStrandsEditor/Private/GroomCacheImporter.cpp
@@ -177,7 +177,8 @@ TArray<UGroomCache*> FGroomCacheImporter::ImportGroomCache(const FString& Source
 						// Since the groom cache needs to be independent of this settings (since the data is serialized directly into the asset
 						// and not rebuilt from a description), it does not contains these extra control points.
 						// We account for them for validation here
-						const uint32 ExtraControlPointCount = GetHairStrandsUsesTriangleStrips() ? CacheInputDatas[GroupIndex].Strands.GetNumCurves() : 0u;
+						const bool bHasExtraPoint = GroupPlatformData[GroupIndex].Strands.BulkData.HasFlags(FHairStrandsBulkData::DataFlags_HasExtraEndingPoint);
+						const uint32 ExtraControlPointCount = bHasExtraPoint ? CacheInputDatas[GroupIndex].Strands.GetNumCurves() : 0u;
 
 						const bool bCurveCountMatch = CacheInputDatas[GroupIndex].Strands.GetNumCurves() == GroupPlatformData[GroupIndex].Strands.BulkData.GetNumCurves();
 						const bool bPointCountMatch = CacheInputDatas[GroupIndex].Strands.GetNumPoints() + ExtraControlPointCount == GroupPlatformData[GroupIndex].Strands.BulkData.GetNumPoints();
@@ -200,7 +201,8 @@ TArray<UGroomCache*> FGroomCacheImporter::ImportGroomCache(const FString& Source
 							bGuidesOnly = true;
 						}
 
-						const uint32 ExtraControlPointCount = GetHairStrandsUsesTriangleStrips() ? CacheInputDatas[GroupIndex].Guides.GetNumCurves() : 0u;
+						const bool bHasExtraPoint = GroupPlatformData[GroupIndex].Guides.BulkData.HasFlags(FHairStrandsBulkData::DataFlags_HasExtraEndingPoint);
+						const uint32 ExtraControlPointCount = bHasExtraPoint ? CacheInputDatas[GroupIndex].Guides.GetNumCurves() : 0u;
 
 						const bool bCurveCountMatch = CacheInputDatas[GroupIndex].Guides.GetNumCurves() == GroupPlatformData[GroupIndex].Guides.BulkData.GetNumCurves();
 						const bool bPointCountMatch = CacheInputDatas[GroupIndex].Guides.GetNumPoints() + ExtraControlPointCount == GroupPlatformData[GroupIndex].Guides.BulkData.GetNumPoints();
diff --git a/Engine/Plugins/Runtime/HairStrands/Source/HairStrandsEditor/Private/GroomImportOptionsWindow.cpp b/Engine/Plugins/Runtime/HairStrands/Source/HairStrandsEditor/Private/GroomImportOptionsWindow.cpp
--- a/Engine/Plugins/Runtime/HairStrands/Source/HairStrandsEditor/Private/GroomImportOptionsWindow.cpp
+++ b/Engine/Plugins/Runtime/HairStrands/Source/HairStrandsEditor/Private/GroomImportOptionsWindow.cpp
@@ -147,6 +147,13 @@ FGroomImportStatus GetGroomImportStatus(UGroomHairGroupsPreview* InDescription,
 void SGroomImportOptionsWindow::UpdateStatus(UGroomHairGroupsPreview* Description) const
 {
 	ImportStatus = GetGroomImportStatus(Description, GroomCacheImportOptions, GroupsMapping);
+	// If a groom is a groom guide-cache and has groom cache compatible asset, remove the 'NoCurve' flags 
+	// which would cause the asset to be reported as invalid
+	if (EnumHasAnyFlags(ImportStatus.Status, EHairDescriptionStatus::NoCurve) &&
+		EnumHasAnyFlags(ImportStatus.Status, EHairDescriptionStatus::GuidesOnlyCompatible))
+	{
+		EnumRemoveFlags(ImportStatus.Status, EHairDescriptionStatus::NoCurve);
+	}
 }
 
 FText GetGroomImportStatusText(FGroomImportStatus In, bool bAddPrefix)
