Geometry facade for grooms and cards + new rendering + use of curve selection

diff --git a/Engine/Plugins/Runtime/HairStrands/Source/HairStrandsDataflow/Private/AttachGuidesRootsNode.cpp b/Engine/Plugins/Runtime/HairStrands/Source/HairStrandsDataflow/Private/AttachGuidesRootsNode.cpp
--- a/Engine/Plugins/Runtime/HairStrands/Source/HairStrandsDataflow/Private/AttachGuidesRootsNode.cpp
+++ b/Engine/Plugins/Runtime/HairStrands/Source/HairStrandsDataflow/Private/AttachGuidesRootsNode.cpp
@@ -1,60 +1,95 @@
 // Copyright Epic Games, Inc. All Rights Reserved.
 
 #include "AttachGuidesRootsNode.h"
-#include "GroomCollectionFacades.h"
 #include "Dataflow/DataflowObjectInterface.h"
+#include "GeometryCollection/Facades/CollectionCurveFacade.h"
+#include "GeometryCollection/Facades/CollectionVertexBoneWeightsFacade.h"
 
 #include UE_INLINE_GENERATED_CPP_BY_NAME(AttachGuidesRootsNode)
 
-namespace UE::Groom::Private
+void FAttachGuidesRootsDataflowNode::Evaluate(UE::Dataflow::FContext& Context, const FDataflowOutput* Out) const
+{}
+
+void FAttachCurveRootsDataflowNode::Evaluate(UE::Dataflow::FContext& Context, const FDataflowOutput* Out) const
 {
-	FCollectionAttributeKey GetKinematicWeightsKey()
+	if (Out->IsA<FManagedArrayCollection>(&Collection))
+	{
+		FManagedArrayCollection GeometryCollection = GetValue<FManagedArrayCollection>(Context, &Collection);
+		GeometryCollection::Facades::FCollectionCurveGeometryFacade CurvesFacade(GeometryCollection);
+
+		if(CurvesFacade.IsValid())
+		{
+			GeometryCollection::Facades::FVertexBoneWeightsFacade SkinningFacade(GeometryCollection, false);
+			SkinningFacade.DefineSchema();
+			
+			const bool bValidSelection = CurveSelection.IsValidForCollection(CurvesFacade.GetManagedArrayCollection());
+			
+			const FDataflowCurveSelection& DataflowSelection = GetValue<FDataflowCurveSelection>(Context, &CurveSelection);
+			
+			ParallelFor(CurvesFacade.GetNumCurves(), [&SkinningFacade, &CurvesFacade, DataflowSelection, bValidSelection](int32 CurveIndex)
+			{
+				int32 PointIndex = (CurveIndex == 0) ? 0 : CurvesFacade.GetCurvePointOffsets()[CurveIndex-1];
+				if(!bValidSelection || (bValidSelection && DataflowSelection.IsSelected(CurveIndex)))
+				{
+					SkinningFacade.ModifyKinematicWeight(2*PointIndex,1.0f);
+					SkinningFacade.ModifyKinematicWeight(2*PointIndex+1, 1.0f);
+				}
+					
+				PointIndex += 1;
+				if(!bValidSelection || (bValidSelection && DataflowSelection.IsSelected(CurveIndex)))
+				{
+					SkinningFacade.ModifyKinematicWeight(2*PointIndex, 1.0f);
+					SkinningFacade.ModifyKinematicWeight(2*PointIndex+1, 1.0f);
+				}
+			}, EParallelForFlags::None);
+		}
+		
+		SetValue(Context, MoveTemp(GeometryCollection), &Collection);
+	}
+	else if (Out->IsA<FCollectionAttributeKey>(&KinematicWeightsKey))
 	{
 		FCollectionAttributeKey Key;
-		Key.Group =  UE::Groom::FGroomGuidesFacade::VerticesGroup.ToString();
-		Key.Attribute = UE::Groom::FGroomGuidesFacade::PointKinematicWeightsAttribute.ToString();
-		return Key;
+		Key.Group =  FGeometryCollection::VerticesGroup.ToString();
+		Key.Attribute = GeometryCollection::Facades::FVertexBoneWeightsFacade::KinematicWeightAttributeName.ToString();
+		
+		SetValue(Context, Key, &KinematicWeightsKey);
 	}
 }
 
-void FAttachGuidesRootsDataflowNode::Evaluate(UE::Dataflow::FContext& Context, const FDataflowOutput* Out) const
+void FBuildCurveWeightsDataflowNode::Evaluate(UE::Dataflow::FContext& Context, const FDataflowOutput* Out) const
 {
 	if (Out->IsA<FManagedArrayCollection>(&Collection))
 	{
-		FManagedArrayCollection GroomCollection = GetValue<FManagedArrayCollection>(Context, &Collection);
-		
-		UE::Groom::FGroomGuidesFacade GuidesFacade(GroomCollection);
+		FManagedArrayCollection GeometryCollection = GetValue<FManagedArrayCollection>(Context, &Collection);
+		GeometryCollection::Facades::FCollectionCurveGeometryFacade CurvesFacade(GeometryCollection);
 
-		if(GuidesFacade.IsValid())
+		if(CurvesFacade.IsValid())
 		{
-			TArray<float> KinematicWeights = GuidesFacade.GetPointKinematicWeights();
+			const FCollectionAttributeKey WeightKey = GetValue<FCollectionAttributeKey>(Context, &WeightsAttribute);
+			const FDataflowCurveSelection& DataflowSelection = GetValue<FDataflowCurveSelection>(Context, &CurveSelection);
 			
-			const int32 LocalIndex = GroupIndex;
-			ParallelFor(GuidesFacade.GetNumPoints(), [&KinematicWeights, &GuidesFacade, LocalIndex](int32 PointIndex)
+			if(TManagedArray<float>* WeightsValue = GeometryCollection.FindAttribute<float>(FName(WeightKey.Attribute), FName(WeightKey.Group)))
 			{
-				const int32 CurveIndex = GuidesFacade.GetPointCurveIndices()[PointIndex];
-				const int32 ObjectIndex = GuidesFacade.GetCurveObjectIndices()[CurveIndex];
-
-				if(LocalIndex == INDEX_NONE || ObjectIndex == LocalIndex)
+				const bool bValidSelection = CurveSelection.IsValidForCollection(CurvesFacade.GetManagedArrayCollection());;
+				
+				ParallelFor(CurvesFacade.GetNumPoints(), [this, &WeightsValue, &CurvesFacade, &DataflowSelection, bValidSelection](int32 PointIndex)
 				{
-					const int32 PointOffset = (CurveIndex == 0) ? 0 : GuidesFacade.GetCurvePointOffsets()[CurveIndex-1];
-					const int32 LocalPoint = PointIndex - PointOffset;
-
-					if(LocalPoint < 2)
+					const int32 CurveIndex = CurvesFacade.GetPointCurveIndices()[PointIndex];
+					if(!bValidSelection || (bValidSelection && DataflowSelection.IsSelected(CurveIndex)))
 					{
-						KinematicWeights[2*PointIndex] = 1.0f;
-						KinematicWeights[2*PointIndex+1] = 1.0f;
+						const int32 NextPoint = CurvesFacade.GetCurvePointOffsets()[CurveIndex];
+						const int32 PrevPoint = (CurveIndex == 0) ? 0 : CurvesFacade.GetCurvePointOffsets()[CurveIndex-1];
+						
+						const float PointCoord = static_cast<float>(PointIndex - PrevPoint) / (NextPoint - 1 - PrevPoint);
+
+						(*WeightsValue)[2 * PointIndex] = CurveWeights.GetRichCurveConst()->Eval(PointCoord);
+						(*WeightsValue)[2 * PointIndex + 1] = (*WeightsValue)[2 * PointIndex];
 					}
-				}
-			}, EParallelForFlags::None);
-			GuidesFacade.SetPointKinematicWeights(KinematicWeights);
+				}, EParallelForFlags::None);
+			}
 		}
 		
-		SetValue(Context, MoveTemp(GroomCollection), &Collection);
-	}
-	else if (Out->IsA<FCollectionAttributeKey>(&KinematicWeightsKey))
-	{
-		SetValue(Context, UE::Groom::Private::GetKinematicWeightsKey(), &KinematicWeightsKey);
+		SetValue(Context, MoveTemp(GeometryCollection), &Collection);
 	}
 }
 
diff --git a/Engine/Plugins/Runtime/HairStrands/Source/HairStrandsDataflow/Private/BuildGroomSkinningNodes.cpp b/Engine/Plugins/Runtime/HairStrands/Source/HairStrandsDataflow/Private/BuildGroomSkinningNodes.cpp
--- a/Engine/Plugins/Runtime/HairStrands/Source/HairStrandsDataflow/Private/BuildGroomSkinningNodes.cpp
+++ b/Engine/Plugins/Runtime/HairStrands/Source/HairStrandsDataflow/Private/BuildGroomSkinningNodes.cpp
@@ -1,13 +1,14 @@
 // Copyright Epic Games, Inc. All Rights Reserved.
 
 #include "BuildGroomSkinningNodes.h"
-#include "GroomCollectionFacades.h"
 #include "MeshDescriptionToDynamicMesh.h"
 #include "SkeletalMeshAttributes.h"
 #include "SkeletalMeshLODRenderDataToDynamicMesh.h"
 #include "DynamicMesh/DynamicMesh3.h"
 #include "Dataflow/DataflowObjectInterface.h"
 #include "DynamicMesh/MeshTransforms.h"
+#include "GeometryCollection/Facades/CollectionVertexBoneWeightsFacade.h"
+#include "GeometryCollection/Facades/CollectionMeshFacade.h"
 #include "Operations/TransferBoneWeights.h"
 #include "Rendering/SkeletalMeshRenderData.h"
 
@@ -15,7 +16,7 @@
 
 namespace UE::Groom::Private
 {
-	FORCEINLINE bool SkeletalMeshToDynamicMesh(const USkeletalMesh* SkeletalMesh, const int32 LodIndex, UE::Geometry::FDynamicMesh3& ToDynamicMesh)
+	static bool SkeletalMeshToDynamicMesh(const USkeletalMesh* SkeletalMesh, const int32 LodIndex, UE::Geometry::FDynamicMesh3& ToDynamicMesh)
 	{
 		if (SkeletalMesh->HasMeshDescription(LodIndex))
 		{
@@ -54,175 +55,54 @@ namespace UE::Groom::Private
 
 		return true;
 	}
-
-	template<typename FacadeType>
-	static void BuildSkinningData(FManagedArrayCollection& GroomCollection, const int32 GroupIndex, const TObjectPtr<USkeletalMesh> SkeletalMesh, const int32 LODIndex,
-		const FTransform& RelativeTransform)
+		
+	static void BuildGeometrySkinWeights(GeometryCollection::Facades::FCollectionMeshFacade& MeshFacade, GeometryCollection::Facades::FVertexBoneWeightsFacade& SkinningFacade,
+		const FDataflowVertexSelection& VertexSelection, const TObjectPtr<USkeletalMesh> SkeletalMesh, const int32 LODIndex, const FTransform& RelativeTransform)
 	{
-		FacadeType GroomFacade(GroomCollection);
-
-		if(GroomFacade.IsValid())
+		if(MeshFacade.IsValid() && SkeletalMesh && SkeletalMesh->IsValidLODIndex(LODIndex))
 		{
-			TManagedArray<TObjectPtr<UObject>>& ObjectSkeletalMeshes =
-				GroomCollection.AddAttribute<TObjectPtr<UObject>>(UE::Groom::FGroomGuidesFacade::ObjectSkeletalMeshesAttribute, FacadeType::ObjectsGroup);
-
-			TManagedArray<int32>& ObjectMeshLODs =
-				GroomCollection.AddAttribute<int32>(UE::Groom::FGroomGuidesFacade::ObjectMeshLODsAttribute, FacadeType::ObjectsGroup);
+			const bool bValidSelection = VertexSelection.IsValidForCollection(SkinningFacade.GetManagedArrayCollection());
 
-			TManagedArray<FIntVector4>& VertexBoneIndices =
-				GroomCollection.AddAttribute<FIntVector4>(UE::Groom::FGroomGuidesFacade::PointBoneIndicesAttribute, FacadeType::VerticesGroup);
-
-			TManagedArray<FVector4f>& VertexBoneWeights =
-				GroomCollection.AddAttribute<FVector4f>(UE::Groom::FGroomGuidesFacade::PointBoneWeightsAttribute, FacadeType::VerticesGroup);
-			
-			if(GroupIndex == INDEX_NONE)
+			TMap<FName, FBoneIndexType> TargetBoneToIndex;
+			TargetBoneToIndex.Reserve(SkeletalMesh->GetRefSkeleton().GetRawBoneNum());
+			for (int32 BoneIdx = 0; BoneIdx < SkeletalMesh->GetRefSkeleton().GetRawBoneNum(); ++BoneIdx)
 			{
-				for(int32 LocalIndex = 0; LocalIndex < ObjectSkeletalMeshes.Num(); ++LocalIndex)
-				{
-					ObjectSkeletalMeshes[LocalIndex] = SkeletalMesh;
-					ObjectMeshLODs[LocalIndex] = LODIndex;
-				}
+				TargetBoneToIndex.Add(SkeletalMesh->GetRefSkeleton().GetRawRefBoneInfo()[BoneIdx].Name, BoneIdx);
 			}
-			else if(ObjectSkeletalMeshes.IsValidIndex(GroupIndex) && ObjectMeshLODs.IsValidIndex(GroupIndex))
+			UE::Geometry::FDynamicMesh3 DynamicMesh;
+			if (UE::Groom::Private::SkeletalMeshToDynamicMesh(SkeletalMesh, LODIndex, DynamicMesh))
 			{
-				ObjectSkeletalMeshes[GroupIndex] = SkeletalMesh;
-				ObjectMeshLODs[GroupIndex] = LODIndex;
-			}
-			if(SkeletalMesh && SkeletalMesh->IsValidLODIndex(LODIndex))
-			{
-				TMap<FName, FBoneIndexType> TargetBoneToIndex;
-				TargetBoneToIndex.Reserve(SkeletalMesh->GetRefSkeleton().GetRawBoneNum());
-				for (int32 BoneIdx = 0; BoneIdx < SkeletalMesh->GetRefSkeleton().GetRawBoneNum(); ++BoneIdx)
-				{
-					TargetBoneToIndex.Add(SkeletalMesh->GetRefSkeleton().GetRawRefBoneInfo()[BoneIdx].Name, BoneIdx);
-				}
-				UE::Geometry::FDynamicMesh3 DynamicMesh;
-				if (UE::Groom::Private::SkeletalMeshToDynamicMesh(SkeletalMesh, LODIndex, DynamicMesh))
-				{
-					MeshTransforms::ApplyTransform(DynamicMesh, RelativeTransform, true);
+				MeshTransforms::ApplyTransform(DynamicMesh, RelativeTransform, true);
 
-					UE::Geometry::FTransferBoneWeights TransferBoneWeights(&DynamicMesh, FSkeletalMeshAttributes::DefaultSkinWeightProfileName);
-					TransferBoneWeights.bUseParallel = true;
-					TransferBoneWeights.MaxNumInfluences = 4;
-					TransferBoneWeights.TransferMethod = UE::Geometry::FTransferBoneWeights::ETransferBoneWeightsMethod::ClosestPointOnSurface;
+				UE::Geometry::FTransferBoneWeights TransferBoneWeights(&DynamicMesh, FSkeletalMeshAttributes::DefaultSkinWeightProfileName);
+				TransferBoneWeights.bUseParallel = true;
+				TransferBoneWeights.MaxNumInfluences = 4;
+				TransferBoneWeights.TransferMethod = UE::Geometry::FTransferBoneWeights::ETransferBoneWeightsMethod::ClosestPointOnSurface;
 
-					if (TransferBoneWeights.Validate() == UE::Geometry::EOperationValidationResult::Ok)
+				if (TransferBoneWeights.Validate() == UE::Geometry::EOperationValidationResult::Ok)
+				{
+					for(int32 GeometryIndex = 0; GeometryIndex < SkinningFacade.NumGeometry(); ++GeometryIndex)
 					{
-						const int32 LocalIndex = GroupIndex;
-						ParallelFor(GroomFacade.GetNumPoints(), [&TransferBoneWeights, &TargetBoneToIndex, &GroomFacade, LocalIndex, &VertexBoneIndices, &VertexBoneWeights](
-							int32 PointIndex)
+						SkinningFacade.ModifyGeometryBinding(GeometryIndex, SkeletalMesh, LODIndex);
+						
+						const TArrayView<const FVector3f> VertexPositions = MeshFacade.GetVertexPositions(GeometryIndex);
+						const TArray<int32> VertexIndices = MeshFacade.GetVertexIndices(GeometryIndex);
+						
+						const int32 NumVertices = VertexPositions.Num();
+						ParallelFor(NumVertices, [&TransferBoneWeights, &TargetBoneToIndex, &VertexPositions, &VertexIndices, &SkinningFacade,
+							&VertexSelection, bValidSelection](const int32 LocalIndex)
 						{
-							const int32 CurveIndex = GroomFacade.GetPointCurveIndices()[PointIndex];
-							const int32 ObjectIndex = GroomFacade.GetCurveObjectIndices()[CurveIndex];
+							const int32 VertexIndex = VertexIndices[LocalIndex];
 
-							if(LocalIndex == INDEX_NONE || ObjectIndex == LocalIndex)
+							if(!bValidSelection || (bValidSelection && (VertexSelection.IsSelected(VertexIndex))))
 							{
 								TArray<int32> BoneIndices;
 								TArray<float> BoneWeights;
-								TransferBoneWeights.TransferWeightsToPoint(BoneIndices, BoneWeights,
-									GroomFacade.GetPointRestPositions()[PointIndex],&TargetBoneToIndex);
-
-								FIntVector4 PointBoneIndices(INDEX_NONE);
-								FVector4f PointBoneWeights(0.0f);
-
-								for(int32 BoneIdx = 0; BoneIdx < BoneIndices.Num(); ++BoneIdx)
-								{
-									PointBoneIndices[BoneIdx] = BoneIndices[BoneIdx];
-									PointBoneWeights[BoneIdx] = BoneWeights[BoneIdx];
-								}
-								VertexBoneIndices[2*PointIndex] = PointBoneIndices;
-								VertexBoneIndices[2*PointIndex+1] = PointBoneIndices;
-								
-								VertexBoneWeights[2*PointIndex] = PointBoneWeights;
-								VertexBoneWeights[2*PointIndex+1] = PointBoneWeights;
+									
+								TransferBoneWeights.TransferWeightsToPoint(BoneIndices, BoneWeights,VertexPositions[LocalIndex], &TargetBoneToIndex);
+								SkinningFacade.ModifyBoneWeight(VertexIndex, BoneIndices, BoneWeights);
 							}
-
-						}, TransferBoneWeights.bUseParallel  ? EParallelForFlags::None : EParallelForFlags::ForceSingleThread);
-					}
-				}
-			}
-		}
-	}
-	
-	template<typename FacadeType>
-	FORCEINLINE void ExtractSkinningData(FManagedArrayCollection& GroomCollection, const FCollectionAttributeKey& BoneIndicesKey, const FCollectionAttributeKey& BoneWeightsKey)
-	{
-		if (!BoneIndicesKey.Attribute.IsEmpty() && !BoneWeightsKey.Attribute.IsEmpty())
-		{
-			FacadeType GroomFacade(GroomCollection);
-
-			if(GroomFacade.IsValid())
-			{
-				const TManagedArray<FIntVector4>* BoneIndices = GroomCollection.FindAttributeTyped<FIntVector4>(FName(UE::Groom::FGroomGuidesFacade::PointBoneIndicesAttribute), FacadeType::VerticesGroup);
-				const TManagedArray<FVector4f>* BoneWeights = GroomCollection.FindAttributeTyped<FVector4f>(FName(UE::Groom::FGroomGuidesFacade::PointBoneWeightsAttribute), FacadeType::VerticesGroup);
-				
-				if(BoneIndices && BoneWeights)
-				{
-					TManagedArray<TArray<int32>>& IndicesArray = GroomCollection.AddAttribute<TArray<int32>>(FName(BoneIndicesKey.Attribute), FName(BoneIndicesKey.Group));
-					TManagedArray<TArray<float>>& WeightsArray = GroomCollection.AddAttribute<TArray<float>>(FName(BoneWeightsKey.Attribute), FName(BoneWeightsKey.Group));
-
-					const int32 NumVertices = GroomFacade.GetNumVertices();
-					for(int32 VertexIndex = 0; VertexIndex < NumVertices; ++VertexIndex)
-					{
-						const FIntVector4& PointIndices = (*BoneIndices)[VertexIndex];
-						const FVector4f& PointWeights = (*BoneWeights)[VertexIndex];
-
-						IndicesArray[VertexIndex].Reset(4);
-						WeightsArray[VertexIndex].Reset(4);
-						for(int32 WeightIndex = 0; WeightIndex < 4; ++WeightIndex)
-						{
-							if(PointIndices[WeightIndex] != INDEX_NONE)
-							{
-								IndicesArray[VertexIndex].Add(PointIndices[WeightIndex]);
-								WeightsArray[VertexIndex].Add(PointWeights[WeightIndex]);
-							}
-						}
-					}
-				}
-			}
-		}
-	}
-
-	template<typename FacadeType>
-	FORCEINLINE void ReportSkinningData(FManagedArrayCollection& GroomCollection, const FCollectionAttributeKey& BoneIndicesKey, const FCollectionAttributeKey& BoneWeightsKey)
-	{
-		if (!BoneIndicesKey.Attribute.IsEmpty() && !BoneWeightsKey.Attribute.IsEmpty())
-		{
-			FacadeType GroomFacade(GroomCollection);
-
-			if(GroomFacade.IsValid())
-			{
-				const TManagedArray<TArray<int32>>* IndicesArray = GroomCollection.FindAttributeTyped<TArray<int32>>(FName(BoneIndicesKey.Attribute), FName(BoneIndicesKey.Group));
-				const TManagedArray<TArray<float>>* WeightsArray = GroomCollection.FindAttributeTyped<TArray<float>>(FName(BoneWeightsKey.Attribute), FName(BoneWeightsKey.Group));
-
-				if(IndicesArray && WeightsArray)
-				{
-					TManagedArray<FIntVector4>& BoneIndices = GroomCollection.AddAttribute<FIntVector4>(FName(UE::Groom::FGroomGuidesFacade::PointBoneIndicesAttribute), FacadeType::VerticesGroup);
-					TManagedArray<FVector4f>& BoneWeights = GroomCollection.AddAttribute<FVector4f>(FName(UE::Groom::FGroomGuidesFacade::PointBoneWeightsAttribute), FacadeType::VerticesGroup);
-				
-					const int32 NumVertices = GroomFacade.GetNumVertices();
-					for(int32 VertexIndex = 0; VertexIndex < NumVertices; ++VertexIndex)
-					{
-						FIntVector4 PointIndices(INDEX_NONE);
-						FVector4f PointWeights(0.0f);
-
-						const int32 NumWeights = FMath::Min(4, (*IndicesArray)[VertexIndex].Num());
-						float TotalWeight = 0.0;
-						for(int32 BoneIdx = 0; BoneIdx < NumWeights; ++BoneIdx)
-						{
-							PointIndices[BoneIdx] = (*IndicesArray)[VertexIndex][BoneIdx];
-							PointWeights[BoneIdx] = (*WeightsArray)[VertexIndex][BoneIdx];
-							TotalWeight += PointWeights[BoneIdx];
-						}
-						if(TotalWeight != 0.0)
-						{
-							for(int32 BoneIdx = 0; BoneIdx < NumWeights; ++BoneIdx)
-							{
-								PointWeights[BoneIdx] /= TotalWeight;
-							}
-						}
-						BoneIndices[VertexIndex] = PointIndices;
-						BoneWeights[VertexIndex] = PointWeights;
+						}, TransferBoneWeights.bUseParallel ? EParallelForFlags::None : EParallelForFlags::ForceSingleThread);
 					}
 				}
 			}
@@ -242,62 +122,67 @@ FTransferSkinWeightsGroomNode::FTransferSkinWeightsGroomNode(const UE::Dataflow:
 }
 
 void FTransferSkinWeightsGroomNode::Evaluate(UE::Dataflow::FContext& Context, const FDataflowOutput* Out) const
+{}
+
+FTransferGeometrySkinWeightsDataflowNode::FTransferGeometrySkinWeightsDataflowNode(const UE::Dataflow::FNodeParameters& InParam, FGuid InGuid)
+	: FDataflowNode(InParam, InGuid)
+{
+	RegisterInputConnection(&Collection);
+	RegisterInputConnection(&VertexSelection);
+	RegisterInputConnection(&SkeletalMesh);
+	RegisterInputConnection(&LODIndex).SetCanHidePin(true).SetPinIsHidden(true);
+	RegisterInputConnection(&RelativeTransform).SetCanHidePin(true).SetPinIsHidden(true);
+	
+	RegisterOutputConnection(&Collection, &Collection);
+	RegisterOutputConnection(&BoneIndicesKey);
+	RegisterOutputConnection(&BoneWeightsKey);
+}
+
+void FTransferGeometrySkinWeightsDataflowNode::Evaluate(UE::Dataflow::FContext& Context, const FDataflowOutput* Out) const
 {
 	if (Out->IsA(&Collection))
 	{
-		FManagedArrayCollection GroomCollection = GetValue(Context, &Collection);
+		FManagedArrayCollection GeometryCollection = GetValue(Context, &Collection);
+		GeometryCollection::Facades::FCollectionMeshFacade MeshFacade(GeometryCollection);
 
-		if(CurvesType == EGroomCollectionType::Guides)
-		{
-			UE::Groom::Private::BuildSkinningData<UE::Groom::FGroomGuidesFacade>(GroomCollection, GroupIndex, SkeletalMesh, LODIndex, RelativeTransform);
-		}
-		else
+		if(MeshFacade.IsValid())
 		{
-			UE::Groom::Private::BuildSkinningData<UE::Groom::FGroomStrandsFacade>(GroomCollection, GroupIndex, SkeletalMesh, LODIndex, RelativeTransform);
+			GeometryCollection::Facades::FVertexBoneWeightsFacade SkinningFacade(GeometryCollection, false);
+			SkinningFacade.DefineSchema();
+			
+			UE::Groom::Private::BuildGeometrySkinWeights(MeshFacade, SkinningFacade, GetValue<FDataflowVertexSelection>(Context, &VertexSelection),
+				GetValue<TObjectPtr<USkeletalMesh>>(Context, &SkeletalMesh), GetValue<int32>(Context, &LODIndex),
+				GetValue<FTransform>(Context, &RelativeTransform));
 		}
-		SetValue(Context, MoveTemp(GroomCollection), &Collection);
+
+		SetValue(Context, MoveTemp(GeometryCollection), &Collection);
 	}
-	else if (Out->IsA(&BoneIndicesKey))
+	else if (Out->IsA<FCollectionAttributeKey>(&BoneIndicesKey))
 	{
 		SetValue(Context, GetBoneIndicesKey(), &BoneIndicesKey);
 	}
-	else if (Out->IsA(&BoneWeightsKey))
+	else if (Out->IsA<FCollectionAttributeKey>(&BoneWeightsKey))
 	{
 		SetValue(Context, GetBoneWeightsKey(), &BoneWeightsKey);
 	}
-	else if (Out->IsA(&SkeletalMesh))
-	{
-		SafeForwardInput(Context, &SkeletalMesh, &SkeletalMesh);
-	}
 }
 
-FCollectionAttributeKey FTransferSkinWeightsGroomNode::GetBoneIndicesKey() const
+FCollectionAttributeKey FTransferGeometrySkinWeightsDataflowNode::GetBoneIndicesKey()
 {
 	FCollectionAttributeKey Key;
-	Key.Group = (CurvesType == EGroomCollectionType::Guides) ? UE::Groom::FGroomGuidesFacade::VerticesGroup.ToString() :
-		UE::Groom::FGroomStrandsFacade::VerticesGroup.ToString();
-	Key.Attribute = UE::Groom::FGroomGuidesFacade::PointBoneIndicesAttribute.ToString();
+	Key.Group = FGeometryCollection::VerticesGroup.ToString();
+	Key.Attribute = GeometryCollection::Facades::FVertexBoneWeightsFacade::BoneIndicesAttributeName.ToString();
+	
 	return Key;
 }
 
-FCollectionAttributeKey FTransferSkinWeightsGroomNode::GetBoneWeightsKey() const
+FCollectionAttributeKey FTransferGeometrySkinWeightsDataflowNode::GetBoneWeightsKey()
 {
 	FCollectionAttributeKey Key;
-	Key.Group = (CurvesType == EGroomCollectionType::Guides) ? UE::Groom::FGroomGuidesFacade::VerticesGroup.ToString() :
-		UE::Groom::FGroomStrandsFacade::VerticesGroup.ToString();
-	Key.Attribute = UE::Groom::FGroomGuidesFacade::PointBoneWeightsAttribute.ToString();
+	Key.Group = FGeometryCollection::VerticesGroup.ToString();
+	Key.Attribute = GeometryCollection::Facades::FVertexBoneWeightsFacade::BoneWeightsAttributeName.ToString();
 	return Key;
 }
 
-TArray<UE::Dataflow::FRenderingParameter> FTransferSkinWeightsGroomNode::GetRenderParametersImpl() const
-{
-	if(CurvesType == EGroomCollectionType::Guides)
-	{
-		return { {TEXT("GuidesRender"), FName("FGroomCollection"), {TEXT("Collection")}}};
-	}
-	else
-	{
-		return { {TEXT("StrandsRender"), FName("FGroomCollection"), {TEXT("Collection")}}};
-	}
-}
+
 
diff --git a/Engine/Plugins/Runtime/HairStrands/Source/HairStrandsDataflow/Private/BuildGroomSplineSkinningNode.cpp b/Engine/Plugins/Runtime/HairStrands/Source/HairStrandsDataflow/Private/BuildGroomSplineSkinningNode.cpp
--- a/Engine/Plugins/Runtime/HairStrands/Source/HairStrandsDataflow/Private/BuildGroomSplineSkinningNode.cpp
+++ b/Engine/Plugins/Runtime/HairStrands/Source/HairStrandsDataflow/Private/BuildGroomSplineSkinningNode.cpp
@@ -3,11 +3,52 @@
 #include "BuildGroomSplineSkinningNode.h"
 #include "Algo/MaxElement.h"
 #include "Algo/Count.h"
+#include "Dataflow/DataflowSelection.h"
+#include "GeometryCollection/Facades/CollectionMeshFacade.h"
+#include "GeometryCollection/Facades/CollectionVertexBoneWeightsFacade.h"
+#include "GeometryCollection/Facades/CollectionCurveFacade.h"
 
 #include UE_INLINE_GENERATED_CPP_BY_NAME(BuildGroomSplineSkinningNode)
 
 namespace UE::Groom::Private
 {
+	const FName FCollectionSplineAttributes::VertexSplineParamAttribute("SplineParam");
+	const FName FCollectionSplineAttributes::VertexSplineBonesAttribute("SplineBones");
+
+	static FCollectionAttributeKey GetBoneIndicesKey()
+	{
+		FCollectionAttributeKey Key;
+		Key.Group = FGeometryCollection::VerticesGroup.ToString();
+		Key.Attribute = GeometryCollection::Facades::FVertexBoneWeightsFacade::BoneIndicesAttributeName.ToString();
+	
+		return Key;
+	}
+
+	static FCollectionAttributeKey GetBoneWeightsKey()
+	{
+		FCollectionAttributeKey Key;
+		Key.Group = FGeometryCollection::VerticesGroup.ToString();
+		Key.Attribute = GeometryCollection::Facades::FVertexBoneWeightsFacade::BoneWeightsAttributeName.ToString();
+		return Key;
+	}
+
+	static FCollectionAttributeKey GetSplineParamKey()
+	{
+		FCollectionAttributeKey Key;
+		Key.Group = FGeometryCollection::VerticesGroup.ToString();
+		Key.Attribute = UE::Groom::Private::FCollectionSplineAttributes::VertexSplineParamAttribute.ToString();
+	
+		return Key;
+	}
+
+	static FCollectionAttributeKey GetSplineBonesKey()
+	{
+		FCollectionAttributeKey Key;
+		Key.Group = FGeometryCollection::VerticesGroup.ToString();
+		Key.Attribute = UE::Groom::Private::FCollectionSplineAttributes::VertexSplineBonesAttribute.ToString();
+		return Key;
+	}
+	
 	struct FCatmullRomSpline
 	{
 		TArray<FVector> CVs;
@@ -150,328 +191,289 @@ namespace UE::Groom::Private
 		}
 	}
 	
-	template<typename FacadeType>
-	static void BuildSplineSkinningWeights(FManagedArrayCollection& GroomCollection, int GroupIndex, const TObjectPtr<USkeletalMesh> SkeletalMesh, int LODIndex,
-		const TArray<UE::Groom::Private::FCatmullRomSpline>& Splines, FCollectionAttributeKey SplineParamKey, FCollectionAttributeKey SplineBoneKey)
+	static void BuildSplineSkinningWeights(FManagedArrayCollection& GeometryCollection, const GeometryCollection::Facades::FCollectionMeshFacade& MeshFacade,
+		const GeometryCollection::Facades::FCollectionCurveGeometryFacade& CurveFacade, const FDataflowVertexSelection& VertexSelection, 
+		const TArray<UE::Groom::Private::FCatmullRomSpline>& Splines, const FCollectionAttributeKey& SplineParamKey, const FCollectionAttributeKey& SplineBonesKey)
 	{
-		FacadeType GroomFacade(GroomCollection);
-
-		if (GroomFacade.IsValid() && (GroupIndex == INDEX_NONE || GroupIndex < GroomFacade.GetNumObjects()))
+		if(MeshFacade.IsValid() && CurveFacade.IsValid() && !SplineParamKey.Attribute.IsEmpty() && !SplineBonesKey.Attribute.IsEmpty())
 		{
-			// Set groom meshes
-			TManagedArray<TObjectPtr<UObject>>& ObjectSkeletalMeshes =
-				GroomCollection.AddAttribute<TObjectPtr<UObject>>(UE::Groom::FGroomGuidesFacade::ObjectSkeletalMeshesAttribute, FacadeType::ObjectsGroup);
+			const bool bValidSelection = VertexSelection.IsValidForCollection(GeometryCollection);
 
-			TManagedArray<int32>& ObjectMeshLODs =
-				GroomCollection.AddAttribute<int32>(UE::Groom::FGroomGuidesFacade::ObjectMeshLODsAttribute, FacadeType::ObjectsGroup);
+			const bool ExistingSplineParams = GeometryCollection.HasAttribute(FName(SplineParamKey.Attribute), FName(SplineParamKey.Group)) &&
+				GeometryCollection.HasAttribute(FName(SplineBonesKey.Attribute), FName(SplineBonesKey.Group));
 
-			if (GroupIndex == INDEX_NONE)
+			TManagedArray<float>& SplineParams = GeometryCollection.AddAttribute<float>(FName(SplineParamKey.Attribute), FName(SplineParamKey.Group));
+			TManagedArray<FIntVector2>& SplineBones = GeometryCollection.AddAttribute<FIntVector2>(FName(SplineBonesKey.Attribute), FName(SplineBonesKey.Group));
+
+			if (!ExistingSplineParams)
 			{
-				for (int32 LocalIndex = 0; LocalIndex < ObjectSkeletalMeshes.Num(); ++LocalIndex)
+				for (int32 VertexIndex = 0; VertexIndex < SplineParams.Num(); ++VertexIndex)
 				{
-					ObjectSkeletalMeshes[LocalIndex] = SkeletalMesh;
-					ObjectMeshLODs[LocalIndex] = LODIndex;
+					SplineParams[VertexIndex] = 0;
+					SplineBones[VertexIndex] = FIntVector2(INDEX_NONE, INDEX_NONE);
 				}
 			}
-			else if (ObjectSkeletalMeshes.IsValidIndex(GroupIndex) && ObjectMeshLODs.IsValidIndex(GroupIndex))
-			{
-				ObjectSkeletalMeshes[GroupIndex] = SkeletalMesh;
-				ObjectMeshLODs[GroupIndex] = LODIndex;
-			}
-
-			// Compute skinning weights
-			TManagedArray<float>& VerticesSplineParams = GroomCollection.AddAttribute<float>(FName(SplineParamKey.Attribute), FName(SplineParamKey.Group));
-			TManagedArray<FIntVector2>& VerticesSplineBones = GroomCollection.AddAttribute<FIntVector2>(FName(SplineBoneKey.Attribute), FName(SplineBoneKey.Group));
-			TManagedArray<float>& VerticesSplineSegments = GroomCollection.AddAttribute<float>(FName("SplineSegment_Build"), FName(SplineParamKey.Group));
-			TManagedArray<float>& VerticesSplineSegmentParams = GroomCollection.AddAttribute<float>(FName("SplineSegmentParam_Build"), FName(SplineParamKey.Group));
-
-			const TArray<FVector3f>& Positions = GroomFacade.GetPointRestPositions();
-			const TArray<int32>& CurvePointOffsets = GroomFacade.GetCurvePointOffsets();
 
-			const int FirstCurveIndex = GroupIndex > 0 ? GroomFacade.GetObjectCurveOffsets()[GroupIndex - 1] : 0;
-			const int LastCurveIndex = GroupIndex != INDEX_NONE ? GroomFacade.GetObjectCurveOffsets()[GroupIndex] : GroomFacade.GetNumCurves();
-			const int32 NumCurves = LastCurveIndex - FirstCurveIndex;
-
-			ParallelFor(NumCurves, [&FirstCurveIndex, &Positions, &CurvePointOffsets, &Splines, &VerticesSplineParams, &VerticesSplineBones, &VerticesSplineSegments, &VerticesSplineSegmentParams](int32 CurveIndex)
+			const int32 NumCurves = CurveFacade.GetNumCurves();
+			ParallelFor(NumCurves, [&Splines, &SplineParams, &SplineBones, &CurveFacade, &VertexSelection, bValidSelection](int32 CurveIndex)
 			{
 				TArray<int32> SplineCounts;
 				SplineCounts.Init(0, Splines.Num());
-					
-				const int32 FirstPoint = FirstCurveIndex + CurveIndex > 0 ? CurvePointOffsets[FirstCurveIndex + CurveIndex - 1] : 0;
-				const int32 LastPoint = CurvePointOffsets[FirstCurveIndex + CurveIndex];
+
+				const int32 FirstPoint = CurveIndex > 0 ? CurveFacade.GetCurvePointOffsets()[CurveIndex-1] : 0;
+				const int32 LastPoint = CurveFacade.GetCurvePointOffsets()[CurveIndex];
 
 				for (int32 PointIndex = FirstPoint; PointIndex < LastPoint; ++PointIndex)
 				{
-					const FVector Position(Positions[PointIndex]);
-					int32 ClosestSplineIndex = 0;
-					float MinParam=0, MinDistance = UE_MAX_FLT;
-
-					for (int32 SplineIndex = 0; SplineIndex < Splines.Num(); ++SplineIndex)
+					if(!bValidSelection || (bValidSelection && (VertexSelection.IsSelected(2*PointIndex) || VertexSelection.IsSelected(2*PointIndex+1))))
 					{
-						float Param, Distance;
-						Splines[SplineIndex].ClosestParam(Position, Param, Distance);
+						const FVector PointPosition = FVector(CurveFacade.GetPointRestPositions()[PointIndex]);
+						int32 ClosestSplineIndex = 0;
+					
+						float MinDistance = UE_MAX_FLT;
+						float MinParam = 0.0f; 
 
-						if (Distance < MinDistance)
+						for (int32 SplineIndex = 0; SplineIndex < Splines.Num(); ++SplineIndex)
 						{
-							MinParam = Param;
-							MinDistance = Distance;
-							ClosestSplineIndex = SplineIndex;
+							float Param, Distance;
+							Splines[SplineIndex].ClosestParam(PointPosition, Param, Distance);
+
+							if (Distance < MinDistance)
+							{
+								MinParam = Param;
+								MinDistance = Distance;
+								ClosestSplineIndex = SplineIndex;
 
+							}
 						}
-					}
-					
-					VerticesSplineParams[PointIndex * 2] = MinParam;
-					VerticesSplineParams[PointIndex * 2 + 1] = MinParam;
-					VerticesSplineBones[PointIndex * 2] = Splines[ClosestSplineIndex].Bones;
-					VerticesSplineBones[PointIndex * 2 + 1] = Splines[ClosestSplineIndex].Bones;
+				
+						SplineParams[PointIndex * 2] = MinParam;
+						SplineParams[PointIndex * 2 + 1] = MinParam;
+						SplineBones[PointIndex * 2] = Splines[ClosestSplineIndex].Bones;
+						SplineBones[PointIndex * 2 + 1] = Splines[ClosestSplineIndex].Bones;
 
-					++SplineCounts[ClosestSplineIndex];
+						++SplineCounts[ClosestSplineIndex];
+					}
 				}
-				
+
 				// Ensure that all CVs on a strand are bound to the same spline
 				if (Algo::Count(SplineCounts, 0) < SplineCounts.Num() - 1)
 				{
 					// Use the most commonly bound spline
-					int32 SplineIndex = Algo::MaxElement(SplineCounts) - &SplineCounts[0];
+					const int32 SplineIndex = Algo::MaxElement(SplineCounts) - &SplineCounts[0];
 
 					for (int32 PointIndex = FirstPoint; PointIndex < LastPoint; ++PointIndex)
 					{
-						if (VerticesSplineBones[PointIndex * 2] != Splines[SplineIndex].Bones)
+						if(!bValidSelection || (bValidSelection && (VertexSelection.IsSelected(2*PointIndex) || VertexSelection.IsSelected(2*PointIndex+1))))
 						{
-							const FVector Position(Positions[PointIndex]);
-							float MinDistance = UE_MAX_FLT;
-							float Param;
-							Splines[SplineIndex].ClosestParam(Position, Param, MinDistance);
+							if (SplineBones[PointIndex * 2] != Splines[SplineIndex].Bones)
+							{
+								const FVector PointPosition = FVector(CurveFacade.GetPointRestPositions()[PointIndex]);
+								float MinDistance = UE_MAX_FLT;
+								float MinParam = 0.0f;
+								Splines[SplineIndex].ClosestParam(PointPosition, MinParam, MinDistance);
 							
-							VerticesSplineParams[PointIndex * 2] = Param;
-							VerticesSplineParams[PointIndex * 2 + 1] = Param;
-							VerticesSplineBones[PointIndex * 2] = Splines[SplineIndex].Bones;
-							VerticesSplineBones[PointIndex * 2 + 1] = Splines[SplineIndex].Bones;
+								SplineParams[PointIndex * 2] = MinParam;
+								SplineParams[PointIndex * 2 + 1] = MinParam;
+								SplineBones[PointIndex * 2] = Splines[SplineIndex].Bones;
+								SplineBones[PointIndex * 2 + 1] = Splines[SplineIndex].Bones;
+							}
 						}
 					}
 				}
 			});
 		}
 	}
-
-	template<typename FacadeType>
-	static void ConvertFromSplineToLinearWeights(FManagedArrayCollection & GroomCollection, int32 GroomGroupIndex, const FCollectionAttributeKey & SplineParamKey, const FCollectionAttributeKey& SplineBonesKey,
-		const FCollectionAttributeKey & BoneIndicesKey, const FCollectionAttributeKey & BoneWeightsKey)
+	
+	static void ConvertFromSplineToLinearWeights(FManagedArrayCollection& GeometryCollection, const GeometryCollection::Facades::FCollectionMeshFacade& MeshFacade,
+		GeometryCollection::Facades::FVertexBoneWeightsFacade& SkinningFacade, const FDataflowVertexSelection& VertexSelection, const FCollectionAttributeKey& SplineParamKey, const FCollectionAttributeKey& SplineBonesKey)
 	{
-		const TManagedArray<float>* SplineParams = GroomCollection.FindAttributeTyped<float>(FName(SplineParamKey.Attribute), FName(SplineParamKey.Group));
-		const TManagedArray<FIntVector2>* SplineBones = GroomCollection.FindAttributeTyped<FIntVector2>(FName(SplineBonesKey.Attribute), FName(SplineBonesKey.Group));
-
-		if (SplineParams && SplineBones)
+		if(MeshFacade.IsValid() && SkinningFacade.IsValid() && !SplineParamKey.Attribute.IsEmpty() && !SplineBonesKey.Attribute.IsEmpty())
 		{
-			FacadeType GroomFacade(GroomCollection);
+			const bool bValidSelection = VertexSelection.IsValidForCollection(SkinningFacade.GetManagedArrayCollection());
+			
+			const TManagedArray<float>* SplineParams = GeometryCollection.FindAttributeTyped<float>(FName(SplineParamKey.Attribute), FName(SplineParamKey.Group));
+			const TManagedArray<FIntVector2>* SplineBones = GeometryCollection.FindAttributeTyped<FIntVector2>(FName(SplineBonesKey.Attribute), FName(SplineBonesKey.Group));
 
-			if (GroomFacade.IsValid() && (GroomGroupIndex == INDEX_NONE || GroomGroupIndex < GroomFacade.GetNumObjects()))
+			if (SplineParams && SplineBones)
 			{
-				bool ExistingBoneIndices = !GroomCollection.HasAttribute(FName(BoneIndicesKey.Attribute), FName(BoneIndicesKey.Group));
-				TManagedArray<FIntVector4>& BoneIndices = GroomCollection.AddAttribute<FIntVector4>(FName(BoneIndicesKey.Attribute), FName(BoneIndicesKey.Group));
-				TManagedArray<FVector4f>& BoneWeights = GroomCollection.AddAttribute<FVector4f>(FName(BoneWeightsKey.Attribute), FName(BoneWeightsKey.Group));
-
-				if (!ExistingBoneIndices)
-				{
-					for (int32 VertexIndex = 0; VertexIndex < GroomFacade.GetNumPoints() * 2; ++VertexIndex)
-					{
-						BoneIndices[VertexIndex] = FIntVector4(-1, -1, -1, -1);
-						BoneWeights[VertexIndex] = FVector4f(0, 0, 0, 0);
-					}
-				}
-
-				TArray<int32> GroupIndices;
-				if (GroomGroupIndex == INDEX_NONE)
-				{
-					for (int32 LocalIndex = 0; LocalIndex < GroomFacade.GetNumObjects(); ++LocalIndex)
-					{
-						GroupIndices.Add(LocalIndex);
-					}
-				}
-				else
-				{
-					GroupIndices.Add(GroomGroupIndex);
-				}
-
-				for (int32 GroupIndex : GroupIndices)
+				TArray<int32> BoneIndices;
+				TArray<float> BoneWeights;
+				
+				for(int32 GeometryIndex = 0; GeometryIndex < SkinningFacade.NumGeometry(); ++GeometryIndex)
 				{
-					const int FirstCurveIndex = GroupIndex > 0 ? GroomFacade.GetObjectCurveOffsets()[GroupIndex - 1] : 0;
-					const int LastCurveIndex = GroomFacade.GetObjectCurveOffsets()[GroupIndex];
-					const int32 FirstVertex = FirstCurveIndex > 0 ? GroomFacade.GetCurvePointOffsets()[FirstCurveIndex - 1] * 2 : 0;
-					const int32 LastVertex = GroomFacade.GetCurvePointOffsets()[LastCurveIndex-1] * 2;
-
-					for (int32 VertexIndex = FirstVertex; VertexIndex < LastVertex; ++VertexIndex)
+					const TArray<int32> VertexIndices = MeshFacade.GetVertexIndices(GeometryIndex);
+					for (const int32& VertexIndex : VertexIndices)
 					{
-						float SplineParam = (*SplineParams)[VertexIndex];
-
-						int Segment = trunc(SplineParam);
-						float SegmentParam = SplineParam - Segment;
-
-						if (Segment < (*SplineBones)[VertexIndex][1])
-						{
-							BoneIndices[VertexIndex] = FIntVector4(Segment, Segment + 1, -1 , -1);
-							BoneWeights[VertexIndex] = FVector4f(1.0 - SegmentParam, SegmentParam, 0, 0);
-						}
-						else
+						if(!bValidSelection || (bValidSelection && (VertexSelection.IsSelected(VertexIndex))))
 						{
-							BoneIndices[VertexIndex] = FIntVector4(Segment, -1, -1, -1);
-							BoneWeights[VertexIndex] = FVector4f(1, 0, 0, 0);
+							const float SplineParam = (*SplineParams)[VertexIndex];
+							const int32 SegmentIndex = trunc(SplineParam);
+							const float SegmentParam = SplineParam - SegmentIndex;
+							
+							if (SegmentIndex < (*SplineBones)[VertexIndex][1])
+							{
+								BoneIndices = {SegmentIndex, SegmentIndex + 1};
+								BoneWeights = {1.0f - SegmentParam, SegmentParam};
+							}
+							else
+							{
+								BoneIndices = {SegmentIndex};
+								BoneWeights = {1.0f};
+							}
+
+							SkinningFacade.ModifyBoneWeight(VertexIndex, BoneIndices, BoneWeights);
 						}
 					}
 				}
 			}
 		}
 	}
-
-	template<typename FacadeType>
-	static void ConvertFromLinearToSplineWeights(FManagedArrayCollection & GroomCollection, int32 GroomGroupIndex, const FCollectionAttributeKey & SplineParamKey, const FCollectionAttributeKey & SplineBonesKey)
+	
+	static void ConvertFromLinearToSplineWeights(FManagedArrayCollection& GeometryCollection, GeometryCollection::Facades::FCollectionMeshFacade& MeshFacade,
+		GeometryCollection::Facades::FVertexBoneWeightsFacade& SkinningFacade, const FDataflowVertexSelection& VertexSelection,
+		const FCollectionAttributeKey& SplineParamKey, const FCollectionAttributeKey& SplineBonesKey)
 	{
-		const TManagedArray<FIntVector4>* BoneIndices = GroomCollection.FindAttributeTyped<FIntVector4>(FName(UE::Groom::FGroomGuidesFacade::PointBoneIndicesAttribute), FacadeType::VerticesGroup);
-		const TManagedArray<FVector4f>* BoneWeights = GroomCollection.FindAttributeTyped<FVector4f>(FName(UE::Groom::FGroomGuidesFacade::PointBoneWeightsAttribute), FacadeType::VerticesGroup);
-		
-		TManagedArray<TObjectPtr<UObject>>* ObjectSkeletalMeshes =
-			GroomCollection.FindAttributeTyped<TObjectPtr<UObject>>(UE::Groom::FGroomGuidesFacade::ObjectSkeletalMeshesAttribute, FacadeType::ObjectsGroup);
-		TManagedArray<int32>* ObjectMeshLODs =
-			GroomCollection.FindAttributeTyped<int32>(UE::Groom::FGroomGuidesFacade::ObjectMeshLODsAttribute, FacadeType::ObjectsGroup);
-
-		if (BoneIndices && BoneWeights && ObjectSkeletalMeshes && ObjectMeshLODs && !SplineParamKey.Attribute.IsEmpty() && !SplineBonesKey.Attribute.IsEmpty())
+		if(MeshFacade.IsValid() && SkinningFacade.IsValid() && !SplineParamKey.Attribute.IsEmpty() && !SplineBonesKey.Attribute.IsEmpty())
 		{
-			FacadeType GroomFacade(GroomCollection);
+			const bool bValidSelection = VertexSelection.IsValidForCollection(SkinningFacade.GetManagedArrayCollection());
 
-			if (GroomFacade.IsValid() && (GroomGroupIndex == INDEX_NONE || GroomGroupIndex < GroomFacade.GetNumObjects()))
-			{
-				TManagedArray<float>& VertexSplineParams = GroomCollection.AddAttribute<float>(FName(SplineParamKey.Attribute), FName(SplineParamKey.Group));
-				TManagedArray<FIntVector2>& VertexSplineBones = GroomCollection.AddAttribute<FIntVector2>(FName(SplineBonesKey.Attribute), FName(SplineBonesKey.Group));
+			const TArray<TObjectPtr<UObject>>& SkeletalMeshs = SkinningFacade.GetSkeletalMeshes().GetConstArray();
+			const TArray<int32>& GeometryLods = SkinningFacade.GetGeometryLODs().GetConstArray();
+			
+			const TArray<TArray<int32>>& BoneIndices = SkinningFacade.GetBoneIndices().GetConstArray();
+			const TArray<TArray<float>>& BoneWeights = SkinningFacade.GetBoneWeights().GetConstArray();
 
-				TArray<int32> GroupIndices;
-				if (GroomGroupIndex == INDEX_NONE)
-				{
-					for (int32 LocalIndex = 0; LocalIndex < GroomFacade.GetNumObjects(); ++LocalIndex)
-					{
-						GroupIndices.Add(LocalIndex);
-					}
-				}
-				else
+			const bool ExistingSplineParams = GeometryCollection.HasAttribute(FName(SplineParamKey.Attribute), FName(SplineParamKey.Group)) &&
+				GeometryCollection.HasAttribute(FName(SplineBonesKey.Attribute), FName(SplineBonesKey.Group));
+			
+			TManagedArray<float>& SplineParams = GeometryCollection.AddAttribute<float>(FName(SplineParamKey.Attribute), FName(SplineParamKey.Group));
+			TManagedArray<FIntVector2>& SplineBones = GeometryCollection.AddAttribute<FIntVector2>(FName(SplineBonesKey.Attribute), FName(SplineBonesKey.Group));
+
+			if (!ExistingSplineParams)
+			{
+				for (int32 VertexIndex = 0; VertexIndex < SplineParams.Num(); ++VertexIndex)
 				{
-					GroupIndices.Add(GroomGroupIndex);
+					SplineParams[VertexIndex] = 0;
+					SplineBones[VertexIndex] = FIntVector2(INDEX_NONE, INDEX_NONE);
 				}
-
-				for (int32 GroupIndex : GroupIndices)
+			}
+		
+			for(int32 GeometryIndex = 0; GeometryIndex < SkinningFacade.NumGeometry(); ++GeometryIndex)
+			{
+				if (SkeletalMeshs.IsValidIndex(GeometryIndex) && GeometryLods.IsValidIndex(GeometryIndex))
 				{
-					if (ObjectSkeletalMeshes->IsValidIndex(GroupIndex) && ObjectMeshLODs->IsValidIndex(GroupIndex))
+					if (TObjectPtr<USkeletalMesh> SkeletalMesh = Cast<USkeletalMesh>(SkeletalMeshs[GeometryIndex]))
 					{
-						if (TObjectPtr<USkeletalMesh> SkeletalMesh = Cast<USkeletalMesh>((*ObjectSkeletalMeshes)[GroupIndex]))
+						const TArray<int32> VertexIndices = MeshFacade.GetVertexIndices(GeometryIndex);
+					
+						// Find all used bones for this group						
+						TSet<int32> UsedBoneIndices;
+						for (const int32& VertexIndex : VertexIndices)
 						{
-							const int FirstCurveIndex = GroupIndex > 0 ? GroomFacade.GetObjectCurveOffsets()[GroupIndex - 1] : 0;
-							const int LastCurveIndex = GroomFacade.GetObjectCurveOffsets()[GroupIndex];
-							const int32 FirstVertex = FirstCurveIndex > 0 ? GroomFacade.GetCurvePointOffsets()[FirstCurveIndex - 1] * 2 : 0;
-							const int32 LastVertex = GroomFacade.GetCurvePointOffsets()[LastCurveIndex-1]*2;
-
-							// Find all used bones for this group						
-							TSet<int32> UsedBoneIndices;
-							for (int32 VertexIndex = FirstVertex; VertexIndex < LastVertex; ++VertexIndex)
+							if(!bValidSelection || (bValidSelection && (VertexSelection.IsSelected(VertexIndex))))
 							{
-								for (int32 InfluenceIndex = 0; InfluenceIndex < 4; ++InfluenceIndex)
+								for (int32 InfluenceIndex = 0; InfluenceIndex < BoneIndices[VertexIndex].Num(); ++InfluenceIndex)
 								{
-									FIntVector4 DebugInfluences = (*BoneIndices)[VertexIndex];
-									if ((*BoneIndices)[VertexIndex][InfluenceIndex] != INDEX_NONE)
+									if (BoneIndices[VertexIndex][InfluenceIndex] != INDEX_NONE)
 									{
-										UsedBoneIndices.Add((*BoneIndices)[VertexIndex][InfluenceIndex]);
+										UsedBoneIndices.Add(BoneIndices[VertexIndex][InfluenceIndex]);
 									}
 								}
 							}
-							
-							// Find root bones
-							TArray<int32> RootBones;
-							TMap<int32, int32> BoneSplineMap;
-							
-							const FReferenceSkeleton& Skeleton = SkeletalMesh->GetRefSkeleton();
-							for (int32 BoneIndex : UsedBoneIndices)
+						}
+
+						// Find root bones
+						TArray<int32> RootBones;
+						TMap<int32, int32> BoneSplineMap;
+						
+						const FReferenceSkeleton& Skeleton = SkeletalMesh->GetRefSkeleton();
+						for (int32 BoneIndex : UsedBoneIndices)
+						{
+							int32 ParentIndex = Skeleton.GetParentIndex(BoneIndex);
+							if (ParentIndex == INDEX_NONE || !UsedBoneIndices.Contains(ParentIndex))
 							{
-								int32 ParentIndex = Skeleton.GetParentIndex(BoneIndex);
-								if (ParentIndex == INDEX_NONE || !UsedBoneIndices.Contains(ParentIndex))
-								{
-									RootBones.Add(BoneIndex);
-								}
+								RootBones.Add(BoneIndex);
 							}
-							
-							// Find end bones
-							TArray<int32> EndBones;						
-							TArray<int32> Children;
+						}
 
-							for (int32 SplineIndex = 0; SplineIndex < RootBones.Num(); ++SplineIndex)
-							{
-								BoneSplineMap.Add(RootBones[SplineIndex], SplineIndex);
-								int32 BoneIndex = RootBones[SplineIndex];
+						// Find end bones
+						TArray<int32> EndBones;						
+						TArray<int32> Children;
 
-								while (Skeleton.GetDirectChildBones(BoneIndex, Children))
+						for (int32 SplineIndex = 0; SplineIndex < RootBones.Num(); ++SplineIndex)
+						{
+							BoneSplineMap.Add(RootBones[SplineIndex], SplineIndex);
+							int32 BoneIndex = RootBones[SplineIndex];
+
+							while (Skeleton.GetDirectChildBones(BoneIndex, Children))
+							{
+								BoneSplineMap.Add(Children[0], SplineIndex);
+								BoneIndex = Children[0];
+								if (Children.Num() > 1)
 								{
-									BoneSplineMap.Add(Children[0], SplineIndex);
-									BoneIndex = Children[0];
-									if (Children.Num() > 1)
-									{
-										// TODO: Add support for branching splines
-									}
+									// TODO: Add support for branching splines
 								}
-
-								EndBones.Add(BoneIndex);
 							}
-							
-							
-							for (int32 VertexIndex = FirstVertex; VertexIndex < LastVertex; ++VertexIndex)
+							EndBones.Add(BoneIndex);
+						}
+						
+						for (const int32& VertexIndex : VertexIndices)
+						{
+							if(!bValidSelection || (bValidSelection && (VertexSelection.IsSelected(VertexIndex))))
 							{
-								const FIntVector4& Influences = (*BoneIndices)[VertexIndex];
-								const FVector4f& Weights = (*BoneWeights)[VertexIndex];
+								const TArray<int32>& VertexBones = BoneIndices[VertexIndex];
+								const TArray<float>& VertexWeights = BoneWeights[VertexIndex];
 
 								int32 MaxWeightIndex = INDEX_NONE;
 								float MaxWeight = 0;
-								for (int32 BoneIndex = 0; BoneIndex < 4; ++BoneIndex)
+								for (int32 BoneIndex = 0; BoneIndex < BoneIndices[VertexIndex].Num(); ++BoneIndex)
 								{
-									if (Weights[BoneIndex] > MaxWeight)
+									if (VertexWeights[BoneIndex] > MaxWeight)
 									{
-										MaxWeight = Weights[BoneIndex];
+										MaxWeight = VertexWeights[BoneIndex];
 										MaxWeightIndex = BoneIndex;
 									}
 								}
 								
 								if (MaxWeight > 0)
 								{
+									const int32 SplineIndex = BoneSplineMap[VertexBones[MaxWeightIndex]];
 									float ParentWeight = 0, ChildWeight = 0;
 									
-									for (int32 InfluenceIndex = 0; InfluenceIndex < 4; ++InfluenceIndex)
+									for (int32 InfluenceIndex = 0; InfluenceIndex < VertexBones.Num(); ++InfluenceIndex)
 									{
-										if (Influences[MaxWeightIndex] > VertexSplineBones[VertexIndex][0] && Influences[InfluenceIndex] == Influences[MaxWeightIndex]-1)
+										if (VertexBones[MaxWeightIndex] > RootBones[SplineIndex] && VertexBones[InfluenceIndex] == VertexBones[MaxWeightIndex] - 1)
 										{
-											ParentWeight = Weights[InfluenceIndex];
+											ParentWeight = VertexWeights[InfluenceIndex];
 										}
 
-										if (Influences[MaxWeightIndex] < VertexSplineBones[VertexIndex][1] && Influences[InfluenceIndex] == Influences[MaxWeightIndex] + 1)
+										if (VertexBones[MaxWeightIndex] < EndBones[SplineIndex] && VertexBones[InfluenceIndex] == VertexBones[MaxWeightIndex] + 1)
 										{
-											ChildWeight = Weights[InfluenceIndex];
+											ChildWeight = VertexWeights[InfluenceIndex];
 										}
 									}
 									
 									if (ParentWeight > ChildWeight)
 									{
-										int Segment = Influences[MaxWeightIndex] - 1; 
+										int Segment = VertexBones[MaxWeightIndex] - 1; 
 										float SegmentParam = MaxWeight / (MaxWeight + ParentWeight);
-										VertexSplineParams[VertexIndex] = Segment + SegmentParam;
+										SplineParams[VertexIndex] = Segment + SegmentParam;
 									}
 									else
 									{
-										int Segment = Influences[MaxWeightIndex];
+										int Segment = VertexBones[MaxWeightIndex];
 										float SegmentParam = (1 - MaxWeight) / (MaxWeight + ChildWeight);
-										VertexSplineParams[VertexIndex] = Segment + SegmentParam;
+										SplineParams[VertexIndex] = Segment + SegmentParam;
 									}		
 									
-									const int32 SplineIndex = BoneSplineMap[Influences[MaxWeightIndex]];
-									VertexSplineBones[VertexIndex] = FIntVector2(RootBones[SplineIndex], EndBones[SplineIndex]);
+									SplineBones[VertexIndex] = FIntVector2(RootBones[SplineIndex], EndBones[SplineIndex]);
 								}
 								else if (!RootBones.IsEmpty())
 								{
 									// Rigidly attach to first spline root if there are no weights for this vertex
-									VertexSplineBones[VertexIndex] = FIntVector2(RootBones[0], EndBones[0]);
-									VertexSplineParams[VertexIndex] = 0;
+									SplineBones[VertexIndex] = FIntVector2(RootBones[0], EndBones[0]);
+									SplineParams[VertexIndex] = 0;
 								}
 							}
 						}
@@ -480,8 +482,7 @@ namespace UE::Groom::Private
 			}
 		}
 	}
-}	
-
+}
 
 FBuildGroomSplineSkinWeightsNode::FBuildGroomSplineSkinWeightsNode(const UE::Dataflow::FNodeParameters& InParam, FGuid InGuid)
 	: FDataflowNode(InParam, InGuid)
@@ -495,240 +496,237 @@ FBuildGroomSplineSkinWeightsNode::FBuildGroomSplineSkinWeightsNode(const UE::Dat
 }
 
 void FBuildGroomSplineSkinWeightsNode::Evaluate(UE::Dataflow::FContext& Context, const FDataflowOutput* Out) const
+{
+}
+
+FConvertSplineToLinearSkinWeightsNode::FConvertSplineToLinearSkinWeightsNode(const UE::Dataflow::FNodeParameters& InParam, FGuid InGuid)
+	: FDataflowNode(InParam, InGuid)
+{
+	RegisterInputConnection(&Collection);
+	RegisterInputConnection(&SplineParamKey);
+	RegisterInputConnection(&SplineBonesKey);
+	RegisterOutputConnection(&Collection, &Collection);
+	RegisterOutputConnection(&BoneIndicesKey);
+	RegisterOutputConnection(&BoneWeightsKey);
+}
+
+void FConvertSplineToLinearSkinWeightsNode::Evaluate(UE::Dataflow::FContext& Context, const FDataflowOutput* Out) const
+{
+}
+
+FConvertLinearToSplineSkinWeightsNode::FConvertLinearToSplineSkinWeightsNode(const UE::Dataflow::FNodeParameters& InParam, FGuid InGuid)
+	: FDataflowNode(InParam, InGuid)
+{
+	RegisterInputConnection(&Collection);
+	RegisterOutputConnection(&Collection, &Collection);
+	RegisterOutputConnection(&SplineParamKey);
+	RegisterOutputConnection(&SplineBoneKey);
+}
+
+void FConvertLinearToSplineSkinWeightsNode::Evaluate(UE::Dataflow::FContext& Context, const FDataflowOutput* Out) const
+{
+}
+
+FBuildSplineSkinWeightsDataflowNode::FBuildSplineSkinWeightsDataflowNode(const UE::Dataflow::FNodeParameters& InParam, FGuid InGuid)
+	: FDataflowNode(InParam, InGuid)
+{
+	RegisterInputConnection(&Collection);
+	RegisterInputConnection(&VertexSelection);
+	RegisterInputConnection(&SkeletalMesh);
+	RegisterInputConnection(&LODIndex).SetCanHidePin(true).SetPinIsHidden(true);
+	RegisterInputConnection(&RootBones).SetCanHidePin(true).SetPinIsHidden(true);
+	RegisterInputConnection(&SamplesPerSegment).SetCanHidePin(true).SetPinIsHidden(true);
+
+	RegisterOutputConnection(&Collection, &Collection);
+	RegisterOutputConnection(&SplineParamKey);
+	RegisterOutputConnection(&SplineBonesKey);
+}
+
+void FBuildSplineSkinWeightsDataflowNode::Evaluate(UE::Dataflow::FContext& Context, const FDataflowOutput* Out) const
 {
 	if (Out->IsA<FManagedArrayCollection>(&Collection))
 	{
-		FManagedArrayCollection GroomCollection = GetValue<FManagedArrayCollection>(Context, &Collection);
-		FCollectionAttributeKey ParamKey = GetSplineParamKey();
-		FCollectionAttributeKey BoneKey = GetSplineBonesKey();
-
-		if (!ParamKey.Attribute.IsEmpty() && !BoneKey.Attribute.IsEmpty() && SkeletalMesh && SkeletalMesh->IsValidLODIndex(LODIndex))
+		FManagedArrayCollection GeometryCollection = GetValue<FManagedArrayCollection>(Context, &Collection);
+		GeometryCollection::Facades::FCollectionMeshFacade MeshFacade(GeometryCollection);
+		GeometryCollection::Facades::FCollectionCurveGeometryFacade CurveFacade(GeometryCollection);
+		
+		if(MeshFacade.IsValid() && CurveFacade.IsValid())
 		{
-			TArray<UE::Groom::Private::FCatmullRomSpline> Splines;
-			UE::Groom::Private::BuildSplines(SkeletalMesh, RootBones, SamplesPerSegment, Splines);
+			GeometryCollection::Facades::FVertexBoneWeightsFacade SkinningFacade(GeometryCollection, false);
+			SkinningFacade.DefineSchema();
 
-			if (CurvesType == EGroomCollectionType::Guides)
-			{
-				UE::Groom::Private::BuildSplineSkinningWeights<UE::Groom::FGroomGuidesFacade>(GroomCollection, GroupIndex, SkeletalMesh, LODIndex, Splines, ParamKey, BoneKey);
-			}
-			else
+			const TObjectPtr<USkeletalMesh> InputSkelMesh = GetValue(Context, &SkeletalMesh);
+			const int32 InputLOD = GetValue(Context, &LODIndex);
+
+			if(InputSkelMesh && InputSkelMesh->IsValidLODIndex(InputLOD))
 			{
-				UE::Groom::Private::BuildSplineSkinningWeights<UE::Groom::FGroomStrandsFacade>(GroomCollection, GroupIndex, SkeletalMesh, LODIndex, Splines, ParamKey, BoneKey);
-			}
+				const FDataflowVertexSelection InputSelection = GetValue<FDataflowVertexSelection>(Context, &VertexSelection);
+				const bool bValidSelection = InputSelection.IsValidForCollection(GeometryCollection);
+				
+				TArray<UE::Groom::Private::FCatmullRomSpline> Splines;
+				UE::Groom::Private::BuildSplines(InputSkelMesh,
+					GetValue(Context, &RootBones), GetValue(Context, &SamplesPerSegment), Splines);
 
+				UE::Groom::Private::BuildSplineSkinningWeights(GeometryCollection, MeshFacade, CurveFacade,
+						InputSelection, Splines, GetSplineParamKey(Context), GetSplineBonesKey(Context));
+				
+				for(int32 GeometryIndex = 0; GeometryIndex < SkinningFacade.NumGeometry(); ++GeometryIndex)
+				{
+					bool bValidGeometry = false;
+					for(int32 VertexIndex :  MeshFacade.GetVertexIndices(GeometryIndex))
+					{
+						if(!bValidSelection || (bValidSelection && InputSelection.IsSelected(VertexIndex)))
+						{
+							bValidGeometry = true;
+							break;
+						}
+					}
+					if(bValidGeometry)
+					{
+						SkinningFacade.ModifyGeometryBinding(GeometryIndex, InputSkelMesh, InputLOD);
+					}
+				}
+			}
 		}
 
-		SetValue(Context, MoveTemp(GroomCollection), &Collection);
+		SetValue(Context, MoveTemp(GeometryCollection), &Collection);
 	}
 	else if (Out->IsA<FCollectionAttributeKey>(&SplineParamKey))
 	{
-		SetValue(Context, GetSplineParamKey(), &SplineParamKey);
+		SetValue(Context, GetSplineParamKey(Context), &SplineParamKey);
 	}
-	else if (Out->IsA<FCollectionAttributeKey>(&SplineBoneKey))
+	else if (Out->IsA<FCollectionAttributeKey>(&SplineBonesKey))
 	{
-		SetValue(Context, GetSplineBonesKey(), &SplineBoneKey);
+		SetValue(Context, GetSplineBonesKey(Context), &SplineBonesKey);
 	}
 }
 
-FCollectionAttributeKey FBuildGroomSplineSkinWeightsNode::GetSplineParamKey() const
+FCollectionAttributeKey FBuildSplineSkinWeightsDataflowNode::GetSplineParamKey(UE::Dataflow::FContext& Context) const
 {
 	FCollectionAttributeKey Key;
-	Key.Group = (CurvesType == EGroomCollectionType::Guides) ? UE::Groom::FGroomGuidesFacade::VerticesGroup.ToString() :
-		UE::Groom::FGroomStrandsFacade::VerticesGroup.ToString();
-	Key.Attribute = SplineParamName;
+	Key.Group = FGeometryCollection::VerticesGroup.ToString();
+	Key.Attribute = UE::Groom::Private::FCollectionSplineAttributes::VertexSplineParamAttribute.ToString();
 	return Key;
 }
 
-FCollectionAttributeKey FBuildGroomSplineSkinWeightsNode::GetSplineBonesKey() const
+FCollectionAttributeKey FBuildSplineSkinWeightsDataflowNode::GetSplineBonesKey(UE::Dataflow::FContext& Context) const
 {
 	FCollectionAttributeKey Key;
-	Key.Group = (CurvesType == EGroomCollectionType::Guides) ? UE::Groom::FGroomGuidesFacade::VerticesGroup.ToString() :
-		UE::Groom::FGroomStrandsFacade::VerticesGroup.ToString();
-	Key.Attribute = SplineBonesParamName;
+	Key.Group = FGeometryCollection::VerticesGroup.ToString();
+	Key.Attribute = UE::Groom::Private::FCollectionSplineAttributes::VertexSplineBonesAttribute.ToString();
 	return Key;
 }
 
-TArray<UE::Dataflow::FRenderingParameter> FBuildGroomSplineSkinWeightsNode::GetRenderParametersImpl() const
-{
-	if (CurvesType == EGroomCollectionType::Guides)
-	{
-		return { {TEXT("GuidesRender"), FName("FGroomCollection"), {TEXT("Collection")}} };
-	}
-	else
-	{
-		return { {TEXT("StrandsRender"), FName("FGroomCollection"), {TEXT("Collection")}} };
-	}
-}
-
-FConvertSplineToLinearSkinWeightsNode::FConvertSplineToLinearSkinWeightsNode(const UE::Dataflow::FNodeParameters& InParam, FGuid InGuid)
+FSplineToLinearSkinWeightsDataflowNode::FSplineToLinearSkinWeightsDataflowNode(const UE::Dataflow::FNodeParameters& InParam, FGuid InGuid)
 	: FDataflowNode(InParam, InGuid)
 {
 	RegisterInputConnection(&Collection);
-	RegisterInputConnection(&SplineParamKey);
-	RegisterInputConnection(&SplineBonesKey);
+	RegisterInputConnection(&VertexSelection);
+	RegisterInputConnection(&SplineParamKey).SetCanHidePin(true).SetPinIsHidden(true);
+	RegisterInputConnection(&SplineBonesKey).SetCanHidePin(true).SetPinIsHidden(true);
+	
 	RegisterOutputConnection(&Collection, &Collection);
 	RegisterOutputConnection(&BoneIndicesKey);
 	RegisterOutputConnection(&BoneWeightsKey);
 }
 
-void FConvertSplineToLinearSkinWeightsNode::Evaluate(UE::Dataflow::FContext& Context, const FDataflowOutput* Out) const
+void FSplineToLinearSkinWeightsDataflowNode::Evaluate(UE::Dataflow::FContext& Context, const FDataflowOutput* Out) const
 {
 	if (Out->IsA<FManagedArrayCollection>(&Collection))
 	{
-		FManagedArrayCollection GroomCollection = GetValue<FManagedArrayCollection>(Context, &Collection);
-		FCollectionAttributeKey SplineParam = GetSplineParamKey(Context);
-		FCollectionAttributeKey SplineBones = GetSplineBonesKey(Context);
+		FManagedArrayCollection GeometryCollection = GetValue<FManagedArrayCollection>(Context, &Collection);
 
-		if (!SplineParam.Attribute.IsEmpty() && !SplineBones.Attribute.IsEmpty())
+		GeometryCollection::Facades::FCollectionMeshFacade MeshFacade(GeometryCollection);
+		
+		if(MeshFacade.IsValid())
 		{
-			if (CurvesType == EGroomCollectionType::Guides)
-			{
-				UE::Groom::Private::ConvertFromSplineToLinearWeights<UE::Groom::FGroomGuidesFacade>(GroomCollection, GroupIndex, SplineParam, SplineBones, GetBoneIndicesKey(), GetBoneWeightsKey());
-			}
-			else
+			GeometryCollection::Facades::FVertexBoneWeightsFacade SkinningFacade(GeometryCollection, false);
+			if(!SkinningFacade.IsValid())
 			{
-				UE::Groom::Private::ConvertFromSplineToLinearWeights<UE::Groom::FGroomStrandsFacade>(GroomCollection, GroupIndex, SplineParam, SplineBones, GetBoneIndicesKey(), GetBoneWeightsKey());
+				SkinningFacade.DefineSchema();
 			}
+			UE::Groom::Private::ConvertFromSplineToLinearWeights(GeometryCollection, MeshFacade, SkinningFacade,
+				GetValue<FDataflowVertexSelection>(Context, &VertexSelection), GetSplineParamKey(Context), GetSplineBonesKey(Context));
 		}
-		SetValue(Context, MoveTemp(GroomCollection), &Collection);
+		SetValue(Context, MoveTemp(GeometryCollection), &Collection);
 	}
 	else if (Out->IsA<FCollectionAttributeKey>(&BoneIndicesKey))
 	{
-		SetValue(Context, GetBoneIndicesKey(), &BoneIndicesKey);
+		SetValue(Context, UE::Groom::Private::GetBoneIndicesKey(), &BoneIndicesKey);
 	}
 	else if (Out->IsA<FCollectionAttributeKey>(&BoneWeightsKey))
 	{
-		SetValue(Context, GetBoneWeightsKey(), &BoneWeightsKey);
+		SetValue(Context, UE::Groom::Private::GetBoneWeightsKey(), &BoneWeightsKey);
 	}
 }
 
-FCollectionAttributeKey FConvertSplineToLinearSkinWeightsNode::GetBoneIndicesKey() const
-{
-	FCollectionAttributeKey Key;
-	Key.Group = (CurvesType == EGroomCollectionType::Guides) ? UE::Groom::FGroomGuidesFacade::VerticesGroup.ToString() :
-		UE::Groom::FGroomStrandsFacade::VerticesGroup.ToString();
-	Key.Attribute = UE::Groom::FGroomGuidesFacade::PointBoneIndicesAttribute.ToString();
-	return Key;
-}
-
-FCollectionAttributeKey FConvertSplineToLinearSkinWeightsNode::GetBoneWeightsKey() const
-{
-	FCollectionAttributeKey Key;
-	Key.Group = (CurvesType == EGroomCollectionType::Guides) ? UE::Groom::FGroomGuidesFacade::VerticesGroup.ToString() :
-		UE::Groom::FGroomStrandsFacade::VerticesGroup.ToString();
-	Key.Attribute = UE::Groom::FGroomGuidesFacade::PointBoneWeightsAttribute.ToString();
-	return Key;
-}
-
-
-
-FCollectionAttributeKey FConvertSplineToLinearSkinWeightsNode::GetSplineParamKey(UE::Dataflow::FContext& Context) const
+FCollectionAttributeKey FSplineToLinearSkinWeightsDataflowNode::GetSplineParamKey(UE::Dataflow::FContext& Context) const
 {
 	FCollectionAttributeKey Key = GetValue(Context, &SplineParamKey, SplineParamKey);
-
-	// If nothing set use default value
+	
 	if (Key.Attribute.IsEmpty() && Key.Group.IsEmpty())
 	{
-		Key.Group = (CurvesType == EGroomCollectionType::Guides) ? UE::Groom::FGroomGuidesFacade::VerticesGroup.ToString() :
-			UE::Groom::FGroomStrandsFacade::VerticesGroup.ToString();
-		Key.Attribute = SplineParamName;
+		Key.Group = FGeometryCollection::VerticesGroup.ToString();
+		Key.Attribute = UE::Groom::Private::FCollectionSplineAttributes::VertexSplineParamAttribute.ToString();
 	}
 
 	return Key;
 }
 
-FCollectionAttributeKey FConvertSplineToLinearSkinWeightsNode::GetSplineBonesKey(UE::Dataflow::FContext& Context) const
+FCollectionAttributeKey FSplineToLinearSkinWeightsDataflowNode::GetSplineBonesKey(UE::Dataflow::FContext& Context) const
 {
 	FCollectionAttributeKey Key = GetValue(Context, &SplineBonesKey, SplineBonesKey);
 
 	if (Key.Attribute.IsEmpty() && Key.Group.IsEmpty())
 	{
-		Key.Group = (CurvesType == EGroomCollectionType::Guides) ? UE::Groom::FGroomGuidesFacade::VerticesGroup.ToString() :
-			UE::Groom::FGroomStrandsFacade::VerticesGroup.ToString();
-		Key.Attribute = SplineBonesName;
+		Key.Group = FGeometryCollection::VerticesGroup.ToString();
+		Key.Attribute = UE::Groom::Private::FCollectionSplineAttributes::VertexSplineBonesAttribute.ToString();
 	}
 
 	return Key;
 }
 
-TArray<UE::Dataflow::FRenderingParameter> FConvertSplineToLinearSkinWeightsNode::GetRenderParametersImpl() const
-{
-	if (CurvesType == EGroomCollectionType::Guides)
-	{
-		return { {TEXT("GuidesRender"), FName("FGroomCollection"), {TEXT("Collection")}} };
-	}
-	else
-	{
-		return { {TEXT("StrandsRender"), FName("FGroomCollection"), {TEXT("Collection")}} };
-	}
-}
-
-FConvertLinearToSplineSkinWeightsNode::FConvertLinearToSplineSkinWeightsNode(const UE::Dataflow::FNodeParameters& InParam, FGuid InGuid)
+FLinearToSplineSkinWeightsDataflowNode::FLinearToSplineSkinWeightsDataflowNode(const UE::Dataflow::FNodeParameters& InParam, FGuid InGuid)
 	: FDataflowNode(InParam, InGuid)
 {
 	RegisterInputConnection(&Collection);
+	RegisterInputConnection(&VertexSelection);
+	RegisterInputConnection(&BoneIndicesKey).SetCanHidePin(true).SetPinIsHidden(true);
+	RegisterInputConnection(&BoneWeightsKey).SetCanHidePin(true).SetPinIsHidden(true);
+	
 	RegisterOutputConnection(&Collection, &Collection);
 	RegisterOutputConnection(&SplineParamKey);
-	RegisterOutputConnection(&SplineBoneKey);
+	RegisterOutputConnection(&SplineBonesKey);
 }
 
-void FConvertLinearToSplineSkinWeightsNode::Evaluate(UE::Dataflow::FContext& Context, const FDataflowOutput* Out) const
+void FLinearToSplineSkinWeightsDataflowNode::Evaluate(UE::Dataflow::FContext& Context, const FDataflowOutput* Out) const
 {
 	using namespace UE::Dataflow;
 
 	if (Out->IsA<FManagedArrayCollection>(&Collection))
 	{
-		FManagedArrayCollection GroomCollection = GetValue<FManagedArrayCollection>(Context, &Collection);
-
-		if (CurvesType == EGroomCollectionType::Guides)
-		{
-			UE::Groom::Private::ConvertFromLinearToSplineWeights<UE::Groom::FGroomGuidesFacade>(GroomCollection, GroupIndex, GetSplineParamKey(), GetSplineBonesKey());
-		}
-		else
+		FManagedArrayCollection GeometryCollection = GetValue<FManagedArrayCollection>(Context, &Collection);
+		
+		GeometryCollection::Facades::FCollectionMeshFacade MeshFacade(GeometryCollection);
+		GeometryCollection::Facades::FVertexBoneWeightsFacade SkinningFacade(GeometryCollection);
+		
+		if(MeshFacade.IsValid() && SkinningFacade.IsValid())
 		{
-			UE::Groom::Private::ConvertFromLinearToSplineWeights<UE::Groom::FGroomStrandsFacade>(GroomCollection, GroupIndex, GetSplineParamKey(), GetSplineBonesKey());
+			UE::Groom::Private::ConvertFromLinearToSplineWeights(GeometryCollection, MeshFacade, SkinningFacade,
+				GetValue<FDataflowVertexSelection>(Context, &VertexSelection), UE::Groom::Private::GetSplineParamKey(), UE::Groom::Private::GetSplineBonesKey());
 		}
 
-		SetValue(Context, MoveTemp(GroomCollection), &Collection);
+		SetValue(Context, MoveTemp(GeometryCollection), &Collection);
 	}
 	else if (Out->IsA<FCollectionAttributeKey>(&SplineParamKey))
 	{
-		SetValue(Context, GetSplineParamKey(), &SplineParamKey);
+		SetValue(Context, UE::Groom::Private::GetSplineParamKey(), &SplineParamKey);
 	}
-	else if (Out->IsA<FCollectionAttributeKey>(&SplineBoneKey))
+	else if (Out->IsA<FCollectionAttributeKey>(&SplineBonesKey))
 	{
-		SetValue(Context, GetSplineBonesKey(), &SplineBoneKey);
+		SetValue(Context, UE::Groom::Private::GetSplineBonesKey(), &SplineBonesKey);
 	}
-
 }
 
-TArray<UE::Dataflow::FRenderingParameter> FConvertLinearToSplineSkinWeightsNode::GetRenderParametersImpl() const
-{
-	if (CurvesType == EGroomCollectionType::Guides)
-	{
-		return { {TEXT("GuidesRender"), FName("FGroomCollection"), {TEXT("Collection")}} };
-	}
-	else
-	{
-		return { {TEXT("StrandsRender"), FName("FGroomCollection"), {TEXT("Collection")}} };
-	}
-}
 
-FCollectionAttributeKey FConvertLinearToSplineSkinWeightsNode::GetSplineParamKey() const
-{
-	FCollectionAttributeKey Key;
-	Key.Group = (CurvesType == EGroomCollectionType::Guides) ? UE::Groom::FGroomGuidesFacade::VerticesGroup.ToString() :
-		UE::Groom::FGroomStrandsFacade::VerticesGroup.ToString();
-	Key.Attribute = SplineParamName;
-	return Key;
-}
-
-FCollectionAttributeKey FConvertLinearToSplineSkinWeightsNode::GetSplineBonesKey() const
-{
-	FCollectionAttributeKey Key;
-	Key.Group = (CurvesType == EGroomCollectionType::Guides) ? UE::Groom::FGroomGuidesFacade::VerticesGroup.ToString() :
-		UE::Groom::FGroomStrandsFacade::VerticesGroup.ToString();
-	Key.Attribute = SplineBonesName;
-	return Key;
-}
 
diff --git a/Engine/Plugins/Runtime/HairStrands/Source/HairStrandsDataflow/Private/BuildGuidesLODsNode.cpp b/Engine/Plugins/Runtime/HairStrands/Source/HairStrandsDataflow/Private/BuildGuidesLODsNode.cpp
--- a/Engine/Plugins/Runtime/HairStrands/Source/HairStrandsDataflow/Private/BuildGuidesLODsNode.cpp
+++ b/Engine/Plugins/Runtime/HairStrands/Source/HairStrandsDataflow/Private/BuildGuidesLODsNode.cpp
@@ -1,21 +1,21 @@
 // Copyright Epic Games, Inc. All Rights Reserved.
 
 #include "BuildGuidesLODsNode.h"
-#include "GroomCollectionFacades.h"
 #include "Dataflow/DataflowObjectInterface.h"
+#include "GeometryCollection/Facades/CollectionCurveFacade.h"
 
 #include UE_INLINE_GENERATED_CPP_BY_NAME(BuildGuidesLODsNode)
 
 namespace UE::Groom::Private
 {
-	FORCEINLINE float ComputeGuidesMetric(const FGroomGuidesFacade& GuidesFacade, const uint32 GuideIndexA, const uint32 GuideIndexB, const float& GuideLengthA, const float& GuideLengthB,
+	FORCEINLINE float ComputeGuidesMetric(const TArray<FVector3f>& PointRestPositions, const TArray<int32>& CurvePointOffsets,const uint32 GuideIndexA, const uint32 GuideIndexB, const float& GuideLengthA, const float& GuideLengthB,
 										 const float ShapeWeight, const float ProximityWeight)
 	{
-		const uint32 PointOffsetA = (GuideIndexA > 0) ? GuidesFacade.GetCurvePointOffsets()[GuideIndexA-1] : 0;
-		const uint32 PointOffsetB = (GuideIndexB > 0) ? GuidesFacade.GetCurvePointOffsets()[GuideIndexB-1] : 0;
+		const uint32 PointOffsetA = (GuideIndexA > 0) ? CurvePointOffsets[GuideIndexA-1] : 0;
+		const uint32 PointOffsetB = (GuideIndexB > 0) ? CurvePointOffsets[GuideIndexB-1] : 0;
 		
-		const uint32 NumPointsA = GuidesFacade.GetCurvePointOffsets()[GuideIndexA] - PointOffsetA;
-		const uint32 NumPointsB = GuidesFacade.GetCurvePointOffsets()[GuideIndexB] - PointOffsetB;
+		const uint32 NumPointsA = CurvePointOffsets[GuideIndexA] - PointOffsetA;
+		const uint32 NumPointsB = CurvePointOffsets[GuideIndexB] - PointOffsetB;
 
 		if(NumPointsA == NumPointsB)
 		{
@@ -23,11 +23,11 @@ namespace UE::Groom::Private
 			float ShapeMetric = 0.0f;
 			for(uint32 PointIndex = 0; PointIndex < NumPointsA; ++PointIndex)
 			{
-				const FVector3f& GuidePositionA = GuidesFacade.GetPointRestPositions()[PointIndex+PointOffsetA];
-				const FVector3f& GuidePositionB = GuidesFacade.GetPointRestPositions()[PointIndex+PointOffsetB];
+				const FVector3f& GuidePositionA = PointRestPositions[PointIndex+PointOffsetA];
+				const FVector3f& GuidePositionB = PointRestPositions[PointIndex+PointOffsetB];
 				
 				ProximityMetric += (GuidePositionB - GuidePositionA).Size();
-				ShapeMetric += (GuidePositionB - GuidesFacade.GetPointRestPositions()[PointOffsetB] - GuidePositionA + GuidesFacade.GetPointRestPositions()[PointOffsetA]).Size();
+				ShapeMetric += (GuidePositionB - PointRestPositions[PointOffsetB] - GuidePositionA + PointRestPositions[PointOffsetA]).Size();
 			}
 
 			const float MetricScale = 1.0f / (NumPointsA * 0.5f * (GuideLengthA+GuideLengthB));
@@ -35,89 +35,124 @@ namespace UE::Groom::Private
 		}
 		return 0.0f;
 	}
-
-	FORCEINLINE void BuildGuidesLODs(FGroomGuidesFacade& GuidesFacade)
+	
+	FORCEINLINE void BuildInternalLODs(const int32 NumObjects, const int32 NumCurves, const TArray<int32>& GeometryCurveOffsets,
+		const TArray<int32>& CurvePointOffsets, const TArray<FVector3f>& PointRestPositions, TArray<int32>& CurveParentIndices, TArray<int32>& CurveLodIndices,
+			const FDataflowCurveSelection& CurveSelection, const bool bValidSelection)
 	{
-		const uint32 NumObjects = GuidesFacade.GetNumObjects();
-		
-		TArray<int32> CurveParentIndices; CurveParentIndices.Init(INDEX_NONE, GuidesFacade.GetNumCurves());
-		TArray<int32> CurveLodIndices; CurveLodIndices.Init(INDEX_NONE, GuidesFacade.GetNumCurves());
+		CurveParentIndices.Init(INDEX_NONE, NumCurves);
+		CurveLodIndices.Init(INDEX_NONE, NumCurves);
 		
 		uint32 CurveOffset = 0, PointOffset = 0;
-		for(uint32 ObjectIndex = 0; ObjectIndex < NumObjects; ++ObjectIndex)
+		for(int32 ObjectIndex = 0; ObjectIndex < NumObjects; ++ObjectIndex)
 		{ 
-			TArray<float> GuidesLengths; GuidesLengths.Init(0.0f, GuidesFacade.GetObjectCurveOffsets()[ObjectIndex] - CurveOffset);
-			for(int32 GuideIndex = CurveOffset; GuideIndex < GuidesFacade.GetObjectCurveOffsets()[ObjectIndex]; ++GuideIndex)
+			TArray<float> GuidesLengths; GuidesLengths.Init(0.0f, GeometryCurveOffsets[ObjectIndex] - CurveOffset);
+			for(int32 GuideIndex = CurveOffset; GuideIndex < GeometryCurveOffsets[ObjectIndex]; ++GuideIndex)
 			{
-				for(int32 PointIndex = PointOffset+1; PointIndex < GuidesFacade.GetCurvePointOffsets()[GuideIndex]; ++PointIndex)
+				for(int32 PointIndex = PointOffset+1; PointIndex < CurvePointOffsets[GuideIndex]; ++PointIndex)
 				{
-					GuidesLengths[GuideIndex-CurveOffset] += (GuidesFacade.GetPointRestPositions()[PointIndex]-GuidesFacade.GetPointRestPositions()[PointIndex-1]).Size();
+					GuidesLengths[GuideIndex-CurveOffset] += (PointRestPositions[PointIndex]-PointRestPositions[PointIndex-1]).Size();
 				}
-				PointOffset = GuidesFacade.GetCurvePointOffsets()[GuideIndex];
+				PointOffset = CurvePointOffsets[GuideIndex];
 			}
 			const int32 NumObjectLODs = FMath::CeilLogTwo(GuidesLengths.Num());
-			ParallelFor(GuidesLengths.Num(), [CurveOffset, &GuidesLengths, &GuidesFacade, &CurveParentIndices, &CurveLodIndices, &NumObjectLODs](int32 CurveIndex)
+			ParallelFor(GuidesLengths.Num(), [CurveOffset, &GuidesLengths, &PointRestPositions, &CurvePointOffsets, &CurveParentIndices, &CurveLodIndices, &NumObjectLODs, bValidSelection, &CurveSelection](int32 CurveIndex)
 			{
 				const uint32 GuideIndex = CurveIndex + CurveOffset;
-				const uint32 GuideLod = NumObjectLODs - 1 - FMath::FloorLog2(CurveIndex);
-				const uint32 LodOffset = (GuideLod == NumObjectLODs - 1) ? 0 : 1 << FMath::FloorLog2(CurveIndex);
 
-				float MinMetric = FLT_MAX;
-				int32 MinParent = INDEX_NONE;
-				for(uint32 ParentIndex = CurveOffset; ParentIndex < CurveOffset+LodOffset; ++ParentIndex)
+				if(!bValidSelection || (bValidSelection && (CurveSelection.IsSelected(GuideIndex))))
 				{
-					const float ParentMetric = 1.0f - ComputeGuidesMetric(GuidesFacade,
-						GuideIndex, ParentIndex, GuidesLengths[CurveIndex], GuidesLengths[ParentIndex-CurveOffset],  1.0f, 1.0f);
-					if (ParentMetric < MinMetric)
+					const uint32 GuideLod = NumObjectLODs - 1 - FMath::FloorLog2(CurveIndex);
+					const uint32 LodOffset = (GuideLod == NumObjectLODs - 1) ? 0 : 1 << FMath::FloorLog2(CurveIndex);
+
+					float MinMetric = FLT_MAX;
+					int32 MinParent = INDEX_NONE;
+					for(uint32 ParentIndex = CurveOffset; ParentIndex < CurveOffset+LodOffset; ++ParentIndex)
 					{
-						MinMetric = ParentMetric;
-						MinParent = ParentIndex;
+						const float ParentMetric = 1.0f - ComputeGuidesMetric(PointRestPositions, CurvePointOffsets,
+							GuideIndex, ParentIndex, GuidesLengths[CurveIndex], GuidesLengths[ParentIndex-CurveOffset],  1.0f, 1.0f);
+						if (ParentMetric < MinMetric)
+						{
+							MinMetric = ParentMetric;
+							MinParent = ParentIndex;
+						}
 					}
+					CurveParentIndices[GuideIndex] = MinParent;
+					CurveLodIndices[GuideIndex] = GuideLod;
 				}
-				CurveParentIndices[GuideIndex] = MinParent;
-				CurveLodIndices[GuideIndex] = GuideLod;
 
 				//UE_LOG(LogTemp, Log, TEXT("Guide Index = %d, Num Guides = %d, Num LODs = %d, Guide LOD = %d, LOD Offset = %d, Parent Index = %d"),
 				//	CurveIndex, GuidesLengths.Num(), NumObjectLODs, GuideLod, LodOffset, MinParent);
 			}, EParallelForFlags::None);
 
-			GuidesFacade.SetCurveParentIndices(CurveParentIndices);
-			GuidesFacade.SetCurveLodIndices(CurveLodIndices);
-			CurveOffset = GuidesFacade.GetObjectCurveOffsets()[ObjectIndex];
+			CurveOffset = GeometryCurveOffsets[ObjectIndex];
 		}
 	}
+	
+	FORCEINLINE void BuildCurvesLODs(GeometryCollection::Facades::FCollectionCurveGeometryFacade& CurvesFacade, GeometryCollection::Facades::FCollectionCurveHierarchyFacade& HierarchyFacade, const FDataflowCurveSelection& CurveSelection)
+	{
+		const uint32 NumObjects = CurvesFacade.GetNumGeometry();
+		
+		TArray<int32> CurveParentIndices; CurveParentIndices.Init(INDEX_NONE, CurvesFacade.GetNumCurves());
+		TArray<int32> CurveLodIndices; CurveLodIndices.Init(INDEX_NONE, CurvesFacade.GetNumCurves());
+
+		const bool bValidSelection = CurveSelection.IsValidForCollection(CurvesFacade.GetManagedArrayCollection());
+
+		BuildInternalLODs(CurvesFacade.GetNumGeometry(), CurvesFacade.GetNumCurves(), CurvesFacade.GetGeometryCurveOffsets(),
+			CurvesFacade.GetCurvePointOffsets(), CurvesFacade.GetPointRestPositions(), CurveParentIndices, CurveLodIndices, CurveSelection, bValidSelection);
+		
+		HierarchyFacade.SetCurveParentIndices(CurveParentIndices);
+		HierarchyFacade.SetCurveLodIndices(CurveLodIndices);
+	}
 
 	FCollectionAttributeKey GetCurveParentsKey()
     {
     	FCollectionAttributeKey Key;
-    	Key.Group =  UE::Groom::FGroomGuidesFacade::CurvesGroup.ToString();
-    	Key.Attribute = UE::Groom::FGroomGuidesFacade::CurveParentIndicesAttribute.ToString();
+    	Key.Group =  GeometryCollection::Facades::FCollectionCurveGeometryFacade::CurvesGroup.ToString();
+    	Key.Attribute = GeometryCollection::Facades::FCollectionCurveHierarchyFacade::CurveParentIndicesAttribute.ToString();
     	return Key;
     }
 
 	FCollectionAttributeKey GetCurveLodsKey()
 	{
 		FCollectionAttributeKey Key;
-		Key.Group =  UE::Groom::FGroomGuidesFacade::CurvesGroup.ToString();
-		Key.Attribute = UE::Groom::FGroomGuidesFacade::CurveLodIndicesAttribute.ToString();
+		Key.Group =  GeometryCollection::Facades::FCollectionCurveGeometryFacade::CurvesGroup.ToString();
+		Key.Attribute = GeometryCollection::Facades::FCollectionCurveHierarchyFacade::CurveLodIndicesAttribute.ToString();
 		return Key;
 	}
 }
 
 void FBuildGuidesLODsDataflowNode::Evaluate(UE::Dataflow::FContext& Context, const FDataflowOutput* Out) const
+{
+}
+
+FBuildCurveLODsDataflowNode::FBuildCurveLODsDataflowNode(const UE::Dataflow::FNodeParameters& InParam, FGuid InGuid)
+	: FDataflowNode(InParam, InGuid)
+{
+	RegisterInputConnection(&Collection);
+	RegisterInputConnection(&CurveSelection);
+	RegisterOutputConnection(&Collection, &Collection);
+	RegisterOutputConnection(&CurveParentsKey);
+	RegisterOutputConnection(&CurveLodsKey);
+}
+
+void FBuildCurveLODsDataflowNode::Evaluate(UE::Dataflow::FContext& Context, const FDataflowOutput* Out) const
 {
 	if (Out->IsA<FManagedArrayCollection>(&Collection))
 	{
-		FManagedArrayCollection GroomCollection = GetValue<FManagedArrayCollection>(Context, &Collection);
+		FManagedArrayCollection GeometryCollection = GetValue<FManagedArrayCollection>(Context, &Collection);
 		
-		UE::Groom::FGroomGuidesFacade GuidesFacade(GroomCollection);
+		GeometryCollection::Facades::FCollectionCurveGeometryFacade CurvesFacade(GeometryCollection);
 
-		if(GuidesFacade.IsValid())
+		if(CurvesFacade.IsValid())
 		{
-			UE::Groom::Private::BuildGuidesLODs(GuidesFacade);
+			GeometryCollection::Facades::FCollectionCurveHierarchyFacade HierarchyFacade(GeometryCollection);
+			HierarchyFacade.DefineSchema();
+			
+			UE::Groom::Private::BuildCurvesLODs(CurvesFacade, HierarchyFacade, GetValue(Context, &CurveSelection));
 		}
 		
-		SetValue(Context, MoveTemp(GroomCollection), &Collection);
+		SetValue(Context, MoveTemp(GeometryCollection), &Collection);
 	}
 	else if (Out->IsA<FCollectionAttributeKey>(&CurveParentsKey))
 	{
diff --git a/Engine/Plugins/Runtime/HairStrands/Source/HairStrandsDataflow/Private/GenerateGuidesCurvesNode.cpp b/Engine/Plugins/Runtime/HairStrands/Source/HairStrandsDataflow/Private/GenerateGuidesCurvesNode.cpp
--- a/Engine/Plugins/Runtime/HairStrands/Source/HairStrandsDataflow/Private/GenerateGuidesCurvesNode.cpp
+++ b/Engine/Plugins/Runtime/HairStrands/Source/HairStrandsDataflow/Private/GenerateGuidesCurvesNode.cpp
@@ -3,110 +3,145 @@
 #include "GenerateGuidesCurvesNode.h"
 
 #include "GroomBindingBuilder.h"
-#include "GroomCollectionFacades.h"
 #include "Dataflow/DataflowObjectInterface.h"
+#include "GeometryCollection/Facades/CollectionCurveFacade.h"
 
 #include UE_INLINE_GENERATED_CPP_BY_NAME(GenerateGuidesCurvesNode)
 
 namespace UE::Groom::Private
 {
-	FORCEINLINE uint32 SampleStrandsCurves(const UE::Groom::FGroomStrandsFacade& StrandsFacade, const uint32 GuidesCount, const uint32 GuidesOffset,
-		const uint32 ObjectIndex, const uint32 PrevStrands, const uint32 NextStrands, TArray<uint32>& SampleIndices)
+	FORCEINLINE void GenerateStrandsCurves(const GeometryCollection::Facades::FCollectionCurveGeometryFacade& CurvesFacade, const FDataflowCurveSelection& CurveSelection, const uint32 GuidesCount, TArray<uint32>& SampleIndices)
 	{
-		const uint32 NumStrands = NextStrands-PrevStrands;
-		uint32 NumGuides = (ObjectIndex == StrandsFacade.GetNumObjects()-1) ? GuidesCount - GuidesOffset :
-			GuidesCount * static_cast<float>(NumStrands)/StrandsFacade.GetNumCurves();
-		NumGuides = FMath::Max(1u, NumGuides);
-
 		TArray<FVector3f> RootPositions;
-		RootPositions.Init(FVector3f::ZeroVector, NumStrands);
-		for(uint32 StrandIndex = 0; StrandIndex < NumStrands; ++StrandIndex)
-		{
-			const uint32 RootIndex = (StrandIndex+PrevStrands == 0) ? 0 : StrandsFacade.GetCurvePointOffsets()[StrandIndex+PrevStrands-1];
-			RootPositions[StrandIndex] = StrandsFacade.GetPointRestPositions()[RootIndex];
-		}
+		RootPositions.Reserve(CurvesFacade.GetNumCurves());
 
 		TArray<bool> ValidPoints;
-		ValidPoints.Init(true, NumStrands);
-				
-		GroomBinding_RBFWeighting::FPointsSampler PointsSampler(ValidPoints, RootPositions.GetData(), NumGuides);
+		ValidPoints.Reserve(CurvesFacade.GetNumCurves());
+
+		const bool bValidSelection = CurveSelection.IsValidForCollection(CurvesFacade.GetManagedArrayCollection());;
+			
+		int32 PointOffset = 0;
+		for (int32 CurveIndex = 0; CurveIndex < CurvesFacade.GetNumCurves(); ++CurveIndex)
+		{
+			RootPositions.Add(CurvesFacade.GetPointRestPositions()[PointOffset]);
+			ValidPoints.Add( (!bValidSelection || (bValidSelection && CurveSelection.IsSelected(CurveIndex))));
+			
+			PointOffset = CurvesFacade.GetCurvePointOffsets()[CurveIndex];
+		}
 
+		GroomBinding_RBFWeighting::FPointsSampler PointsSampler(ValidPoints, RootPositions.GetData(), GuidesCount);
 		SampleIndices = PointsSampler.SampleIndices;
-		return NumGuides;
 	}
 
-	FORCEINLINE void BuildGuidesCurves(const UE::Groom::FGroomStrandsFacade& StrandsFacade, const uint32 PrevStrands, const TArray<uint32>& SampleIndices,
-		TArray<FVector3f>& PointRestPositions, TArray<int32>& ObjectCurveOffsets, TArray<int32>& CurvePointOffsets, TArray<int32>& CurveStrandIndices)
+	FORCEINLINE void GenerateStrandsPoints(const GeometryCollection::Facades::FCollectionCurveGeometryFacade& StrandsFacade, const TArray<uint32>& SampleIndices,
+		const bool bMergeGuides, GeometryCollection::Facades::FCollectionCurveGeometryFacade& GuidesFacade)
 	{
-		for(uint32 SampleIndex : SampleIndices)
+		TArray<TArray<FVector3f>> GeometryPointPositions;
+		TArray<TArray<int32>> GeometryCurvePoints;
+		TArray<TArray<int32>> GeometrySourceIndices;
+			
+		GeometryPointPositions.SetNum(StrandsFacade.GetNumGeometry());
+		GeometryCurvePoints.SetNum(StrandsFacade.GetNumGeometry());
+		GeometrySourceIndices.SetNum(StrandsFacade.GetNumGeometry());
+		
+		if(bMergeGuides && GuidesFacade.GetNumGeometry() == StrandsFacade.GetNumGeometry())
 		{
-			const uint32 GuideIndex = SampleIndex + PrevStrands;
-			const int32 PointBegin = (GuideIndex == 0) ? 0 : StrandsFacade.GetCurvePointOffsets()[GuideIndex-1];
-			const int32 PointEnd = StrandsFacade.GetCurvePointOffsets()[GuideIndex];
+			int32 PointOffset = 0;
+			int32 CurveOffset = 0;
+			for (int32 GeometryIndex = 0; GeometryIndex < GuidesFacade.GetNumGeometry(); ++GeometryIndex)
+			{
+				for (int32 CurveIndex = CurveOffset; CurveIndex < GuidesFacade.GetGeometryCurveOffsets()[GeometryIndex]; ++CurveIndex)
+				{
+					for (int32 PointIndex = PointOffset; PointIndex < GuidesFacade.GetCurvePointOffsets()[CurveIndex]; ++PointIndex)
+					{
+						GeometryPointPositions[GeometryIndex].Add(GuidesFacade.GetPointRestPositions()[PointIndex]);
+					}
+					GeometryCurvePoints[GeometryIndex].Add(GuidesFacade.GetCurvePointOffsets()[CurveIndex]-PointOffset);
+					GeometrySourceIndices[GeometryIndex].Add(GuidesFacade.GetCurveSourceIndices()[CurveIndex]);
+					PointOffset = GuidesFacade.GetCurvePointOffsets()[CurveIndex];
+				}
+				CurveOffset = GuidesFacade.GetGeometryCurveOffsets()[GeometryIndex];
+			}
+		}
+			
+		for(const uint32& SampleIndex : SampleIndices)
+		{
+			const int32 PrevPoint = (SampleIndex == 0) ? 0 : StrandsFacade.GetCurvePointOffsets()[SampleIndex-1];
+			const int32 NextPoint = StrandsFacade.GetCurvePointOffsets()[SampleIndex];
 
-			for(int32 PointIndex = PointBegin; PointIndex < PointEnd; ++PointIndex)
+			const int32 GeometryIndex = StrandsFacade.GetCurveGeometryIndices()[SampleIndex];
+			for(int32 PointIndex = PrevPoint; PointIndex < NextPoint; ++PointIndex)
 			{
-				PointRestPositions.Add(StrandsFacade.GetPointRestPositions()[PointIndex]);
+				GeometryPointPositions[GeometryIndex].Add(StrandsFacade.GetPointRestPositions()[PointIndex]);
 			}
-			CurvePointOffsets.Add(PointRestPositions.Num());
-			CurveStrandIndices.Add(SampleIndex);
+			GeometryCurvePoints[GeometryIndex].Add(NextPoint-PrevPoint);
+			GeometrySourceIndices[GeometryIndex].Add(SampleIndex);
+		}
+
+		TArray<FVector3f> PointRestPositions;
+		TArray<int32> GeometryCurveOffsets, CurvePointOffsets, CurveSourceIndices;
+		TArray<FString> GeometryGroupNames;
+		
+		int32 CurveOffset = 0;
+		int32 PointOffset = 0;
+		for (int32 GeometryIndex = 0; GeometryIndex < StrandsFacade.GetNumGeometry(); ++GeometryIndex)
+		{
+			PointRestPositions.Append(GeometryPointPositions[GeometryIndex]);
+			for (int32 CurveIndex = 0; CurveIndex < GeometryCurvePoints[GeometryIndex].Num(); ++CurveIndex)
+			{
+				PointOffset += GeometryCurvePoints[GeometryIndex][CurveIndex];
+				CurvePointOffsets.Add(PointOffset);
+				CurveSourceIndices.Add(GeometrySourceIndices[GeometryIndex][CurveIndex]);
+			}
+			CurveOffset += GeometryCurvePoints[GeometryIndex].Num();
+			GeometryCurveOffsets.Add(CurveOffset);
+				
+			const FString GroupName = StrandsFacade.GetGeometryGroupNames()[GeometryIndex].Replace(TEXT("Strands"), TEXT("Guides"));
+			GeometryGroupNames.Add(GroupName);
 		}
-		ObjectCurveOffsets.Add(CurvePointOffsets.Num());
+		GuidesFacade.InitCurvesCollection(PointRestPositions, CurvePointOffsets, 
+		GeometryCurveOffsets, GeometryGroupNames,  StrandsFacade.GetGeometryCurveThickness(), CurveSourceIndices);
 	}
+
+	
 }
 
 void FGenerateGuidesCurvesDataflowNode::Evaluate(UE::Dataflow::FContext& Context, const FDataflowOutput* Out) const
+{}
+
+FGenerateCurveGeometryDataflowNode::FGenerateCurveGeometryDataflowNode(const UE::Dataflow::FNodeParameters& InParam, FGuid InGuid)
+	: FDataflowNode(InParam, InGuid)
+{
+	RegisterInputConnection(&Collection);
+	RegisterInputConnection(&SourceCurves);
+	RegisterInputConnection(&CurveSelection);
+	RegisterInputConnection(&CurveCount).SetCanHidePin(true).SetPinIsHidden(true);
+	RegisterOutputConnection(&Collection, &Collection);
+}
+
+void FGenerateCurveGeometryDataflowNode::Evaluate(UE::Dataflow::FContext& Context, const FDataflowOutput* Out) const
 {
 	if (Out->IsA<FManagedArrayCollection>(&Collection))
 	{
-		FManagedArrayCollection GroomCollection = GetValue<FManagedArrayCollection>(Context, &Collection);
+		FManagedArrayCollection TargetCollection = GetValue<FManagedArrayCollection>(Context, &Collection);
+		const FManagedArrayCollection& SourceCollection = GetValue<FManagedArrayCollection>(Context, &SourceCurves);
+		const int32 CountValue = FMath::Max(1, GetValue<int32>(Context, &CurveCount));
+		const FDataflowCurveSelection& DataflowSelection = GetValue<FDataflowCurveSelection>(Context, &CurveSelection);
+		
+		GeometryCollection::Facades::FCollectionCurveGeometryFacade SourceFacade(SourceCollection);
 		
-		UE::Groom::FGroomStrandsFacade StrandsFacade(GroomCollection);
-		if(StrandsFacade.IsValid())
+		if(SourceFacade.IsValid() && CountValue > 0)
 		{
-			TArray<FVector3f> PointRestPositions;
-			TArray<int32> ObjectCurveOffsets, CurvePointOffsets, CurveStrandIndices;
-			TArray<FString> ObjectGroupNames;
-			
-			uint32 PrevOffset = 0, GuidesOffset = 0;
-			for (int32 ObjectIndex = 0; ObjectIndex < StrandsFacade.GetNumObjects(); ++ObjectIndex)
-			{
-				const uint32 NextOffset = StrandsFacade.GetObjectCurveOffsets()[ObjectIndex];
-
-				TArray<uint32> SampleIndices;
-				
-				// Sample guides among the input strands
-				const uint32 NumGuides = UE::Groom::Private::SampleStrandsCurves(StrandsFacade, GuidesCount, GuidesOffset, ObjectIndex, PrevOffset, NextOffset, SampleIndices);
-
-				// Build the sampled guides
-				UE::Groom::Private::BuildGuidesCurves(StrandsFacade, PrevOffset, SampleIndices, PointRestPositions, ObjectCurveOffsets, CurvePointOffsets, CurveStrandIndices);
-				
-				const FString GroupName = StrandsFacade.GetObjectGroupNames()[ObjectIndex].Replace(
-					*UE::Groom::FGroomStrandsFacade::GroupPrefix.ToString(), *UE::Groom::FGroomGuidesFacade::GroupPrefix.ToString());
-				
-				ObjectGroupNames.Add(GroupName);
-				
-				GuidesOffset += NumGuides;
-				PrevOffset = NextOffset;
-			}
-			
-			UE::Groom::FGroomGuidesFacade GuidesFacade(GroomCollection);
-			const TArray<int32> ObjectPointSamples = GuidesFacade.GetObjectPointSamples();
-
-			// Init the groom collection
-			GuidesFacade.InitGroomCollection(PointRestPositions, CurvePointOffsets, ObjectCurveOffsets, ObjectGroupNames);
+			TArray<uint32> SampleIndices;
+			UE::Groom::Private::GenerateStrandsCurves(SourceFacade, DataflowSelection, CountValue, SampleIndices);
 
-			// Set the curve strand indices
-			GuidesFacade.SetCurveStrandIndices(CurveStrandIndices);
-
-			// Set the point samples if already defined and if the size is matching
-			if(ObjectPointSamples.Num() == GuidesFacade.GetNumObjects())
-			{
-				GuidesFacade.SetObjectPointSamples(ObjectPointSamples);
-			}
+			GeometryCollection::Facades::FCollectionCurveGeometryFacade TargetFacade(TargetCollection);
+			UE::Groom::Private::GenerateStrandsPoints(SourceFacade, SampleIndices, bMergeCurves, TargetFacade);
 		}
 		
-		SetValue(Context, MoveTemp(GroomCollection), &Collection);
+		SetValue(Context, MoveTemp(TargetCollection), &Collection);
 	}
 }
 
+
+
diff --git a/Engine/Plugins/Runtime/HairStrands/Source/HairStrandsDataflow/Private/GetGroomAssetNode.cpp b/Engine/Plugins/Runtime/HairStrands/Source/HairStrandsDataflow/Private/GetGroomAssetNode.cpp
--- a/Engine/Plugins/Runtime/HairStrands/Source/HairStrandsDataflow/Private/GetGroomAssetNode.cpp
+++ b/Engine/Plugins/Runtime/HairStrands/Source/HairStrandsDataflow/Private/GetGroomAssetNode.cpp
@@ -1,108 +1,110 @@
-﻿// Copyright Epic Games, Inc. All Rights Reserved.
+// Copyright Epic Games, Inc. All Rights Reserved.
 
 #include "GetGroomAssetNode.h"
 #include "GroomEdit.h"
 #include "GroomInstance.h"
-#include "GroomCollectionFacades.h"
 #include "Dataflow/DataflowInputOutput.h"
 #include "Dataflow/DataflowObjectInterface.h"
+#include "GeometryCollection/Facades/CollectionCurveFacade.h"
+#include "GeometryCollection/Facades/CollectionRenderingFacade.h"
 
 #include UE_INLINE_GENERATED_CPP_BY_NAME(GetGroomAssetNode)
 
 namespace UE::Groom::Private
 {
-	FORCEINLINE const UGroomAsset* FillEditableAsset(Dataflow::FContext& Context, const UGroomAsset* NodeGroom, FEditableGroom& EditGroom)
-	{
-		const UGroomAsset* GroomAsset = NodeGroom;
-		if (!GroomAsset)
-		{
-			if (const Dataflow::FEngineContext* EngineContext = Context.AsType<Dataflow::FEngineContext>())
-			{
-				GroomAsset = Cast<UGroomAsset>(EngineContext->Owner);
-			}
-		}
-
-		if (GroomAsset)
-		{
-			ConvertFromGroomAsset(const_cast<UGroomAsset*>(GroomAsset), &EditGroom, false, false, false);
-		}
-		return GroomAsset;
-	}
-	
-	template<typename FacadeType>
-	FORCEINLINE void BuildGroomCollection(FManagedArrayCollection& GroomCollection, FEditableGroom& EditGroom, const FString& GroomName)
+	FORCEINLINE void BuildCurveCollection(FManagedArrayCollection& GroomCollection, FEditableGroom& EditGroom, const FString& GroomName, const EGroomCollectionType CurvesType, const float CurveThickness)
 	{
 		TArray<FVector3f> PointRestPositions;
-		TArray<int32> ObjectCurveOffsets;
+		TArray<int32> GeometryCurveOffsets;
 		TArray<int32> CurvePointOffsets;
-		TArray<FString> ObjectGroupNames;
+		TArray<FString> GeometryGroupNames;
+		TArray<float> GeometryCurveThickness;
+		TArray<int32> CurveSourceIndices;
 
 		uint32 GroupIndex = 0;
 		for(FEditableGroomGroup& GroomGroup : EditGroom.Groups)
 		{
-			for( const typename FacadeType::FEditableType& EditType : FacadeType::GetEditableGroom(GroomGroup))
+			if(CurvesType == EGroomCollectionType::Guides)
+			{
+				for (const FEditableHairGuide& EditType : GroomGroup.Guides)
+				{
+					for(auto& GuidePoint : EditType.ControlPoints)
+					{
+						PointRestPositions.Add(GuidePoint.Position);
+					}
+					CurvePointOffsets.Add(PointRestPositions.Num());
+					CurveSourceIndices.Add(0);
+				}
+			}
+			else
 			{
-				for(auto& GuidePoint : EditType.ControlPoints)
+				for (const FEditableHairStrand& EditType : GroomGroup.Strands)
 				{
-					PointRestPositions.Add(GuidePoint.Position);
+					for(auto& GuidePoint : EditType.ControlPoints)
+					{
+						PointRestPositions.Add(GuidePoint.Position);
+					}
+					CurvePointOffsets.Add(PointRestPositions.Num());
+					CurveSourceIndices.Add(0);
 				}
-				CurvePointOffsets.Add(PointRestPositions.Num());
 			}
-			ObjectCurveOffsets.Add(CurvePointOffsets.Num());
+			GeometryCurveOffsets.Add(CurvePointOffsets.Num());
 			
-			const FString GroupName = GroomName + TEXT("_") + FacadeType::GroupPrefix.ToString();
-			ObjectGroupNames.Add(GroupName);
+			FString GroupName = (CurvesType == EGroomCollectionType::Guides) ? GroomName + TEXT("_Guides") : GroomName + TEXT("_Strands");
+			GroupName += TEXT("_") + GroomGroup.GroupName.ToString();
+			GeometryGroupNames.Add(GroupName);
+			GeometryCurveThickness.Add(CurveThickness);
 			++GroupIndex;
 		}
-		FacadeType CurvesFacade(GroomCollection);
-		CurvesFacade.InitGroomCollection(PointRestPositions, CurvePointOffsets, ObjectCurveOffsets, ObjectGroupNames);
+		GeometryCollection::Facades::FCollectionCurveGeometryFacade CurvesFacade(GroomCollection);
+		CurvesFacade.InitCurvesCollection(PointRestPositions, CurvePointOffsets, GeometryCurveOffsets, GeometryGroupNames, GeometryCurveThickness, CurveSourceIndices);
 	}
 }
 
 void FGetGroomAssetDataflowNode::Evaluate(UE::Dataflow::FContext& Context, const FDataflowOutput* Out) const
 {
-	if (Out->IsA<FManagedArrayCollection>(&Collection))
+}
+
+void FGetGroomAssetDataflowNode_v2::Evaluate(UE::Dataflow::FContext& Context, const FDataflowOutput* Out) const
+{
+	if (Out->IsA<TObjectPtr<const UGroomAsset>>(&GroomAsset))
 	{
-		FManagedArrayCollection GroomCollection;
-		FEditableGroom EditGroom;
-		if(const UGroomAsset* LocalGroom = UE::Groom::Private::FillEditableAsset(Context, GroomAsset, EditGroom))
+		TObjectPtr<const UGroomAsset> LocalAsset = GroomAsset;
+		if (!LocalAsset)
 		{
-			// Fill the guides facade
-			UE::Groom::Private::BuildGroomCollection<UE::Groom::FGroomGuidesFacade>(
-					GroomCollection, EditGroom, LocalGroom->GetName());
-
-			// Fill the strands facade
-			UE::Groom::Private::BuildGroomCollection<UE::Groom::FGroomStrandsFacade>(
-					GroomCollection, EditGroom, LocalGroom->GetName());
-
-			// Add guides points samples for future resampling
-			TArray<int32> ObjectPointSamples, CurveStrandIndices;
-			const TArray<FHairGroupsPhysics>& GroupsPhysics = LocalGroom->GetHairGroupsPhysics();
-			for(const FHairGroupsPhysics& GroupPhysics : GroupsPhysics)
+			if (const UE::Dataflow::FEngineContext* EngineContext = Context.AsType<UE::Dataflow::FEngineContext>())
 			{
-				ObjectPointSamples.Add(static_cast<uint8>(GroupPhysics.StrandsParameters.StrandsSize));
+				LocalAsset = Cast<UGroomAsset>(EngineContext->Owner);
 			}
-			
-			UE::Groom::FGroomGuidesFacade GuidesFacade(GroomCollection);
-			if(GroomCollection.NumElements(UE::Groom::FGroomGuidesFacade::ObjectsGroup) == 0)
-			{
-				GroomCollection.AddElements(ObjectPointSamples.Num(), UE::Groom::FGroomGuidesFacade::ObjectsGroup);
-			}
-			GuidesFacade.SetObjectPointSamples(ObjectPointSamples);
 		}
-
-		SetValue(Context, MoveTemp(GroomCollection), &Collection);
+		SetValue(Context, LocalAsset, &GroomAsset);
 	}
 }
 
-TArray<UE::Dataflow::FRenderingParameter> FGetGroomAssetDataflowNode::GetRenderParametersImpl() const
+FGroomAssetToCollectionDataflowNode::FGroomAssetToCollectionDataflowNode(const UE::Dataflow::FNodeParameters& InParam, FGuid InGuid)
+		: FDataflowNode(InParam, InGuid)
 {
-	if(CurvesType == EGroomCollectionType::Guides)
-	{
-		return { {TEXT("GuidesRender"), FName("FGroomCollection"), {TEXT("Collection")}}};
-	}
-	else
+	RegisterInputConnection(&GroomAsset);
+	RegisterInputConnection(&CurvesThickness).SetCanHidePin(true).SetPinIsHidden(true);
+	RegisterInputConnection(&CurvesType).SetCanHidePin(true).SetPinIsHidden(true);
+	RegisterOutputConnection(&Collection);
+}
+
+void FGroomAssetToCollectionDataflowNode::Evaluate(UE::Dataflow::FContext& Context, const FDataflowOutput* Out) const
+{
+	if (Out->IsA<FManagedArrayCollection>(&Collection))
 	{
-		return { {TEXT("StrandsRender"), FName("FGroomCollection"), {TEXT("Collection")}}};
+		FManagedArrayCollection GroomCollection;
+		FEditableGroom EditGroom;
+		if(TObjectPtr<const UGroomAsset> InputAsset = GetValue<TObjectPtr<const UGroomAsset>>(Context, &GroomAsset))
+		{
+			const float GeometryThickness = GetValue<float>(Context, &CurvesThickness);
+			ConvertFromGroomAsset(const_cast<UGroomAsset*>(InputAsset.Get()), &EditGroom, false, false, false);
+
+			UE::Groom::Private::BuildCurveCollection(GroomCollection, EditGroom, InputAsset->GetName(), CurvesType, GeometryThickness);
+		}
+
+		SetValue(Context, MoveTemp(GroomCollection), &Collection);
 	}
 }
+
