Geometry facade for grooms and cards + new rendering + use of curve selection

diff --git a/Engine/Plugins/Runtime/HairStrands/Source/HairStrandsDataflow/Private/GetGroomAttributesNodes.cpp b/Engine/Plugins/Runtime/HairStrands/Source/HairStrandsDataflow/Private/GetGroomAttributesNodes.cpp
--- a/Engine/Plugins/Runtime/HairStrands/Source/HairStrandsDataflow/Private/GetGroomAttributesNodes.cpp
+++ b/Engine/Plugins/Runtime/HairStrands/Source/HairStrandsDataflow/Private/GetGroomAttributesNodes.cpp
@@ -1,46 +1,57 @@
 // Copyright Epic Games, Inc. All Rights Reserved.
 
 #include "GetGroomAttributesNodes.h"
-#include "GroomCollectionFacades.h"
 #include "Dataflow/DataflowObjectInterface.h"
+#include "GeometryCollection/GeometryCollection.h"
+#include "GeometryCollection/Facades/CollectionCurveFacade.h"
+#include "GeometryCollection/Facades/CollectionVertexBoneWeightsFacade.h"
 
 #include UE_INLINE_GENERATED_CPP_BY_NAME(GetGroomAttributesNodes)
 
 void FGetGroomAttributesDataflowNode::Evaluate(UE::Dataflow::FContext& Context, const FDataflowOutput* Out) const
+{
+}
+
+FGetCurveAttributesDataflowNode::FGetCurveAttributesDataflowNode(const UE::Dataflow::FNodeParameters& InParam, FGuid InGuid)
+	: FDataflowNode(InParam, InGuid)
+{
+	RegisterOutputConnection(&AttributeKey);
+}
+
+void FGetCurveAttributesDataflowNode::Evaluate(UE::Dataflow::FContext& Context, const FDataflowOutput* Out) const
 {
 	if (Out->IsA<FCollectionAttributeKey>(&AttributeKey))
 	{
 		FCollectionAttributeKey Key;
 
-		const FString VertexGroup = (CurvesType == EGroomCollectionType::Guides) ? UE::Groom::FGroomGuidesFacade::VerticesGroup.ToString() :
-								UE::Groom::FGroomStrandsFacade::VerticesGroup.ToString();
-		const FString CurveGroup = (CurvesType == EGroomCollectionType::Guides) ? UE::Groom::FGroomGuidesFacade::CurvesGroup.ToString() :
-								UE::Groom::FGroomStrandsFacade::CurvesGroup.ToString();
+		const FString VertexGroup = FGeometryCollection::VerticesGroup.ToString();
+		const FString CurveGroup = GeometryCollection::Facades::FCollectionCurveGeometryFacade::CurvesGroup.ToString();
 		if(AttributeType == EGroomAttributeType::KinematicWeights)
 		{ 
 			Key.Group = VertexGroup;
-			Key.Attribute = UE::Groom::FGroomGuidesFacade::PointKinematicWeightsAttribute.ToString();
+			Key.Attribute = GeometryCollection::Facades::FVertexBoneWeightsFacade::KinematicWeightAttributeName.ToString();
 		}
 		else if (AttributeType == EGroomAttributeType::BoneIndices)
 		{
 			Key.Group = VertexGroup;
-			Key.Attribute = UE::Groom::FGroomGuidesFacade::PointBoneIndicesAttribute.ToString();
+			Key.Attribute = GeometryCollection::Facades::FVertexBoneWeightsFacade::BoneIndicesAttributeName.ToString();
 		}
 		else if (AttributeType == EGroomAttributeType::BoneWeights)
 		{
 			Key.Group = VertexGroup;
-			Key.Attribute = UE::Groom::FGroomGuidesFacade::PointBoneWeightsAttribute.ToString();
+			Key.Attribute = GeometryCollection::Facades::FVertexBoneWeightsFacade::BoneWeightsAttributeName.ToString();
 		}
 		else if (AttributeType == EGroomAttributeType::CurveLods)
 		{
 			Key.Group = CurveGroup;
-			Key.Attribute = UE::Groom::FGroomGuidesFacade::CurveLodIndicesAttribute.ToString();
+			Key.Attribute = GeometryCollection::Facades::FCollectionCurveHierarchyFacade::CurveLodIndicesAttribute.ToString();
 		}
 		else if (AttributeType == EGroomAttributeType::CurveParents)
 		{
 			Key.Group = CurveGroup;
-			Key.Attribute = UE::Groom::FGroomGuidesFacade::CurveParentIndicesAttribute.ToString();
+			Key.Attribute = GeometryCollection::Facades::FCollectionCurveHierarchyFacade::CurveParentIndicesAttribute.ToString();
 		}
+		
 		SetValue(Context, MoveTemp(Key), &AttributeKey);
 	}
 }
diff --git a/Engine/Plugins/Runtime/HairStrands/Source/HairStrandsDataflow/Private/GroomAssetTerminalNode.cpp b/Engine/Plugins/Runtime/HairStrands/Source/HairStrandsDataflow/Private/GroomAssetTerminalNode.cpp
--- a/Engine/Plugins/Runtime/HairStrands/Source/HairStrandsDataflow/Private/GroomAssetTerminalNode.cpp
+++ b/Engine/Plugins/Runtime/HairStrands/Source/HairStrandsDataflow/Private/GroomAssetTerminalNode.cpp
@@ -3,10 +3,11 @@
 #include "GroomAssetTerminalNode.h"
 
 #include "AssetCompilingManager.h"
-#include "GroomCollectionFacades.h"
 #include "GroomEdit.h"
 #include "Dataflow/DataflowConnectionTypes.h"
 #include "Dataflow/DataflowObjectInterface.h"
+#include "GeometryCollection/Facades/CollectionCurveFacade.h"
+#include "GeometryCollection/Facades/CollectionVertexBoneWeightsFacade.h"
 
 #include UE_INLINE_GENERATED_CPP_BY_NAME(GroomAssetTerminalNode)
 
@@ -14,25 +15,26 @@ namespace UE::Groom::Private
 {
 	static const FName AttributeKeyTypeName(TEXT("FCollectionAttributeKey"));
 
-	FORCEINLINE void BuildEditableGuides(const UE::Groom::FGroomGuidesFacade& GuidesFacade, FEditableGroom& EditGroom)
+	FORCEINLINE void BuildEditableGuides(const TArray<int32>& ObjectCurveOffsets,
+		const TArray<int32>& CurvePointOffsets, const TArray<FVector3f>& PointRestPositions, const TArray<int32>& CurveStrandIndices, FEditableGroom& EditGroom)
 	{
 		int32 ObjectIndex = 0, CurveIndex = 0;
 		int32 PrevCurve = 0, PrevPoint = 0;;
 		for (FEditableGroomGroup& Group : EditGroom.Groups)
 		{
-			const int32 NextCurve = GuidesFacade.GetObjectCurveOffsets()[ObjectIndex];
+			const int32 NextCurve = ObjectCurveOffsets[ObjectIndex];
 			Group.Guides.SetNum(NextCurve-PrevCurve);
 			
 			// Only rebuild the guides since the strands are the same
 			for (FEditableHairGuide& Guide : Group.Guides)
 			{
-				const int32 NextPoint = GuidesFacade.GetCurvePointOffsets()[CurveIndex];
-				const int32 StrandIndex = GuidesFacade.GetCurveStrandIndices()[CurveIndex];
+				const int32 NextPoint = CurvePointOffsets[CurveIndex];
+				const int32 StrandIndex = CurveStrandIndices[CurveIndex];
 						
 				Guide.ControlPoints.Reset();
 				for(int32 PointIndex = PrevPoint; PointIndex < NextPoint; ++PointIndex)
 				{
-					Guide.ControlPoints.Add({GuidesFacade.GetPointRestPositions()[PointIndex],
+					Guide.ControlPoints.Add({PointRestPositions[PointIndex],
 						FMath::Clamp(static_cast<float>(PointIndex-PrevPoint) / (NextPoint-PrevPoint-1), 0.f, 1.f)});
 				}
 						
@@ -50,157 +52,206 @@ namespace UE::Groom::Private
 		}
 	}
 
+	FORCEINLINE void CopyCollectionAttribute(const FManagedArrayCollection* InputCollection, FManagedArrayCollection* OutputCollection,
+		const FCollectionAttributeKey& AttributeToCopy, const FString& GroupPrefix)
+	{
+		const FName AttributeName(AttributeToCopy.Attribute);
+		const FName SourceGroup(AttributeToCopy.Group);
+		const FName TargetGroup( GroupPrefix + AttributeToCopy.Group);
+
+		if (InputCollection->HasGroup(SourceGroup))
+		{
+			if(!OutputCollection->HasGroup(TargetGroup))
+			{
+				OutputCollection->AddGroup(TargetGroup);
+			}
+			if (InputCollection->NumElements(SourceGroup) != OutputCollection->NumElements(TargetGroup))
+			{
+				OutputCollection->EmptyGroup(TargetGroup);
+				OutputCollection->AddElements(InputCollection->NumElements(SourceGroup), TargetGroup);
+			}
+			OutputCollection->CopyAttribute(*InputCollection, AttributeName, AttributeName, SourceGroup, TargetGroup);
+		}
+	}
+
 	FORCEINLINE void CopyCollectionAttributes(const FManagedArrayCollection* InputCollection, FManagedArrayCollection* OutputCollection,
 		const TArray<FCollectionAttributeKey>& AttributesToCopy = TArray<FCollectionAttributeKey>())
 	{
 		for(const FCollectionAttributeKey& AttributeToCopy : AttributesToCopy)
 		{
-			const FName AttributeName(AttributeToCopy.Attribute);
-			const FName GroupName(AttributeToCopy.Group);
+			CopyCollectionAttribute(InputCollection, OutputCollection, AttributeToCopy, "");
+		}
+	}
 
-			if (InputCollection->HasGroup(GroupName))
+	template<typename AttributeType>
+	static void BuildVerticesAttribute(const FManagedArrayCollection& InCollection, FManagedArrayCollection* OutCollection, const int32 NumPoints, const FName& AttributeName, const FName& VerticesGroup, const FName& PointsGroup)
+	{
+		const TManagedArray<AttributeType>& VerticesAttribute = InCollection.GetAttribute<AttributeType>(AttributeName, VerticesGroup);
+		if(OutCollection->NumElements(PointsGroup) != NumPoints)
+		{
+			if(OutCollection->NumElements(PointsGroup) > 0)
 			{
-				if(!OutputCollection->HasGroup(GroupName))
-				{
-					OutputCollection->AddGroup(GroupName);
-				}
-				if (InputCollection->NumElements(GroupName) != OutputCollection->NumElements(GroupName))
-				{
-					OutputCollection->EmptyGroup(GroupName);
-					OutputCollection->AddElements(InputCollection->NumElements(GroupName), GroupName);
-				}
-				OutputCollection->CopyAttribute(*InputCollection, AttributeName, GroupName);
+				OutCollection->EmptyGroup(PointsGroup);
 			}
+			OutCollection->AddElements(NumPoints, PointsGroup);
+		}
+		TManagedArray<AttributeType>& PointsAttribute = OutCollection->AddAttribute<AttributeType>(AttributeName, PointsGroup);
+		
+		for(int32 PointIndex = 0; PointIndex < NumPoints; ++PointIndex)
+		{
+			PointsAttribute[PointIndex] = VerticesAttribute[2*PointIndex];
 		}
 	}
 
-	template<typename FacadeType, typename AttributeType>
-	static void BuildVerticesAttribute(const FManagedArrayCollection& InCollection, FManagedArrayCollection* OutCollection, const int32 NumPoints, const FName& AttributeName)
+	template<typename AttributeType, typename CompressedType, int32 NumElements>
+	static void BuildVerticesArray(const FManagedArrayCollection& InCollection, FManagedArrayCollection* OutCollection, const int32 NumPoints,
+		const FName& AttributeName, const FName& VerticesGroup, const FName& PointsGroup, const AttributeType DefaultValue)
 	{
-		const TManagedArray<AttributeType>& VerticesAttribute = InCollection.GetAttribute<AttributeType>(AttributeName, FacadeType::VerticesGroup);
-		if(OutCollection->NumElements(FacadeType::PointsGroup) != NumPoints)
+		const TManagedArray<TArray<AttributeType>>& VerticesArray = InCollection.GetAttribute<TArray<AttributeType>>(AttributeName, VerticesGroup);
+		if(OutCollection->NumElements(PointsGroup) != NumPoints)
 		{
-			if(OutCollection->NumElements(FacadeType::PointsGroup) > 0)
+			if(OutCollection->NumElements(PointsGroup) > 0)
 			{
-				OutCollection->EmptyGroup(FacadeType::PointsGroup);
+				OutCollection->EmptyGroup(PointsGroup);
 			}
-			OutCollection->AddElements(NumPoints, FacadeType::PointsGroup);
+			OutCollection->AddElements(NumPoints, PointsGroup);
 		}
-		TManagedArray<AttributeType>& PointsAttribute = OutCollection->AddAttribute<AttributeType>(AttributeName, FacadeType::PointsGroup);
+		TManagedArray<CompressedType>& PointsAttribute = OutCollection->AddAttribute<CompressedType>(AttributeName, PointsGroup);
 		
 		for(int32 PointIndex = 0; PointIndex < NumPoints; ++PointIndex)
 		{
-			PointsAttribute[PointIndex] = VerticesAttribute[2*PointIndex];
+			for(int32 ElemIndex = 0; ElemIndex < NumElements; ++ElemIndex)
+			{
+				PointsAttribute[PointIndex][ElemIndex] = VerticesArray[2*PointIndex].IsValidIndex(ElemIndex) ? VerticesArray[2*PointIndex][ElemIndex] : DefaultValue;
+			}
 		}
 	}
 
-	template<typename FacadeType>
+	static void TransferVerticesAttribute(const FManagedArrayCollection& InCollection, FManagedArrayCollection* OutCollection, const FCollectionAttributeKey& AttributeToCopy,
+		const int32 NumPoints, const FName& PointsGroup)
+	{
+		const FName AttributeName(AttributeToCopy.Attribute);
+		const FName VerticesGroup(AttributeToCopy.Group);
+		
+		if(InCollection.GetAttributeType(AttributeName, VerticesGroup) == EManagedArrayType::FFloatType)
+		{
+			BuildVerticesAttribute<float>(InCollection, OutCollection, NumPoints, AttributeName, VerticesGroup, PointsGroup);
+		}
+		else if(InCollection.GetAttributeType(AttributeName, VerticesGroup) == EManagedArrayType::FVector4fType)
+		{
+			BuildVerticesAttribute<FVector4f>(InCollection, OutCollection, NumPoints, AttributeName, VerticesGroup, PointsGroup);
+		}
+		else if(InCollection.GetAttributeType(AttributeName, VerticesGroup) == EManagedArrayType::FVectorType)
+		{
+			BuildVerticesAttribute<FVector3f>(InCollection, OutCollection, NumPoints, AttributeName, VerticesGroup, PointsGroup);
+		}
+		else if(InCollection.GetAttributeType(AttributeName, VerticesGroup) == EManagedArrayType::FVector2DType)
+		{
+			BuildVerticesAttribute<FVector2f>(InCollection, OutCollection, NumPoints, AttributeName, VerticesGroup, PointsGroup);
+		}
+		if(InCollection.GetAttributeType(AttributeName, VerticesGroup) == EManagedArrayType::FInt32Type)
+		{
+			BuildVerticesAttribute<int32>(InCollection, OutCollection, NumPoints, AttributeName, VerticesGroup, PointsGroup);
+		}
+		else if(InCollection.GetAttributeType(AttributeName, VerticesGroup) == EManagedArrayType::FIntVector4Type)
+		{
+			BuildVerticesAttribute<FIntVector4>(InCollection, OutCollection, NumPoints, AttributeName, VerticesGroup, PointsGroup);
+		}
+		else if(InCollection.GetAttributeType(AttributeName, VerticesGroup) == EManagedArrayType::FIntVectorType)
+		{
+			BuildVerticesAttribute<FIntVector3>(InCollection, OutCollection, NumPoints, AttributeName, VerticesGroup, PointsGroup);
+		}
+		else if(InCollection.GetAttributeType(AttributeName, VerticesGroup) == EManagedArrayType::FIntVector2Type)
+		{
+			BuildVerticesAttribute<FIntVector2>(InCollection, OutCollection, NumPoints, AttributeName, VerticesGroup, PointsGroup);
+		}
+		else if(InCollection.GetAttributeType(AttributeName, VerticesGroup) == EManagedArrayType::FBoolType)
+		{
+			BuildVerticesAttribute<bool>(InCollection, OutCollection, NumPoints, AttributeName, VerticesGroup, PointsGroup);
+		}
+		else if(InCollection.GetAttributeType(AttributeName, VerticesGroup) == EManagedArrayType::FLinearColorType)
+		{
+			BuildVerticesAttribute<FLinearColor>(InCollection, OutCollection, NumPoints, AttributeName, VerticesGroup, PointsGroup);
+		}
+		else if(InCollection.GetAttributeType(AttributeName, VerticesGroup) == EManagedArrayType::FQuatType)
+		{
+			BuildVerticesAttribute<FQuat4f>(InCollection, OutCollection, NumPoints, AttributeName, VerticesGroup, PointsGroup);
+		}
+		else if(InCollection.GetAttributeType(AttributeName, VerticesGroup) == EManagedArrayType::FTransform3fType)
+		{
+			BuildVerticesAttribute<FTransform3f>(InCollection, OutCollection, NumPoints, AttributeName, VerticesGroup, PointsGroup); 
+		}
+		else if(InCollection.GetAttributeType(AttributeName, VerticesGroup) == EManagedArrayType::FInt32ArrayType)
+		{
+			BuildVerticesArray<int32,FIntVector4,4>(InCollection, OutCollection, NumPoints, AttributeName, VerticesGroup, PointsGroup, INDEX_NONE);
+		}
+		else if(InCollection.GetAttributeType(AttributeName, VerticesGroup) == EManagedArrayType::FFloatArrayType)
+		{
+			BuildVerticesArray<float,FVector4f,4>(InCollection, OutCollection, NumPoints, AttributeName, VerticesGroup, PointsGroup, 0.0f);
+		}
+	}
+	
 	static void TransferVerticesAttributes(const FManagedArrayCollection& InCollection, FManagedArrayCollection* OutCollection,
-		const int32 NumPoints, const TArray<FName>& AttributesToSkip)
+		const int32 NumPoints, const TArray<FName>& AttributesToSkip, const FName& VerticesGroup, const FName& PointsGroup)
 	{
 		// Transfer vertices weight maps onto the points to be stored onto the rest collection
-		const TArray<FName> AttributeNames = InCollection.AttributeNames(FacadeType::VerticesGroup);
+		const TArray<FName> AttributeNames = InCollection.AttributeNames(VerticesGroup);
 		for(const FName& AttributeName : AttributeNames)
 		{
 			if(!AttributesToSkip.Contains(AttributeName))
 			{
-				if(InCollection.GetAttributeType(AttributeName, FacadeType::VerticesGroup) == EManagedArrayType::FFloatType)
-				{
-					BuildVerticesAttribute<FacadeType, float>(InCollection, OutCollection, NumPoints, AttributeName);
-				}
-				else if(InCollection.GetAttributeType(AttributeName, FacadeType::VerticesGroup) == EManagedArrayType::FVector4fType)
-				{
-					BuildVerticesAttribute<FacadeType, FVector4f>(InCollection, OutCollection, NumPoints, AttributeName);
-				}
-				else if(InCollection.GetAttributeType(AttributeName, FacadeType::VerticesGroup) == EManagedArrayType::FVectorType)
-				{
-					BuildVerticesAttribute<FacadeType, FVector3f>(InCollection, OutCollection, NumPoints, AttributeName);
-				}
-				else if(InCollection.GetAttributeType(AttributeName, FacadeType::VerticesGroup) == EManagedArrayType::FVector2DType)
-				{
-					BuildVerticesAttribute<FacadeType, FVector2f>(InCollection, OutCollection, NumPoints, AttributeName);
-				}
-				if(InCollection.GetAttributeType(AttributeName, FacadeType::VerticesGroup) == EManagedArrayType::FInt32Type)
-				{
-					BuildVerticesAttribute<FacadeType, int32>(InCollection, OutCollection, NumPoints, AttributeName);
-				}
-				else if(InCollection.GetAttributeType(AttributeName, FacadeType::VerticesGroup) == EManagedArrayType::FIntVector4Type)
-				{
-					BuildVerticesAttribute<FacadeType, FIntVector4>(InCollection, OutCollection, NumPoints, AttributeName);
-				}
-				else if(InCollection.GetAttributeType(AttributeName, FacadeType::VerticesGroup) == EManagedArrayType::FIntVectorType)
-				{
-					BuildVerticesAttribute<FacadeType, FIntVector3>(InCollection, OutCollection, NumPoints, AttributeName);
-				}
-				else if(InCollection.GetAttributeType(AttributeName, FacadeType::VerticesGroup) == EManagedArrayType::FIntVector2Type)
-				{
-					BuildVerticesAttribute<FacadeType, FIntVector2>(InCollection, OutCollection, NumPoints, AttributeName);
-				}
-				else if(InCollection.GetAttributeType(AttributeName, FacadeType::VerticesGroup) == EManagedArrayType::FBoolType)
-				{
-					BuildVerticesAttribute<FacadeType, bool>(InCollection, OutCollection, NumPoints, AttributeName);
-				}
-				else if(InCollection.GetAttributeType(AttributeName, FacadeType::VerticesGroup) == EManagedArrayType::FLinearColorType)
-				{
-					BuildVerticesAttribute<FacadeType, FLinearColor>(InCollection, OutCollection, NumPoints, AttributeName);
-				}
-				else if(InCollection.GetAttributeType(AttributeName, FacadeType::VerticesGroup) == EManagedArrayType::FQuatType)
-				{
-					BuildVerticesAttribute<FacadeType, FQuat4f>(InCollection, OutCollection, NumPoints, AttributeName);
-				}
-				else if(InCollection.GetAttributeType(AttributeName, FacadeType::VerticesGroup) == EManagedArrayType::FTransform3fType)
-				{
-					BuildVerticesAttribute<FacadeType, FTransform3f>(InCollection, OutCollection, NumPoints, AttributeName); 
-				}
+				FCollectionAttributeKey AttributeKey(AttributeName.ToString(),VerticesGroup.ToString());
+				TransferVerticesAttribute(InCollection, OutCollection, AttributeKey, NumPoints, PointsGroup);
 			}
 		}
 	}
 
-	template<typename FacadeType>
-	FORCEINLINE void RegisterSkeletalMeshes(const FManagedArrayCollection& InCollection, const FacadeType& GroomFacade, UGroomAsset* GroomAsset)
+	FORCEINLINE void RegisterSkeletalMeshes(const FManagedArrayCollection& InCollection, const FName& SkelMeshAttribute, const FName& MeshLODAttribute, const FName& GeometryGroup, const int32 NumGeometry, UGroomAsset* GroomAsset)
 	{
 		const TManagedArray<TObjectPtr<UObject>>& ObjectSkeletalMeshes =
-						InCollection.GetAttribute<TObjectPtr<UObject>>(UE::Groom::FGroomGuidesFacade::ObjectSkeletalMeshesAttribute, FacadeType::ObjectsGroup);
+						InCollection.GetAttribute<TObjectPtr<UObject>>(SkelMeshAttribute, GeometryGroup);
 
 		const TManagedArray<int32>& ObjectMeshLODs =
-			InCollection.GetAttribute<int32>(UE::Groom::FGroomGuidesFacade::ObjectMeshLODsAttribute, FacadeType::ObjectsGroup);
+			InCollection.GetAttribute<int32>(MeshLODAttribute, GeometryGroup);
 				
-		for(int32 GroupIndex = 0; GroupIndex < GroomFacade.GetNumObjects(); ++GroupIndex)
+		for(int32 GroupIndex = 0; GroupIndex < NumGeometry; ++GroupIndex)
 		{
 			GroomAsset->GetDataflowSettings().SetSkeletalMesh(GroupIndex,
 				Cast<USkeletalMesh>(ObjectSkeletalMeshes[GroupIndex]), ObjectMeshLODs[GroupIndex]);
 		}
 	}
-	template<typename FacadeType>
-	FORCEINLINE void TransferCollectionAttributes(const FManagedArrayCollection& InCollection, FManagedArrayCollection* OutCollection, const FacadeType& GroomFacade,
-		const TArray<FCollectionAttributeKey>& ExternalAttributes, const TArray<FCollectionAttributeKey>& InternalAttributes)
+	
+	static void TransferCurvesAttributes(const FManagedArrayCollection& InCollection, FManagedArrayCollection* OutCollection, 
+		const TArray<FCollectionAttributeKey>& ExternalAttributes, const TArray<FCollectionAttributeKey>& InternalAttributes, const FString& GroupPrefix)
 	{
-		// Register all the internal attributes to be copied
-		TArray<FCollectionAttributeKey> AttributesToCopy;
-
-		auto AddAttributeKeys = [&AttributesToCopy](const TArray<FCollectionAttributeKey>& AttributeKeys)
-		{
-			for (const FCollectionAttributeKey& AttributeKey : AttributeKeys)
+		const FName PointsGroup(GroupPrefix + GeometryCollection::Facades::FCollectionCurveGeometryFacade::PointsGroup.ToString());
+		const int32 NumPoints = InCollection.NumElements(FGeometryCollection::VerticesGroup) / 2;
+		
+		auto AddAttributeKeys = [&InCollection, &OutCollection, &GroupPrefix, &PointsGroup, &NumPoints](const TArray<FCollectionAttributeKey>& AttributeKeys)
 			{
-				if (AttributeKey.Group == FacadeType::CurvesGroup || AttributeKey.Group == FacadeType::ObjectsGroup ||
-					AttributeKey.Group == FacadeType::PointsGroup || AttributeKey.Group == FacadeType::EdgesGroup)
+				for (const FCollectionAttributeKey& AttributeKey : AttributeKeys)
 				{
-					AttributesToCopy.Add(AttributeKey);
+					if(InCollection.HasAttribute(FName(AttributeKey.Attribute), FName(AttributeKey.Group)))
+					{
+						if (AttributeKey.Group == GeometryCollection::Facades::FCollectionCurveGeometryFacade::CurvesGroup ||
+							AttributeKey.Group == GeometryCollection::Facades::FCollectionCurveGeometryFacade::PointsGroup)
+						{
+							CopyCollectionAttribute(&InCollection, OutCollection, AttributeKey, GroupPrefix);
+						}
+						else if (AttributeKey.Group == FGeometryCollection::GeometryGroup)
+						{
+							CopyCollectionAttribute(&InCollection, OutCollection, AttributeKey, GroupPrefix);
+						}
+						else if(AttributeKey.Group == FGeometryCollection::VerticesGroup)
+						{
+							TransferVerticesAttribute(InCollection, OutCollection, AttributeKey, NumPoints, PointsGroup);
+						}
+					}
 				}
-			}
-		};
+			};
 
 		AddAttributeKeys(ExternalAttributes);
 		AddAttributeKeys(InternalAttributes);
-
-		// Copy attributes from input collection
-		UE::Groom::Private::CopyCollectionAttributes(&InCollection, OutCollection, AttributesToCopy);
-				
-		// Skip default vertex attributes as not defined by the user
-		const TArray<FName> AttributesToSkip = {FacadeType::VertexLinearColorsAttribute};
-				
-		// Transfer vertices weight maps onto the points to be stored onto the rest collection
-		UE::Groom::Private::TransferVerticesAttributes<FacadeType>(InCollection, OutCollection, GroomFacade.GetNumPoints(), AttributesToSkip);
 	}
 }
 
@@ -213,78 +264,115 @@ FGroomAssetTerminalDataflowNode::FGroomAssetTerminalDataflowNode(const UE::Dataf
 	RegisterOutputConnection(&Collection, &Collection);
 }
 
-
 void FGroomAssetTerminalDataflowNode::SetAssetValue(TObjectPtr<UObject> Asset, UE::Dataflow::FContext& Context) const
+{
+}
+
+void FGroomAssetTerminalDataflowNode::Evaluate(UE::Dataflow::FContext& Context, const FDataflowOutput* Out) const
+{
+}
+
+FGroomAssetTerminalDataflowNode_v2::FGroomAssetTerminalDataflowNode_v2(const UE::Dataflow::FNodeParameters& InParam, FGuid InGuid)
+	: FDataflowTerminalNode(InParam, InGuid)
+{
+	DataflowAssetWeakPtr = Cast<UDataflow>(InParam.OwningObject);
+
+	RegisterInputConnection(&GuidesCollection);
+	RegisterInputConnection(&StrandsCollection);
+}
+
+void FGroomAssetTerminalDataflowNode_v2::SetAssetValue(TObjectPtr<UObject> Asset, UE::Dataflow::FContext& Context) const
 {
 	if (UGroomAsset* GroomAsset = Cast<UGroomAsset>(Asset.Get()))
 	{
-		const FManagedArrayCollection& InCollection = GetValue<FManagedArrayCollection>(Context, &Collection);
-		UE::Groom::FGroomGuidesFacade GuidesFacade(InCollection);
-		UE::Groom::FGroomStrandsFacade StrandsFacade(InCollection);
+		const FManagedArrayCollection& GroomStrands = GetValue<FManagedArrayCollection>(Context, &StrandsCollection);
+		const FManagedArrayCollection& GroomGuides = GetValue<FManagedArrayCollection>(Context, &GuidesCollection);
+		
+		GeometryCollection::Facades::FCollectionCurveGeometryFacade GuidesFacade(GroomGuides);
+		GeometryCollection::Facades::FCollectionCurveGeometryFacade StrandsFacade(GroomStrands);
 
-		FManagedArrayCollection* OutCollection = (GuidesFacade.IsValid() || StrandsFacade.IsValid()) ? new FManagedArrayCollection() : nullptr;
+		if(GuidesFacade.IsValid() || StrandsFacade.IsValid())
+		{ 
+			FManagedArrayCollection* OutCollection = new FManagedArrayCollection();
 
-		if(StrandsFacade.IsValid())
-		{
-			UE::Groom::Private::TransferCollectionAttributes<UE::Groom::FGroomStrandsFacade>(InCollection, OutCollection, StrandsFacade, AttributeKeys, {});
-		}
-		if(GuidesFacade.IsValid())
-		{
-			UE::Groom::Private::TransferCollectionAttributes<UE::Groom::FGroomGuidesFacade>(InCollection, OutCollection, GuidesFacade, AttributeKeys,
-				{{UE::Groom::FGroomGuidesFacade::CurveParentIndicesAttribute.ToString(), UE::Groom::FGroomGuidesFacade::CurvesGroup.ToString()},
-				{UE::Groom::FGroomGuidesFacade::CurveLodIndicesAttribute.ToString(), UE::Groom::FGroomGuidesFacade::CurvesGroup.ToString()}});
+			TArray<FCollectionAttributeKey> SkinningAttributes = {
+				{GeometryCollection::Facades::FVertexBoneWeightsFacade::KinematicWeightAttributeName.ToString(), FGeometryCollection::VerticesGroup.ToString()},
+				{GeometryCollection::Facades::FVertexBoneWeightsFacade::BoneIndicesAttributeName.ToString(), FGeometryCollection::VerticesGroup.ToString()},
+				{GeometryCollection::Facades::FVertexBoneWeightsFacade::BoneWeightsAttributeName.ToString(), FGeometryCollection::VerticesGroup.ToString()}};
 
-			// Build an editable groom asset for the strands
-			FEditableGroom EditGroom;
-			ConvertFromGroomAsset(const_cast<UGroomAsset*>(GroomAsset), &EditGroom, false, false, false);
+			TArray<FCollectionAttributeKey> HierarchyAttributes = {
+				{GeometryCollection::Facades::FCollectionCurveHierarchyFacade::CurveParentIndicesAttribute.ToString(), GeometryCollection::Facades::FCollectionCurveGeometryFacade::CurvesGroup.ToString()},
+				{GeometryCollection::Facades::FCollectionCurveHierarchyFacade::CurveLodIndicesAttribute.ToString(), GeometryCollection::Facades::FCollectionCurveGeometryFacade::CurvesGroup.ToString()}};
 
-			if(GuidesFacade.GetNumObjects() == EditGroom.Groups.Num())
+			TArray<FCollectionAttributeKey> ExternalAttributes;
+			for(int32 AttributeIndex = 0; AttributeIndex < AttributeKeys.Num(); ++AttributeIndex)
 			{
-				// Build the editable guides
-				UE::Groom::Private::BuildEditableGuides(GuidesFacade,EditGroom);
-
-				// Ensure compilation dependent assets is done
-				FAssetCompilingManager::Get().FinishCompilationForObjects({ GroomAsset });
-
-				// Convert to groom asset
-				ConvertToGroomAsset(const_cast<UGroomAsset*>(GroomAsset), &EditGroom, EEditableGroomOperations::ControlPoints_Modified);
+				ExternalAttributes.Add(GetValue<FCollectionAttributeKey>(Context, GetConnectionReference(AttributeIndex)));
 			}
-			// To prevent future reconstruction in the BuildData we set the type to be imported
-			for(FHairGroupsInterpolation& GroupInterpolation : GroomAsset->GetHairGroupsInterpolation())
+		
+			TArray<FCollectionAttributeKey> InternalAttributes;
+			InternalAttributes.Append(SkinningAttributes);
+			InternalAttributes.Append(HierarchyAttributes);
+		
+			if(StrandsFacade.IsValid())
 			{
-				GroupInterpolation.InterpolationSettings.GuideType = EGroomGuideType::Imported;
+				UE::Groom::Private::TransferCurvesAttributes(GroomStrands, OutCollection, ExternalAttributes, InternalAttributes, FString("Strands"));
 			}
-		}  
+			if(GuidesFacade.IsValid())
+			{
+				UE::Groom::Private::TransferCurvesAttributes(GroomGuides, OutCollection, ExternalAttributes, InternalAttributes, FString("Guides"));
+
+				// Build an editable groom asset for the guides
+				FEditableGroom EditGroom;
+				ConvertFromGroomAsset(const_cast<UGroomAsset*>(GroomAsset), &EditGroom, false, false, false);
+
+				if(GuidesFacade.GetNumGeometry() == EditGroom.Groups.Num())
+				{
+					// Build the editable guides
+					UE::Groom::Private::BuildEditableGuides(GuidesFacade.GetGeometryCurveOffsets(), GuidesFacade.GetCurvePointOffsets(),
+						GuidesFacade.GetPointRestPositions(), GuidesFacade.GetCurveSourceIndices(),EditGroom);
+
+					// Ensure compilation dependent assets is done
+					FAssetCompilingManager::Get().FinishCompilationForObjects({ GroomAsset });
+
+					// Convert to groom asset
+					ConvertToGroomAsset(const_cast<UGroomAsset*>(GroomAsset), &EditGroom, EEditableGroomOperations::ControlPoints_Modified);
+				}
+				// To prevent future reconstruction in the BuildData we set the type to be imported
+				for(FHairGroupsInterpolation& GroupInterpolation : GroomAsset->GetHairGroupsInterpolation())
+				{
+					GroupInterpolation.InterpolationSettings.GuideType = EGroomGuideType::Imported;
+				}
+			}  
 
-		if(OutCollection)
-		{
 			GroomAsset->GetDataflowSettings().SetRestCollection(OutCollection);
-			GroomAsset->GetDataflowSettings().InitSkeletalMeshes(GuidesFacade.GetNumObjects());
+			GroomAsset->GetDataflowSettings().InitSkeletalMeshes(GuidesFacade.GetNumGeometry());
 
-			if(InCollection.HasAttribute(UE::Groom::FGroomGuidesFacade::ObjectSkeletalMeshesAttribute, UE::Groom::FGroomGuidesFacade::ObjectsGroup))
+			if(GroomGuides.HasAttribute(GeometryCollection::Facades::FVertexBoneWeightsFacade::SkeletalMeshAttributeName, FGeometryCollection::GeometryGroup))
 			{
-				UE::Groom::Private::RegisterSkeletalMeshes(InCollection, GuidesFacade, GroomAsset);
+				UE::Groom::Private::RegisterSkeletalMeshes(GroomGuides, GeometryCollection::Facades::FVertexBoneWeightsFacade::SkeletalMeshAttributeName,
+					GeometryCollection::Facades::FVertexBoneWeightsFacade::GeometryLODAttributeName,
+					FName(FGeometryCollection::GeometryGroup), GuidesFacade.GetNumGeometry(), GroomAsset);
 			}
-			else if(InCollection.HasAttribute(UE::Groom::FGroomGuidesFacade::ObjectSkeletalMeshesAttribute, UE::Groom::FGroomStrandsFacade::ObjectsGroup))
+			else if(GroomStrands.HasAttribute(GeometryCollection::Facades::FVertexBoneWeightsFacade::SkeletalMeshAttributeName, FGeometryCollection::GeometryGroup))
 			{
-				UE::Groom::Private::RegisterSkeletalMeshes(InCollection, StrandsFacade, GroomAsset);
+				UE::Groom::Private::RegisterSkeletalMeshes(GroomStrands, GeometryCollection::Facades::FVertexBoneWeightsFacade::SkeletalMeshAttributeName,
+					GeometryCollection::Facades::FVertexBoneWeightsFacade::GeometryLODAttributeName,
+					FName(FGeometryCollection::GeometryGroup), StrandsFacade.GetNumGeometry(), GroomAsset);
 			}
 		}
 	}
 }
 
-void FGroomAssetTerminalDataflowNode::Evaluate(UE::Dataflow::FContext& Context, const FDataflowOutput* Out) const
-{
-	const FManagedArrayCollection& InCollection = GetValue<FManagedArrayCollection>(Context, &Collection);
-	SetValue(Context, InCollection, &Collection);
-}
+void FGroomAssetTerminalDataflowNode_v2::Evaluate(UE::Dataflow::FContext& Context, const FDataflowOutput* Out) const
+{}
 
-UE::Dataflow::TConnectionReference<FCollectionAttributeKey> FGroomAssetTerminalDataflowNode::GetConnectionReference(int32 Index) const
+UE::Dataflow::TConnectionReference<FCollectionAttributeKey> FGroomAssetTerminalDataflowNode_v2::GetConnectionReference(int32 Index) const
 {
 	return { &AttributeKeys[Index], Index, &AttributeKeys };
 }
 
-TArray<UE::Dataflow::FPin> FGroomAssetTerminalDataflowNode::AddPins()
+TArray<UE::Dataflow::FPin> FGroomAssetTerminalDataflowNode_v2::AddPins()
 {
 	const int32 Index = AttributeKeys.AddDefaulted();
 	FDataflowInput& Input = RegisterInputArrayConnection(GetConnectionReference(Index));
@@ -295,7 +383,7 @@ TArray<UE::Dataflow::FPin> FGroomAssetTerminalDataflowNode::AddPins()
 	return { { UE::Dataflow::FPin::EDirection::INPUT, Input.GetType(), Input.GetName() } };
 }
 
-TArray<UE::Dataflow::FPin> FGroomAssetTerminalDataflowNode::GetPinsToRemove() const
+TArray<UE::Dataflow::FPin> FGroomAssetTerminalDataflowNode_v2::GetPinsToRemove() const
 {
 	const int32 Index = AttributeKeys.Num() - 1;
 	check(AttributeKeys.IsValidIndex(Index));
@@ -306,7 +394,7 @@ TArray<UE::Dataflow::FPin> FGroomAssetTerminalDataflowNode::GetPinsToRemove() co
 	return Super::GetPinsToRemove();
 }
 
-void FGroomAssetTerminalDataflowNode::OnPinRemoved(const UE::Dataflow::FPin& Pin)
+void FGroomAssetTerminalDataflowNode_v2::OnPinRemoved(const UE::Dataflow::FPin& Pin)
 {
 	const int32 Index = AttributeKeys.Num() - 1;
 	check(AttributeKeys.IsValidIndex(Index));
@@ -324,7 +412,7 @@ void FGroomAssetTerminalDataflowNode::OnPinRemoved(const UE::Dataflow::FPin& Pin
 	return Super::OnPinRemoved(Pin);
 }
 
-void FGroomAssetTerminalDataflowNode::PostSerialize(const FArchive& Ar)
+void FGroomAssetTerminalDataflowNode_v2::PostSerialize(const FArchive& Ar)
 {
 	// because we add pins we need to make sure we restore them when loading
 	// to make sure they can get properly reconnected
@@ -364,9 +452,9 @@ void FGroomAssetTerminalDataflowNode::PostSerialize(const FArchive& Ar)
 	}
 }
 
-bool FGroomAssetTerminalDataflowNode::ShouldInvalidateOnPropertyChanged(const FPropertyChangedEvent& InPropertyChangedEvent) const
+bool FGroomAssetTerminalDataflowNode_v2::ShouldInvalidateOnPropertyChanged(const FPropertyChangedEvent& InPropertyChangedEvent) const
 {
-	if (InPropertyChangedEvent.GetPropertyName() == GET_MEMBER_NAME_CHECKED(FGroomAssetTerminalDataflowNode, AttributeNames))
+	if (InPropertyChangedEvent.GetPropertyName() == GET_MEMBER_NAME_CHECKED(FGroomAssetTerminalDataflowNode_v2, AttributeNames))
 	{
 		// we only changing the name of the inputs, no need to inavlidate the node
 		return false;
@@ -374,15 +462,15 @@ bool FGroomAssetTerminalDataflowNode::ShouldInvalidateOnPropertyChanged(const FP
 	return true;
 }
 
-void FGroomAssetTerminalDataflowNode::OnPropertyChanged(UE::Dataflow::FContext& Context, const FPropertyChangedEvent& InPropertyChangedEvent)
+void FGroomAssetTerminalDataflowNode_v2::OnPropertyChanged(UE::Dataflow::FContext& Context, const FPropertyChangedEvent& InPropertyChangedEvent)
 {
-	if (InPropertyChangedEvent.GetPropertyName() == GET_MEMBER_NAME_CHECKED(FGroomAssetTerminalDataflowNode, AttributeNames))
+	if (InPropertyChangedEvent.GetPropertyName() == GET_MEMBER_NAME_CHECKED(FGroomAssetTerminalDataflowNode_v2, AttributeNames))
 	{
 		SyncInputNames();
 	}
 }
 
-bool FGroomAssetTerminalDataflowNode::SupportsDropConnectionOnNode(FName TypeName, UE::Dataflow::FPin::EDirection Direction) const
+bool FGroomAssetTerminalDataflowNode_v2::SupportsDropConnectionOnNode(FName TypeName, UE::Dataflow::FPin::EDirection Direction) const
 {
 	if (TypeName == UE::Groom::Private::AttributeKeyTypeName && Direction == UE::Dataflow::FPin::EDirection::OUTPUT)
 	{
@@ -391,7 +479,7 @@ bool FGroomAssetTerminalDataflowNode::SupportsDropConnectionOnNode(FName TypeNam
 	return false;
 }
 
-const FDataflowConnection* FGroomAssetTerminalDataflowNode::OnDropConnectionOnNode(const FDataflowConnection& DroppedConnection)
+const FDataflowConnection* FGroomAssetTerminalDataflowNode_v2::OnDropConnectionOnNode(const FDataflowConnection& DroppedConnection)
 {
 	if (DroppedConnection.GetType() == UE::Groom::Private::AttributeKeyTypeName && DroppedConnection.GetDirection() == UE::Dataflow::FPin::EDirection::OUTPUT)
 	{
@@ -418,7 +506,7 @@ const FDataflowConnection* FGroomAssetTerminalDataflowNode::OnDropConnectionOnNo
 	return nullptr;
 }
 
-void FGroomAssetTerminalDataflowNode::SyncInputNames()
+void FGroomAssetTerminalDataflowNode_v2::SyncInputNames()
 {
 	bool bChanged = false;
 	for (int32 KeyIndex = 0; KeyIndex < AttributeKeys.Num(); ++KeyIndex)
@@ -447,7 +535,7 @@ void FGroomAssetTerminalDataflowNode::SyncInputNames()
 	}
 }
 
-FName FGroomAssetTerminalDataflowNode::GenerateUniqueInputName(FName BaseName) const
+FName FGroomAssetTerminalDataflowNode_v2::GenerateUniqueInputName(FName BaseName) const
 {
 	if (BaseName.IsNone())
 	{
diff --git a/Engine/Plugins/Runtime/HairStrands/Source/HairStrandsDataflow/Private/HairStrandsDataflowModule.cpp b/Engine/Plugins/Runtime/HairStrands/Source/HairStrandsDataflow/Private/HairStrandsDataflowModule.cpp
--- a/Engine/Plugins/Runtime/HairStrands/Source/HairStrandsDataflow/Private/HairStrandsDataflowModule.cpp
+++ b/Engine/Plugins/Runtime/HairStrands/Source/HairStrandsDataflow/Private/HairStrandsDataflowModule.cpp
@@ -5,7 +5,6 @@
 #include "AttachGuidesRootsNode.h"
 #include "BuildGuidesLODsNode.h"
 #include "GetGroomAssetNode.h"
-#include "GroomCollectionFacades.h"
 #include "GroomAssetTerminalNode.h"
 #include "GetGroomAttributesNodes.h"
 #include "ResampleGuidesPointsNode.h"
@@ -13,7 +12,6 @@
 #include "BuildGroomSplineSkinningNode.h"
 #include "GenerateGuidesCurvesNode.h"
 #include "SmoothGuidesCurvesNode.h"
-#include "GroomDataflowRendering.h"
 #include "GroomDataflowVisualization.h"
 #include "Dataflow/DataflowCollectionAddScalarVertexPropertyNode.h"
 #include "Dataflow/DataflowCategoryRegistry.h"
@@ -21,52 +19,38 @@
 
 #define LOCTEXT_NAMESPACE "HairStrandsDataflow"
 
-class FGroomCollectionAddScalarVertexPropertyCallbacks : public IDataflowAddScalarVertexPropertyCallbacks
-{
-public:
-
-	const static FName Name;
-
-	virtual ~FGroomCollectionAddScalarVertexPropertyCallbacks() override = default;
-
-	virtual FName GetName() const override
-	{
-		return Name;
-	}
-
-	virtual TArray<FName> GetTargetGroupNames() const override
-	{
-		return {UE::Groom::FGroomGuidesFacade::VerticesGroup, UE::Groom::FGroomStrandsFacade::VerticesGroup};
-	}
-
-	virtual TArray<UE::Dataflow::FRenderingParameter> GetRenderingParameters() const override
-	{
-		return { {TEXT("GuidesRender"), FName("FGroomCollection"), {TEXT("Collection")}},
-			{TEXT("StrandsRender"), FName("FGroomCollection"), {TEXT("Collection")}} };
-	}
-};
-
-const FName FGroomCollectionAddScalarVertexPropertyCallbacks::Name = FName(TEXT("FGroomCollectionAddScalarVertexPropertyCallbacks"));
-
 void FHairStrandsDataflowModule::StartupModule()
 {
+	// Deprecated nodes
+PRAGMA_DISABLE_DEPRECATION_WARNINGS
 	DATAFLOW_NODE_REGISTER_CREATION_FACTORY(FGetGroomAssetDataflowNode);
-	DATAFLOW_NODE_REGISTER_CREATION_FACTORY(FGroomAssetTerminalDataflowNode);
-	DATAFLOW_NODE_REGISTER_CREATION_FACTORY(FResampleGuidesPointsDataflowNode);
-	DATAFLOW_NODE_REGISTER_CREATION_FACTORY(FTransferSkinWeightsGroomNode);
 	DATAFLOW_NODE_REGISTER_CREATION_FACTORY(FGenerateGuidesCurvesDataflowNode);
-	DATAFLOW_NODE_REGISTER_CREATION_FACTORY(FGetGroomAttributesDataflowNode);
+	DATAFLOW_NODE_REGISTER_CREATION_FACTORY(FResampleGuidesPointsDataflowNode);
 	DATAFLOW_NODE_REGISTER_CREATION_FACTORY(FSmoothGuidesCurvesDataflowNode);
+	DATAFLOW_NODE_REGISTER_CREATION_FACTORY(FTransferSkinWeightsGroomNode);
 	DATAFLOW_NODE_REGISTER_CREATION_FACTORY(FAttachGuidesRootsDataflowNode);
 	DATAFLOW_NODE_REGISTER_CREATION_FACTORY(FBuildGuidesLODsDataflowNode);
+	DATAFLOW_NODE_REGISTER_CREATION_FACTORY(FGroomAssetTerminalDataflowNode);
+	DATAFLOW_NODE_REGISTER_CREATION_FACTORY(FGetGroomAttributesDataflowNode);
 	DATAFLOW_NODE_REGISTER_CREATION_FACTORY(FBuildGroomSplineSkinWeightsNode);
 	DATAFLOW_NODE_REGISTER_CREATION_FACTORY(FConvertLinearToSplineSkinWeightsNode);
 	DATAFLOW_NODE_REGISTER_CREATION_FACTORY(FConvertSplineToLinearSkinWeightsNode);
-
-	UE::Groom::RegisterRenderingCallbacks();
-	
-	FDataflowAddScalarVertexPropertyCallbackRegistry::Get().RegisterCallbacks(
-		MakeUnique<FGroomCollectionAddScalarVertexPropertyCallbacks>());
+PRAGMA_ENABLE_DEPRECATION_WARNINGS
+
+	DATAFLOW_NODE_REGISTER_CREATION_FACTORY(FBuildSplineSkinWeightsDataflowNode);
+	DATAFLOW_NODE_REGISTER_CREATION_FACTORY(FLinearToSplineSkinWeightsDataflowNode);
+	DATAFLOW_NODE_REGISTER_CREATION_FACTORY(FSplineToLinearSkinWeightsDataflowNode);
+	DATAFLOW_NODE_REGISTER_CREATION_FACTORY(FGetGroomAssetDataflowNode_v2);
+	DATAFLOW_NODE_REGISTER_CREATION_FACTORY(FGroomAssetToCollectionDataflowNode);
+	DATAFLOW_NODE_REGISTER_CREATION_FACTORY(FGroomAssetTerminalDataflowNode_v2);
+	DATAFLOW_NODE_REGISTER_CREATION_FACTORY(FResampleCurvePointsDataflowNode);
+	DATAFLOW_NODE_REGISTER_CREATION_FACTORY(FTransferGeometrySkinWeightsDataflowNode);
+	DATAFLOW_NODE_REGISTER_CREATION_FACTORY(FGenerateCurveGeometryDataflowNode);
+	DATAFLOW_NODE_REGISTER_CREATION_FACTORY(FGetCurveAttributesDataflowNode);
+	DATAFLOW_NODE_REGISTER_CREATION_FACTORY(FSmoothCurvePointsDataflowNode);
+	DATAFLOW_NODE_REGISTER_CREATION_FACTORY(FAttachCurveRootsDataflowNode);
+	DATAFLOW_NODE_REGISTER_CREATION_FACTORY(FBuildCurveWeightsDataflowNode);
+	DATAFLOW_NODE_REGISTER_CREATION_FACTORY(FBuildCurveLODsDataflowNode);
 
 	UE::Dataflow::FDataflowSimulationVisualizationRegistry::GetInstance().RegisterVisualization(
 		MakeUnique<UE::Groom::FGroomDataflowSimulationVisualization>());
@@ -77,11 +61,6 @@ void FHairStrandsDataflowModule::StartupModule()
 
 void FHairStrandsDataflowModule::ShutdownModule()
 {
-	UE::Groom::DeregisterRenderingCallbacks();
-
-	FDataflowAddScalarVertexPropertyCallbackRegistry::Get().DeregisterCallbacks(
-		FGroomCollectionAddScalarVertexPropertyCallbacks::Name);
-
 	UE::Dataflow::FDataflowSimulationVisualizationRegistry::GetInstance().DeregisterVisualization(
 		UE::Groom::FGroomDataflowSimulationVisualization::Name);
 }
diff --git a/Engine/Plugins/Runtime/HairStrands/Source/HairStrandsDataflow/Private/ResampleGuidesPointsNode.cpp b/Engine/Plugins/Runtime/HairStrands/Source/HairStrandsDataflow/Private/ResampleGuidesPointsNode.cpp
--- a/Engine/Plugins/Runtime/HairStrands/Source/HairStrandsDataflow/Private/ResampleGuidesPointsNode.cpp
+++ b/Engine/Plugins/Runtime/HairStrands/Source/HairStrandsDataflow/Private/ResampleGuidesPointsNode.cpp
@@ -1,42 +1,51 @@
-﻿// Copyright Epic Games, Inc. All Rights Reserved.
+// Copyright Epic Games, Inc. All Rights Reserved.
 
 #include "ResampleGuidesPointsNode.h"
-#include "GroomCollectionFacades.h"
 #include "Dataflow/DataflowObjectInterface.h"
+#include "GeometryCollection/Facades/CollectionCurveFacade.h"
 
 #include UE_INLINE_GENERATED_CPP_BY_NAME(ResampleGuidesPointsNode)
 
 namespace UE::Groom::Private
 {
-	FORCEINLINE int32 ComputeNumPoints(const FGroomGuidesFacade& GroomFacade, const int32 GuidePoints) 
+	FORCEINLINE int32 ComputeCurvesPoints(const GeometryCollection::Facades::FCollectionCurveGeometryFacade& CurvesFacade, const FDataflowCurveSelection& CurveSelection, const int32 GuidePoints) 
 	{
-		int32 CurveOffset = 0, NumPoints = 0;
-		for(int32 ObjectIndex = 0, NumObjects = GroomFacade.GetNumObjects(); ObjectIndex < NumObjects; ++ObjectIndex)
+		const bool bValidSelection = CurveSelection.IsValidForCollection(CurvesFacade.GetManagedArrayCollection());
+
+		int32 PointOffset = 0, NumPoints = 0;
+		for(int32 CurveIndex = 0, NumCurves = CurvesFacade.GetNumCurves(); CurveIndex < NumCurves; ++CurveIndex)
 		{
-			const uint32 SamplesCount = (GuidePoints != 0) ? GuidePoints : GroomFacade.GetObjectPointSamples()[ObjectIndex];
-			NumPoints +=  (GroomFacade.GetObjectCurveOffsets()[ObjectIndex] - CurveOffset) * SamplesCount;
-			CurveOffset = GroomFacade.GetObjectCurveOffsets()[ObjectIndex];
+			if(!bValidSelection || (bValidSelection && (CurveSelection.IsSelected(CurveIndex))))
+			{
+				NumPoints += GuidePoints;
+			}
+			else
+			{
+				NumPoints += CurvesFacade.GetCurvePointOffsets()[CurveIndex] - PointOffset;
+			}
+			PointOffset = CurvesFacade.GetCurvePointOffsets()[CurveIndex];
 		}
 		return NumPoints;
 	}
 
-	FORCEINLINE void BuildGuidePoints(const FGroomGuidesFacade& GroomFacade, const int32 CurveIndex, const int32 PointOffset, const int32 SampleOffset,
-		TArray<FVector3f>& SamplePositions, const int32 GuidePoints) 
+	FORCEINLINE void BuildResampledPoints(const TArray<FVector3f>& PointRestPositions, const TArray<int32>& CurveObjectIndices, const TArray<int32>& CurvePointOffsets,
+		const TArray<int32>& ObjectPointSamples, const int32 CurveIndex, const int32 PointOffset, const int32 SampleOffset,
+					TArray<FVector3f>& SamplePositions, const int32 GuidePoints) 
 	{
-		const int32 ObjectIndex = GroomFacade.GetCurveObjectIndices()[CurveIndex];
-		const int32 PointsCount = (GroomFacade.GetCurvePointOffsets()[CurveIndex]-PointOffset)-1;
-		const int32 SamplesCount = (GuidePoints != 0) ? GuidePoints-1 : GroomFacade.GetObjectPointSamples()[ObjectIndex]-1;
+		const int32 ObjectIndex = CurveObjectIndices[CurveIndex];
+		const int32 PointsCount = (CurvePointOffsets[CurveIndex]-PointOffset)-1;
+		const int32 SamplesCount = (GuidePoints != 0) ? GuidePoints-1 : !ObjectPointSamples.IsEmpty() ? ObjectPointSamples[ObjectIndex]-1 : 0;
 
 		float CurveLength = 0.0;
 		TArray<float> EdgeLengths; EdgeLengths.Init(0.0f, PointsCount);
 		for(int32 EdgeIndex = 0; EdgeIndex < PointsCount; ++EdgeIndex)
 		{
-			EdgeLengths[EdgeIndex] = (GroomFacade.GetPointRestPositions()[EdgeIndex+PointOffset+1] -
-				GroomFacade.GetPointRestPositions()[EdgeIndex+PointOffset]).Length();
+			EdgeLengths[EdgeIndex] = (PointRestPositions[EdgeIndex+PointOffset+1] -
+				PointRestPositions[EdgeIndex+PointOffset]).Length();
 			CurveLength += EdgeLengths[EdgeIndex];
 		}
 
-		SamplePositions[SampleOffset + 1] = GroomFacade.GetPointRestPositions()[PointOffset];
+		SamplePositions[SampleOffset + 1] = PointRestPositions[PointOffset];
 		for(int32 SampleIndex = 1; SampleIndex < SamplesCount-1; ++SampleIndex)
 		{
 			const float SampleCoord = static_cast<float>(SampleIndex) / (SamplesCount - 1.0f);
@@ -52,67 +61,109 @@ namespace UE::Groom::Private
 					const int32 NextPoint = PrevPoint+1;
 					
 					const float SampleAlpha = (LocalLength - SampleLength) / EdgeLengths[EdgeIndex];
-					SamplePositions[SampleOffset + SampleIndex + 1] = GroomFacade.GetPointRestPositions()[PrevPoint] * SampleAlpha +
-						GroomFacade.GetPointRestPositions()[NextPoint] * (1.0-SampleAlpha);
+					SamplePositions[SampleOffset + SampleIndex + 1] = PointRestPositions[PrevPoint] * SampleAlpha +
+						PointRestPositions[NextPoint] * (1.0-SampleAlpha);
 					break;
 				}
 			}
 		}
-		SamplePositions[SampleOffset+SamplesCount] = GroomFacade.GetPointRestPositions()[PointOffset+PointsCount];
+		SamplePositions[SampleOffset+SamplesCount] = PointRestPositions[PointOffset+PointsCount];
 		SamplePositions[SampleOffset] = 2 * SamplePositions[SampleOffset+1] - SamplePositions[SampleOffset+2];
 	}
-	
-	FORCEINLINE void ResampleGuidesPoints(FGroomGuidesFacade& GroomFacade, const int32 GuidePoints)
+
+	FORCEINLINE void ResampleCurvesPoints(GeometryCollection::Facades::FCollectionCurveGeometryFacade& CurvesFacade, const FDataflowCurveSelection& CurveSelection, const int32 GuidePoints)
 	{
-		if(GroomFacade.IsValid())
+		if(CurvesFacade.IsValid())
 		{
-			const int32 NumSamples = ComputeNumPoints(GroomFacade, GuidePoints);
-			const int32 NumCurves = GroomFacade.GetNumCurves();
+			const int32 NumSamples = ComputeCurvesPoints(CurvesFacade, CurveSelection, GuidePoints);
+			const int32 NumCurves = CurvesFacade.GetNumCurves();
+
+			TArray<FVector3f> PointSampledPositions;
+			PointSampledPositions.Init(FVector3f::Zero(), NumSamples);
 
-			TArray<FVector3f> SamplePositions;
-			SamplePositions.Init(FVector3f::Zero(), NumSamples);
+			TArray<int32> CurveSampleOffsets;
+			CurveSampleOffsets.Init(0, NumCurves);
 
-			TArray<int32> CurveOffsets;
-			CurveOffsets.Init(0, NumCurves);
+			const bool bValidSelection = CurveSelection.IsValidForCollection(CurvesFacade.GetManagedArrayCollection());
 
 			int32 PointOffset = 0;
 			int32 SampleOffset = 0;
 			for(int32 CurveIndex = 0; CurveIndex < NumCurves; ++CurveIndex)
 			{
-				BuildGuidePoints(GroomFacade, CurveIndex, PointOffset, SampleOffset, SamplePositions, GuidePoints);
-				const uint32 SamplesCount = (GuidePoints != 0) ? GuidePoints : GroomFacade.GetObjectPointSamples()[GroomFacade.GetCurveObjectIndices()[CurveIndex]];
-					
-				PointOffset = GroomFacade.GetCurvePointOffsets()[CurveIndex];
-				SampleOffset += SamplesCount;
-
-				CurveOffsets[CurveIndex] = SampleOffset;
+				const int32 NumPoints = CurvesFacade.GetCurvePointOffsets()[CurveIndex]-PointOffset;
+				if(!bValidSelection || (bValidSelection && CurveSelection.IsSelected(CurveIndex)))
+				{
+					UE::Groom::Private::BuildResampledPoints(CurvesFacade.GetPointRestPositions(), CurvesFacade.GetCurveGeometryIndices(), CurvesFacade.GetCurvePointOffsets(),
+							{}, CurveIndex, PointOffset, SampleOffset, PointSampledPositions, GuidePoints);
+					SampleOffset += GuidePoints;
+				}
+				else
+				{
+					for(int32 PointIndex = 0; PointIndex < NumPoints; ++PointIndex)
+					{
+						PointSampledPositions[PointIndex+PointOffset] = CurvesFacade.GetPointRestPositions()[PointIndex+PointOffset];
+					}
+					SampleOffset += NumPoints;
+				}
+				CurveSampleOffsets[CurveIndex] = SampleOffset;
+				PointOffset = CurvesFacade.GetCurvePointOffsets()[CurveIndex];
 			}
-			// Resize the points groups
-			GroomFacade.ResizePointsGroups(NumSamples);
+			const TArray<int32> GeometryCurveOffsets = CurvesFacade.GetGeometryCurveOffsets();
+			const TArray<FString> GeometryGroupNames = CurvesFacade.GetGeometryGroupNames();
+			const TArray<float> GeometryCurveThickness = CurvesFacade.GetGeometryCurveThickness();
+			const TArray<int32> CurveSourceIndices = CurvesFacade.GetCurveSourceIndices();
+			
+			CurvesFacade.InitCurvesCollection(PointSampledPositions, CurveSampleOffsets, GeometryCurveOffsets, GeometryGroupNames, GeometryCurveThickness, CurveSourceIndices);
+		}
+	}
+
+	static int32 ClosestPowerOfTwo(const int32 InputValue)
+	{
+		if(InputValue > 0)
+		{
+			const float LogValue = FMath::Log2(static_cast<float>(InputValue));
+			const int32 LowerExp = FMath::FloorToInt(LogValue);
+			const int32 UpperExp = FMath::CeilToInt(LogValue);
 
-			// Set the curve offsets
-			GroomFacade.SetCurvePointOffsets(CurveOffsets);
+			const int32 LowerValue = 1 << LowerExp;
+			const int32 UpperValue = 1 << UpperExp;
 
-			// Set the point positions and kinematic weights
-			GroomFacade.SetPointRestPositions(SamplePositions);
+			return FMath::Clamp((InputValue-LowerValue < UpperValue-InputValue) ? LowerValue : UpperValue, 2, 64);
 		}
+		return 2;
 	}
 }
 
 void FResampleGuidesPointsDataflowNode::Evaluate(UE::Dataflow::FContext& Context, const FDataflowOutput* Out) const
+{}
+
+FResampleCurvePointsDataflowNode::FResampleCurvePointsDataflowNode(const UE::Dataflow::FNodeParameters& InParam, FGuid InGuid)
+: FDataflowNode(InParam, InGuid)
+{
+	RegisterInputConnection(&Collection);
+	RegisterInputConnection(&CurveSelection);
+	RegisterInputConnection(&NumPoints).SetCanHidePin(true).SetPinIsHidden(true);
+	RegisterOutputConnection(&Collection, &Collection);
+}
+
+void FResampleCurvePointsDataflowNode::Evaluate(UE::Dataflow::FContext& Context, const FDataflowOutput* Out) const
 {
 	if (Out->IsA<FManagedArrayCollection>(&Collection))
 	{
-		FManagedArrayCollection GroomCollection = GetValue<FManagedArrayCollection>(Context, &Collection);
+		FManagedArrayCollection GuidesCollection = GetValue<FManagedArrayCollection>(Context, &Collection);
+		const FDataflowCurveSelection& DataflowSelection = GetValue<FDataflowCurveSelection>(Context, &CurveSelection);
 		
-		UE::Groom::FGroomGuidesFacade GuidesFacade(GroomCollection);
-
-		if(GuidesFacade.IsValid())
+		const int32 GuidePoints = IsConnected(&NumPoints) ?
+			UE::Groom::Private::ClosestPowerOfTwo(GetValue<int32>(Context, &NumPoints)) : static_cast<int32>(PointsCount);
+		
+		GeometryCollection::Facades::FCollectionCurveGeometryFacade CurvesFacade(GuidesCollection);
+		
+		if(CurvesFacade.IsValid() && GuidePoints > 0)
 		{
-			UE::Groom::Private::ResampleGuidesPoints(GuidesFacade, static_cast<uint8>(PointsCount));
+			UE::Groom::Private::ResampleCurvesPoints(CurvesFacade, DataflowSelection, GuidePoints);
 		}
 		
-		SetValue(Context, MoveTemp(GroomCollection), &Collection);
+		SetValue(Context, MoveTemp(GuidesCollection), &Collection);
 	}
 }
 
diff --git a/Engine/Plugins/Runtime/HairStrands/Source/HairStrandsDataflow/Private/SmoothGuidesCurvesNode.cpp b/Engine/Plugins/Runtime/HairStrands/Source/HairStrandsDataflow/Private/SmoothGuidesCurvesNode.cpp
--- a/Engine/Plugins/Runtime/HairStrands/Source/HairStrandsDataflow/Private/SmoothGuidesCurvesNode.cpp
+++ b/Engine/Plugins/Runtime/HairStrands/Source/HairStrandsDataflow/Private/SmoothGuidesCurvesNode.cpp
@@ -1,68 +1,109 @@
 // Copyright Epic Games, Inc. All Rights Reserved.
 
 #include "SmoothGuidesCurvesNode.h"
-#include "GroomCollectionFacades.h"
+
+#include "AnimAssetFindReplace.h"
 #include "Dataflow/DataflowObjectInterface.h"
+#include "GeometryCollection/Facades/CollectionCurveFacade.h"
+#include "Views/SInteractiveCurveEditorView.h"
 
 #include UE_INLINE_GENERATED_CPP_BY_NAME(SmoothGuidesCurvesNode)
 
 namespace UE::Groom::Private
 {
-	FORCEINLINE void SmoothGuidesPoints(FGroomGuidesFacade& GroomFacade, const float SmoothingFactor)
+	FORCEINLINE void BuildSmoothedPoints(const TArray<FVector3f>& PointRestPositions, const int32 PointOffset, const int32 PointEnd, const float SmoothingFactor, TArray<FVector3f>& PointSmoothedPositions)
+	{
+		FVector3f DirM1 = PointRestPositions[PointOffset+1] - PointRestPositions[PointOffset];
+		FVector3f DirM2 = DirM1;
+
+		const float Gamma1 = 2.0 * (1.0-SmoothingFactor);
+		const float Gamma2 = - (1.0-SmoothingFactor)*(1.0-SmoothingFactor);
+		const float Gamma3 = SmoothingFactor*SmoothingFactor;
+				
+		PointSmoothedPositions[PointOffset] = PointRestPositions[PointOffset];
+				
+		for(int32 PointIndex = PointOffset; PointIndex < PointEnd-1; ++PointIndex)
+		{
+			const FVector3f DirM3 = PointRestPositions[PointIndex+1] - PointRestPositions[PointIndex];
+			const FVector3f DirMi = Gamma1 * DirM1 + Gamma2 * DirM2 + Gamma3 * DirM3;
+
+			PointSmoothedPositions[PointIndex+1] = PointSmoothedPositions[PointIndex] + DirMi;
+					
+			DirM2 = DirM1;
+			DirM1 = DirMi;
+		}
+	}
+
+	FORCEINLINE void SmoothCurvesPoints(GeometryCollection::Facades::FCollectionCurveGeometryFacade& CurvesFacade, const FDataflowCurveSelection& CurveSelection, float SmoothingFactor)
 	{
-		if(GroomFacade.IsValid())
+		if(CurvesFacade.IsValid())
 		{
-			const int32 NumPoints = GroomFacade.GetNumPoints();
-			const int32 NumCurves = GroomFacade.GetNumCurves();
+			const int32 NumPoints = CurvesFacade.GetNumPoints();
+			const int32 NumCurves = CurvesFacade.GetNumCurves();
+
+			TArray<FVector3f> PointSmoothedPositions;
+			PointSmoothedPositions.Init(FVector3f::Zero(), NumPoints);
 
-			TArray<FVector3f> SmoothedPositions;
-			SmoothedPositions.Init(FVector3f::Zero(), NumPoints);
+			const bool bValidSelection = CurveSelection.IsValidForCollection(CurvesFacade.GetManagedArrayCollection());;
 
 			int32 PointOffset = 0;
 			for(int32 CurveIndex = 0; CurveIndex < NumCurves; ++CurveIndex)
 			{
-				FVector3f DirM1 = GroomFacade.GetPointRestPositions()[PointOffset+1] - GroomFacade.GetPointRestPositions()[PointOffset];
-				FVector3f DirM2 = DirM1;
-
-				const float Gamma1 = 2.0 * (1.0-SmoothingFactor);
-				const float Gamma2 = - (1.0-SmoothingFactor)*(1.0-SmoothingFactor);
-				const float Gamma3 = SmoothingFactor*SmoothingFactor;
-				
-				SmoothedPositions[PointOffset] = GroomFacade.GetPointRestPositions()[PointOffset];
-				
-				for(int32 PointIndex = PointOffset; PointIndex < GroomFacade.GetCurvePointOffsets()[CurveIndex]-1; ++PointIndex)
+				if(!bValidSelection || (bValidSelection && (CurveSelection.IsSelected(CurveIndex))))
 				{
-					const FVector3f DirM3 = GroomFacade.GetPointRestPositions()[PointIndex+1] - GroomFacade.GetPointRestPositions()[PointIndex];
-					const FVector3f DirMi = Gamma1 * DirM1 + Gamma2 * DirM2 + Gamma3 * DirM3;
-
-					SmoothedPositions[PointIndex+1] = SmoothedPositions[PointIndex] + DirMi;
-					
-					DirM2 = DirM1;
-					DirM1 = DirMi;
+					BuildSmoothedPoints(CurvesFacade.GetPointRestPositions(), PointOffset,
+					CurvesFacade.GetCurvePointOffsets()[CurveIndex], SmoothingFactor, PointSmoothedPositions);
+				}
+				else
+				{
+					for(int32 PointIndex = PointOffset; PointIndex < CurvesFacade.GetCurvePointOffsets()[CurveIndex]; ++PointIndex)
+					{
+						PointSmoothedPositions[PointIndex] = CurvesFacade.GetPointRestPositions()[PointIndex];
+					}
 				}
-				PointOffset = GroomFacade.GetCurvePointOffsets()[CurveIndex];
+				PointOffset = CurvesFacade.GetCurvePointOffsets()[CurveIndex];
 			}
 
-			// Set the point positions and kinematic weights
-			GroomFacade.SetPointRestPositions(SmoothedPositions);
+			const TArray<int32> CurvePointOffsets = CurvesFacade.GetCurvePointOffsets();
+			const TArray<int32> GeometryCurveOffsets = CurvesFacade.GetGeometryCurveOffsets();
+			const TArray<FString> GeometryGroupNames = CurvesFacade.GetGeometryGroupNames();
+			const TArray<float> GeometryCurveThickness = CurvesFacade.GetGeometryCurveThickness();
+			const TArray<int32> CurveSourceIndices = CurvesFacade.GetCurveSourceIndices();
+			
+			CurvesFacade.InitCurvesCollection(PointSmoothedPositions, CurvePointOffsets, GeometryCurveOffsets, GeometryGroupNames, GeometryCurveThickness, CurveSourceIndices);
 		}
 	}
 }
 
 void FSmoothGuidesCurvesDataflowNode::Evaluate(UE::Dataflow::FContext& Context, const FDataflowOutput* Out) const
+{}
+
+FSmoothCurvePointsDataflowNode::FSmoothCurvePointsDataflowNode(const UE::Dataflow::FNodeParameters& InParam, FGuid InGuid)
+	: FDataflowNode(InParam, InGuid)
+{
+	RegisterInputConnection(&Collection);
+	RegisterInputConnection(&CurveSelection);
+	RegisterInputConnection(&SmoothingFactor).SetCanHidePin(true).SetPinIsHidden(true);
+	RegisterOutputConnection(&Collection, &Collection);
+}
+
+void FSmoothCurvePointsDataflowNode::Evaluate(UE::Dataflow::FContext& Context, const FDataflowOutput* Out) const
 {
 	if (Out->IsA<FManagedArrayCollection>(&Collection))
 	{
-		FManagedArrayCollection GroomCollection = GetValue<FManagedArrayCollection>(Context, &Collection);
+		FManagedArrayCollection GeometryCollection = GetValue<FManagedArrayCollection>(Context, &Collection);
+		const float SmoothedFactor = GetValue<float>(Context, &SmoothingFactor);
 		
-		UE::Groom::FGroomGuidesFacade GuidesFacade(GroomCollection);
+		GeometryCollection::Facades::FCollectionCurveGeometryFacade CurvesFacade(GeometryCollection);
 
-		if(GuidesFacade.IsValid())
+		if(CurvesFacade.IsValid() && SmoothedFactor > 0.0f)
 		{
-			UE::Groom::Private::SmoothGuidesPoints(GuidesFacade, 1.0-SmoothingFactor);
+			const FDataflowCurveSelection& DataflowSelection = GetValue<FDataflowCurveSelection>(Context, &CurveSelection);
+			UE::Groom::Private::SmoothCurvesPoints(CurvesFacade, DataflowSelection, 1.0-SmoothedFactor);
 		}
 		
-		SetValue(Context, MoveTemp(GroomCollection), &Collection);
+		SetValue(Context, MoveTemp(GeometryCollection), &Collection);
 	}
 }
 
+
diff --git a/Engine/Plugins/Runtime/HairStrands/Source/HairStrandsDataflow/Public/AttachGuidesRootsNode.h b/Engine/Plugins/Runtime/HairStrands/Source/HairStrandsDataflow/Public/AttachGuidesRootsNode.h
--- a/Engine/Plugins/Runtime/HairStrands/Source/HairStrandsDataflow/Public/AttachGuidesRootsNode.h
+++ b/Engine/Plugins/Runtime/HairStrands/Source/HairStrandsDataflow/Public/AttachGuidesRootsNode.h
@@ -11,8 +11,8 @@
 #include "AttachGuidesRootsNode.generated.h"
 
 /** Attach the guides roots by setting their kinematic weights to 1.0f */
-USTRUCT(meta = (Experimental, DataflowGroom))
-struct FAttachGuidesRootsDataflowNode : public FDataflowNode
+USTRUCT(meta = (Experimental, DataflowGroom, Deprecated="5.7"))
+struct UE_DEPRECATED(5.7, "Use BuildCurvesWeights node instead.") FAttachGuidesRootsDataflowNode : public FDataflowNode
 {
 	GENERATED_BODY()
 	
@@ -46,3 +46,79 @@ public:
 	FCollectionAttributeKey KinematicWeightsKey;
 };
 
+/** Attach the curve roots by setting their kinematic weights to 1.0f */
+USTRUCT(meta = (Experimental, DataflowGroom))
+struct FAttachCurveRootsDataflowNode : public FDataflowNode
+{
+	GENERATED_BODY()
+	
+	DATAFLOW_NODE_DEFINE_INTERNAL(FAttachCurveRootsDataflowNode, "AttachCurveRoots", "Groom", "")
+
+public:
+	
+	FAttachCurveRootsDataflowNode(const UE::Dataflow::FNodeParameters& InParam, FGuid InGuid = FGuid::NewGuid())
+	: FDataflowNode(InParam, InGuid)
+	{
+		RegisterInputConnection(&Collection);
+		RegisterInputConnection(&CurveSelection);
+		RegisterOutputConnection(&Collection, &Collection);
+		RegisterOutputConnection(&KinematicWeightsKey);
+	}
+
+	//~ Begin FDataflowNode interface
+	virtual void Evaluate(UE::Dataflow::FContext& Context, const FDataflowOutput* Out) const override;
+	//~ End FDataflowNode interface
+
+	/** Managed array collection to be used to store datas */
+	UPROPERTY(meta = (DataflowInput, DataflowOutput, DisplayName = "Collection", DataflowPassthrough  = "Collection", DataflowRenderGroups = "Surface"))
+	FManagedArrayCollection Collection;
+
+	/** Curve selection to focus the geometry transfer spatially */
+	UPROPERTY(meta = (DataflowInput, DisplayName = "Curve Selection"))
+	FDataflowCurveSelection CurveSelection;
+	
+	/** Vertex kinematic weights key to be used in other nodes if necessary */
+	UPROPERTY(meta = (DisplayName = "Kinematic Weights", DataflowOutput))
+	FCollectionAttributeKey KinematicWeightsKey;
+};
+
+/** Build a weight map that will be identical on every curves */
+USTRUCT(meta = (Experimental, DataflowGroom))
+struct FBuildCurveWeightsDataflowNode : public FDataflowNode
+{
+	GENERATED_BODY()
+	
+	DATAFLOW_NODE_DEFINE_INTERNAL(FBuildCurveWeightsDataflowNode, "BuildCurveWeights", "Groom", "")
+
+public:
+	
+	FBuildCurveWeightsDataflowNode(const UE::Dataflow::FNodeParameters& InParam, FGuid InGuid = FGuid::NewGuid())
+	: FDataflowNode(InParam, InGuid)
+	{
+		RegisterInputConnection(&Collection);
+		RegisterInputConnection(&CurveSelection);
+		RegisterInputConnection(&WeightsAttribute);
+		RegisterOutputConnection(&Collection, &Collection);
+	}
+
+	//~ Begin FDataflowNode interface
+	virtual void Evaluate(UE::Dataflow::FContext& Context, const FDataflowOutput* Out) const override;
+	//~ End FDataflowNode interface
+
+	/** Managed array collection to be used to store datas */
+	UPROPERTY(meta = (DataflowInput, DataflowOutput, DisplayName = "Collection", DataflowPassthrough  = "Collection", DataflowRenderGroups = "Surface"))
+	FManagedArrayCollection Collection;
+
+	/** Curve selection to focus the geometry transfer spatially */
+	UPROPERTY(meta = (DataflowInput, DisplayName = "Curve Selection"))
+	FDataflowCurveSelection CurveSelection;
+
+	/** Weight value along each curves */
+	UPROPERTY(EditAnywhere, Category = "Curves", meta = (DisplayName = "Curve Weights", ViewMinInput = "0.0", ViewMaxInput = "1.0", ViewMinOutput = "0.0", ViewMaxOutput = "1.0", TimeLineLength = "1.0", XAxisName = "Curve Coordinate (0,1)", YAxisName = "Weight value", ToolTip = "This curve determines the weight value from root to tip. \n The X axis range is [0,1], 0 mapping the root and 1 the tip"))
+	FRuntimeFloatCurve CurveWeights;
+	
+	/** Vertex kinematic weights key to be used in other nodes if necessary */
+	UPROPERTY(EditAnywhere, Category = "Curves", meta = (DisplayName = "Weights Attribute", DataflowInput))
+	FCollectionAttributeKey WeightsAttribute;
+};
+
diff --git a/Engine/Plugins/Runtime/HairStrands/Source/HairStrandsDataflow/Public/BuildGroomSkinningNodes.h b/Engine/Plugins/Runtime/HairStrands/Source/HairStrandsDataflow/Public/BuildGroomSkinningNodes.h
--- a/Engine/Plugins/Runtime/HairStrands/Source/HairStrandsDataflow/Public/BuildGroomSkinningNodes.h
+++ b/Engine/Plugins/Runtime/HairStrands/Source/HairStrandsDataflow/Public/BuildGroomSkinningNodes.h
@@ -5,15 +5,16 @@
 #include "CoreMinimal.h"
 #include "GroomAsset.h"
 #include "ChaosLog.h"
-#include "GroomCollectionFacades.h"
 #include "Dataflow/DataflowConnectionTypes.h"
 #include "Dataflow/DataflowCore.h"
+#include "GeometryCollection/Facades/CollectionVertexBoneWeightsFacade.h"
+#include "GetGroomAssetNode.h"
 
 #include "BuildGroomSkinningNodes.generated.h"
 
 /** Build the guides skinning by transferring the indices weights from a skelmesh */
-USTRUCT(meta = (Experimental, DataflowGroom))
-struct FTransferSkinWeightsGroomNode : public FDataflowNode
+USTRUCT(meta = (Experimental, DataflowGroom, Deprecated="5.7"))
+struct UE_DEPRECATED(5.7, "Use the newer version of this node instead.") FTransferSkinWeightsGroomNode : public FDataflowNode
 {
 	GENERATED_BODY()
 	
@@ -25,7 +26,6 @@ public:
 
 	//~ Begin FDataflowNode interface
 	virtual void Evaluate(UE::Dataflow::FContext& Context, const FDataflowOutput* Out) const override;
-	virtual TArray<UE::Dataflow::FRenderingParameter> GetRenderParametersImpl() const override;
 	//~ End FDataflowNode interface
 
 	/** Managed array collection to be used to store datas */
@@ -68,4 +68,54 @@ private :
 	FCollectionAttributeKey GetBoneWeightsKey() const;
 };
 
+/** Build the guides skinning by transferring the indices weights from a skelmesh */
+USTRUCT(meta = (Experimental, DataflowGroom))
+struct FTransferGeometrySkinWeightsDataflowNode : public FDataflowNode
+{
+	GENERATED_BODY()
+	
+	DATAFLOW_NODE_DEFINE_INTERNAL(FTransferGeometrySkinWeightsDataflowNode, "TransferGeometrySkinWeights", "Groom", "")
+
+public:
+	FTransferGeometrySkinWeightsDataflowNode(const UE::Dataflow::FNodeParameters& InParam, FGuid InGuid = FGuid::NewGuid());
+	
+
+private:
+	//~ Begin FDataflowNode interface
+	virtual void Evaluate(UE::Dataflow::FContext& Context, const FDataflowOutput* Out) const override;
+	//~ End FDataflowNode interface
+
+	/** Managed array collection to be used to store datas */
+	UPROPERTY(meta = (DataflowInput, DataflowOutput, DisplayName = "Collection", DataflowPassthrough = "Collection", DataflowRenderGroups = "Surface"))
+	FManagedArrayCollection Collection;
+
+	/** Vertex selection to focus the geometry transfer spatially */
+	UPROPERTY(meta = (DataflowInput, DisplayName = "Vertex Selection"))
+	FDataflowVertexSelection VertexSelection;
+
+	/** SkeletalMesh used to transfer the skinning weights. Will be stored onto the groom asset */
+	UPROPERTY(EditAnywhere, Category="Skeletal Mesh", meta = (DataflowInput, DataflowPassthrough = "SkeletalMesh"))
+	TObjectPtr<USkeletalMesh> SkeletalMesh = nullptr;
+
+	/** LOD used to transfer the weights */
+	UPROPERTY(EditAnywhere, Category="Skeletal Mesh", meta = (DataflowInput, DisplayName = "LOD Index"))
+	int32 LODIndex = 0;
+
+	/** The relative transform between the skeletal mesh and the groom asset. */
+	UPROPERTY(EditAnywhere, Category = "Skeletal Mesh", meta = (DataflowInput, DisplayName = "Relative Transform"))
+	FTransform RelativeTransform;
 
+	/** Bone indices key to be used in other nodes if necessary */
+	UPROPERTY(meta = (DisplayName = "Bone Indices", DataflowOutput))
+	FCollectionAttributeKey BoneIndicesKey;
+
+	/** Bone weights key to be used in other nodes if necessary */
+	UPROPERTY(meta = (DisplayName = "Bone Weights", DataflowOutput))
+	FCollectionAttributeKey BoneWeightsKey;
+	
+	/** Get the bone indices key */
+	static FCollectionAttributeKey GetBoneIndicesKey();
+
+	/** Get the bone weights key */
+	static FCollectionAttributeKey GetBoneWeightsKey();
+};
diff --git a/Engine/Plugins/Runtime/HairStrands/Source/HairStrandsDataflow/Public/BuildGroomSplineSkinningNode.h b/Engine/Plugins/Runtime/HairStrands/Source/HairStrandsDataflow/Public/BuildGroomSplineSkinningNode.h
--- a/Engine/Plugins/Runtime/HairStrands/Source/HairStrandsDataflow/Public/BuildGroomSplineSkinningNode.h
+++ b/Engine/Plugins/Runtime/HairStrands/Source/HairStrandsDataflow/Public/BuildGroomSplineSkinningNode.h
@@ -3,17 +3,16 @@
 #pragma once 
 
 #include "CoreMinimal.h"
-#include "GroomCollectionFacades.h"
 #include "Dataflow/DataflowCore.h"
 #include "Dataflow/DataflowConnectionTypes.h"
 #include "Engine/SkeletalMesh.h"
+#include "GetGroomAssetNode.h"
 
 #include "BuildGroomSplineSkinningNode.generated.h"
 
-
 /** Build spline skinning data from skeletal mesh */
-USTRUCT(meta = (Experimental, DataflowGroom))
-struct FBuildGroomSplineSkinWeightsNode : public FDataflowNode
+USTRUCT(meta = (Experimental, DataflowGroom, Deprecated="5.7"))
+struct UE_DEPRECATED(5.7, "Use the newer version of this node instead.") FBuildGroomSplineSkinWeightsNode : public FDataflowNode
 {
 	GENERATED_BODY()
 
@@ -24,7 +23,6 @@ public:
 
 	//~ Begin FDataflowNode interface
 	virtual void Evaluate(UE::Dataflow::FContext& Context, const FDataflowOutput* Out) const override;
-	virtual TArray<UE::Dataflow::FRenderingParameter> GetRenderParametersImpl() const override;
 	//~ End FDataflowNode interface
 
 	/** Managed array collection to be used to store datas */
@@ -56,7 +54,6 @@ public:
 	UPROPERTY(EditAnywhere, Category = "Skeletal Mesh")
 	int32 SamplesPerSegment = 64;
 
-
 	/** Groom group index to build skinning weights for. -1 will build all groups */
 	UPROPERTY(EditAnywhere, Category="Groom Groups")
 	int32 GroupIndex = INDEX_NONE;
@@ -72,15 +69,11 @@ public:
 	/** Spline bones attribute name. Used for storing root and end spline bone for each vertex. */
 	UPROPERTY(EditAnywhere, Category = "Vertex Attributes")
 	FString SplineBonesParamName = "SplineBones";
-
-//private:
-	FCollectionAttributeKey GetSplineParamKey() const;
-	FCollectionAttributeKey GetSplineBonesKey() const;
 };
 
 /** Convert spline skinning data to linear data */
-USTRUCT(meta = (Experimental, DataflowGroom))
-struct FConvertSplineToLinearSkinWeightsNode : public FDataflowNode
+USTRUCT(meta = (Experimental, DataflowGroom, Deprecated="5.7"))
+struct UE_DEPRECATED(5.7, "Use the newer version of this node instead.") FConvertSplineToLinearSkinWeightsNode : public FDataflowNode
 {
 	GENERATED_BODY()
 
@@ -91,7 +84,6 @@ public:
 
 	//~ Begin FDataflowNode interface
 	virtual void Evaluate(UE::Dataflow::FContext& Context, const FDataflowOutput* Out) const override;
-	virtual TArray<UE::Dataflow::FRenderingParameter> GetRenderParametersImpl() const override;
 	//~ End FDataflowNode interface
 
 	/** Managed array collection to be used to store datas */
@@ -118,7 +110,6 @@ public:
 	UPROPERTY(EditAnywhere, Category = "Groom Groups")
 	int32 GroupIndex = INDEX_NONE;
 
-
 	/** Linear skinning bone indices key */
 	UPROPERTY(meta = (DataflowOutput))
 	FCollectionAttributeKey BoneIndicesKey;
@@ -130,25 +121,11 @@ public:
 	/** Type of curves to convert skinning for (guides/strands) */
 	UPROPERTY(EditAnywhere, Category = "Vertex Attributes")
 	EGroomCollectionType CurvesType = EGroomCollectionType::Guides;
-
-private:
-	/** Get bone indices key */
-	FCollectionAttributeKey GetBoneIndicesKey() const;
-
-	/** Get bone weights key */
-	FCollectionAttributeKey GetBoneWeightsKey() const;
-
-	/** Get spline parameter key */
-	FCollectionAttributeKey GetSplineParamKey(UE::Dataflow::FContext& Context) const;
-
-	/** Get spline bones key */
-	FCollectionAttributeKey GetSplineBonesKey(UE::Dataflow::FContext& Context) const;
 };
 
-
 /** Convert linear skinning data to spline skinning data */
-USTRUCT(meta = (Experimental, DataflowGroom))
-struct FConvertLinearToSplineSkinWeightsNode : public FDataflowNode
+USTRUCT(meta = (Experimental, DataflowGroom, Deprecated="5.7"))
+struct UE_DEPRECATED(5.7, "Use the newer version of this node instead.") FConvertLinearToSplineSkinWeightsNode : public FDataflowNode
 {
 	GENERATED_BODY()
 
@@ -159,7 +136,6 @@ public:
 
 	//~ Begin FDataflowNode interface
 	virtual void Evaluate(UE::Dataflow::FContext& Context, const FDataflowOutput* Out) const override;
-	virtual TArray<UE::Dataflow::FRenderingParameter> GetRenderParametersImpl() const override;
 	//~ End FDataflowNode interface
 
 	/** Managed array collection to be used to store datas */
@@ -189,9 +165,158 @@ public:
 	/** Type of curves to convert skinning for (guides/strands) */
 	UPROPERTY(EditAnywhere, Category = "Vertex Attributes")
 	EGroomCollectionType CurvesType = EGroomCollectionType::Guides;
+};
+
+namespace UE::Groom::Private
+{
+	struct FCollectionSplineAttributes
+	{
+		static const FName VertexSplineParamAttribute;
+		static const FName VertexSplineBonesAttribute;
+	};
+}
+
+/** Build spline skinning data from skeletal mesh */
+USTRUCT(meta = (Experimental, DataflowGroom))
+struct FBuildSplineSkinWeightsDataflowNode : public FDataflowNode
+{
+	GENERATED_BODY()
+
+	DATAFLOW_NODE_DEFINE_INTERNAL(FBuildSplineSkinWeightsDataflowNode, "BuildSplineSkinWeights", "Groom", "")
+
+public: 
+	FBuildSplineSkinWeightsDataflowNode(const UE::Dataflow::FNodeParameters& InParam, FGuid InGuid = FGuid::NewGuid());
+
+private:
+
+	//~ Begin FDataflowNode interface
+	virtual void Evaluate(UE::Dataflow::FContext& Context, const FDataflowOutput* Out) const override;
+	//~ End FDataflowNode interface
+
+	/** Managed array collection to be used to store datas */
+	UPROPERTY(meta = (DataflowInput, DataflowOutput, DataflowPassthrough  = "Collection", DataflowRenderGroups = "Surface"))
+	FManagedArrayCollection Collection;
+
+	/** Vertex selection to focus the geometry transfer spatially */
+	UPROPERTY(meta = (DataflowInput, DisplayName = "Vertex Selection"))
+	FDataflowVertexSelection VertexSelection;
+	
+	/** SkeletalMesh used to compute spline skinning weights. Will be stored onto the groom asset */
+	UPROPERTY(EditAnywhere, Category="Skeletal Mesh", meta = (DataFlowInput))
+	TObjectPtr<USkeletalMesh> SkeletalMesh = nullptr;
+	
+	/** LOD used to build skinning weights */
+	UPROPERTY(EditAnywhere, Category="Skeletal Mesh", meta = (DataflowInput, DisplayName = "LOD Index"))
+	int32 LODIndex = 0;
+
+	/** Root bones to be used for spline skinning. Uses all bones if empty. Note that branches in the skeleton is currently not supported. */
+	UPROPERTY(EditAnywhere, Category = "Skeletal Mesh", meta = (DataFlowInput))
+	TArray<FString> RootBones;
+
+	/** Number of spline samples per bone segment. */
+	UPROPERTY(EditAnywhere, Category = "Skeletal Mesh", meta = (DataflowInput))
+	int32 SamplesPerSegment = 64;
+
+	/** Spline Skinning Parameter key */
+	UPROPERTY(meta = (DisplayName = "Spline Parameter", DataflowOutput))
+	FCollectionAttributeKey SplineParamKey;
+
+	/** Spline Bones key containing root and end spline bone. To be used in other nodes if necessary */
+	UPROPERTY(meta = (DisplayName = "Spline Bones", DataflowOutput))
+	FCollectionAttributeKey SplineBonesKey;
+
+	/** Get the spline param key */
+	FCollectionAttributeKey GetSplineParamKey(UE::Dataflow::FContext& Context) const;
+
+	/** Get the spline bones key */
+	FCollectionAttributeKey GetSplineBonesKey(UE::Dataflow::FContext& Context) const;
+};
+
+/** Convert spline skinning data to linear data */
+USTRUCT(meta = (Experimental, DataflowGroom))
+struct FSplineToLinearSkinWeightsDataflowNode : public FDataflowNode
+{
+	GENERATED_BODY()
+
+	DATAFLOW_NODE_DEFINE_INTERNAL(FSplineToLinearSkinWeightsDataflowNode, "SplineToLinearSkinWeights", "Groom", "")
+	
+public: 
+	FSplineToLinearSkinWeightsDataflowNode(const UE::Dataflow::FNodeParameters& InParam, FGuid InGuid = FGuid::NewGuid());
 
 private:
-	FCollectionAttributeKey GetSplineParamKey() const;
-	FCollectionAttributeKey GetSplineBonesKey() const;
+	//~ Begin FDataflowNode interface
+	virtual void Evaluate(UE::Dataflow::FContext& Context, const FDataflowOutput* Out) const override;
+	//~ End FDataflowNode interface
+
+	/** Managed array collection to be used to store datas */
+	UPROPERTY(meta = (DataflowInput, DataflowOutput, DataflowPassthrough = "Collection", DataflowRenderGroups = "Surface"))
+	FManagedArrayCollection Collection;
+
+	/** Vertex selection to process vertices subset */
+	UPROPERTY(meta = (DataflowInput, DisplayName = "Vertex Selection"))
+	FDataflowVertexSelection VertexSelection;
+
+	/** Spline skinning parameter key */
+	UPROPERTY(meta = (DataflowInput, DisplayName = "Spline Parameter"))
+	FCollectionAttributeKey SplineParamKey;
+
+	/** Spline bones key. Used for storing root and end spline bone for each vertex. */
+	UPROPERTY(meta = (DataflowInput, DisplayName = "Spline Bones"))
+	FCollectionAttributeKey SplineBonesKey;
+
+	/** Bone indices key to be used in other nodes if necessary */
+	UPROPERTY(meta = (DisplayName = "Bone Indices", DataflowOutput))
+	FCollectionAttributeKey BoneIndicesKey;
+
+	/** Bone weights key to be used in other nodes if necessary */
+	UPROPERTY(meta = (DisplayName = "Bone Weights", DataflowOutput))
+	FCollectionAttributeKey BoneWeightsKey;
+
+	/** Get spline parameter key */
+	FCollectionAttributeKey GetSplineParamKey(UE::Dataflow::FContext& Context) const;
+
+	/** Get spline bones key */
+	FCollectionAttributeKey GetSplineBonesKey(UE::Dataflow::FContext& Context) const;
+};
+
+/** Convert linear skinning data to spline skinning data */
+USTRUCT(meta = (Experimental, DataflowGroom))
+struct FLinearToSplineSkinWeightsDataflowNode : public FDataflowNode
+{
+	GENERATED_BODY()
+	
+	DATAFLOW_NODE_DEFINE_INTERNAL(FLinearToSplineSkinWeightsDataflowNode, "LinearToSplineSkinWeights", "Groom", "")
+
+public: 
+	FLinearToSplineSkinWeightsDataflowNode(const UE::Dataflow::FNodeParameters& InParam, FGuid InGuid = FGuid::NewGuid());
+
+private: 
+	//~ Begin FDataflowNode interface
+	virtual void Evaluate(UE::Dataflow::FContext& Context, const FDataflowOutput* Out) const override;
+	//~ End FDataflowNode interface
+
+	/** Managed array collection to be used to store datas */
+	UPROPERTY(meta = (DataflowInput, DataflowOutput, DataflowPassthrough = "Collection", DataflowRenderGroups = "Surface"))
+	FManagedArrayCollection Collection;
+
+	/** Vertex selection to process vertices subset */
+	UPROPERTY(meta = (DataflowInput, DisplayName = "Vertex Selection"))
+	FDataflowVertexSelection VertexSelection;
+	
+	/** Bone indices key to be used in other nodes if necessary */
+	UPROPERTY(meta = (DisplayName = "Bone Indices", DataflowIndices))
+	FCollectionAttributeKey BoneIndicesKey;
+
+	/** Bone weights key to be used in other nodes if necessary */
+	UPROPERTY(meta = (DisplayName = "Bone Weights", DataflowWeights))
+	FCollectionAttributeKey BoneWeightsKey;
+
+	/**Spline skinning parameter key */
+	UPROPERTY(meta = (DataflowOutput, DisplayName = "Spline Parameter"))
+	FCollectionAttributeKey SplineParamKey;
+
+	/** Spline bones key. Used for storing root and end spline bone for each vertex. */
+	UPROPERTY(meta = (DataflowOutput, DisplayName = "Spline Bones"))
+	FCollectionAttributeKey SplineBonesKey;
 };
 
diff --git a/Engine/Plugins/Runtime/HairStrands/Source/HairStrandsDataflow/Public/BuildGuidesLODsNode.h b/Engine/Plugins/Runtime/HairStrands/Source/HairStrandsDataflow/Public/BuildGuidesLODsNode.h
--- a/Engine/Plugins/Runtime/HairStrands/Source/HairStrandsDataflow/Public/BuildGuidesLODsNode.h
+++ b/Engine/Plugins/Runtime/HairStrands/Source/HairStrandsDataflow/Public/BuildGuidesLODsNode.h
@@ -13,8 +13,8 @@
 
 
 /** Builds the guides LODs */
-USTRUCT(meta = (Experimental, DataflowGroom))
-struct FBuildGuidesLODsDataflowNode : public FDataflowNode
+USTRUCT(meta = (Experimental, DataflowGroom, Deprecated="5.7"))
+struct UE_DEPRECATED(5.7, "Use the newer version of this node instead.") FBuildGuidesLODsDataflowNode : public FDataflowNode
 {
 	GENERATED_BODY()
 	
@@ -49,3 +49,38 @@ public:
 	FCollectionAttributeKey CurveLodsKey;
 };
 
+/** Builds the curves LODs */
+USTRUCT(meta = (Experimental, DataflowGroom))
+struct FBuildCurveLODsDataflowNode : public FDataflowNode
+{
+	GENERATED_BODY()
+	
+	DATAFLOW_NODE_DEFINE_INTERNAL(FBuildCurveLODsDataflowNode, "BuildCurveLODs", "Groom", "")
+
+public:
+	
+	FBuildCurveLODsDataflowNode(const UE::Dataflow::FNodeParameters& InParam, FGuid InGuid = FGuid::NewGuid());
+	
+private:
+
+	//~ Begin FDataflowNode interface
+	virtual void Evaluate(UE::Dataflow::FContext& Context, const FDataflowOutput* Out) const override;
+	//~ End FDataflowNode interface
+
+	/** Managed array collection to be used to store data */
+	UPROPERTY(meta = (DataflowInput, DataflowOutput, DisplayName = "Collection", DataflowPassthrough  = "Collection", DataflowRenderGroups = "Surface"))
+	FManagedArrayCollection Collection;
+
+	/** Curve selection to focus the geometry LODs spatially */
+	UPROPERTY(meta = (DataflowInput, DisplayName = "Curve Selection"))
+	FDataflowCurveSelection CurveSelection;
+	
+	/** Curve parent indices key to be used in other nodes if necessary */
+	UPROPERTY(meta = (DisplayName = "Curve Parents", DataflowOutput))
+	FCollectionAttributeKey CurveParentsKey;
+
+	/** Curve lods indices key to be used in other nodes if necessary */
+	UPROPERTY(meta = (DisplayName = "Curve Lods", DataflowOutput))
+	FCollectionAttributeKey CurveLodsKey;
+};
+
diff --git a/Engine/Plugins/Runtime/HairStrands/Source/HairStrandsDataflow/Public/GenerateGuidesCurvesNode.h b/Engine/Plugins/Runtime/HairStrands/Source/HairStrandsDataflow/Public/GenerateGuidesCurvesNode.h
--- a/Engine/Plugins/Runtime/HairStrands/Source/HairStrandsDataflow/Public/GenerateGuidesCurvesNode.h
+++ b/Engine/Plugins/Runtime/HairStrands/Source/HairStrandsDataflow/Public/GenerateGuidesCurvesNode.h
@@ -10,8 +10,8 @@
 #include "GenerateGuidesCurvesNode.generated.h"
 
 /** Build the guides curves from the strands */
-USTRUCT(meta = (Experimental, DataflowGroom))
-struct FGenerateGuidesCurvesDataflowNode : public FDataflowNode
+USTRUCT(meta = (Experimental, DataflowGroom, Deprecated="5.7"))
+struct UE_DEPRECATED(5.7, "Use the newer version of this node instead.") FGenerateGuidesCurvesDataflowNode : public FDataflowNode
 {
 	GENERATED_BODY()
 	
@@ -40,3 +40,42 @@ public:
 	int32 GuidesCount = 0;
 };
 
+/** Build the guides curves from the strands */
+USTRUCT(meta = (Experimental, DataflowGroom))
+struct FGenerateCurveGeometryDataflowNode : public FDataflowNode
+{
+	GENERATED_BODY()
+	
+	DATAFLOW_NODE_DEFINE_INTERNAL(FGenerateCurveGeometryDataflowNode, "GenerateCurveGeometry", "Groom", "")
+
+public:
+	
+	FGenerateCurveGeometryDataflowNode(const UE::Dataflow::FNodeParameters& InParam, FGuid InGuid = FGuid::NewGuid());
+
+private:
+
+	//~ Begin FDataflowNode interface
+	virtual void Evaluate(UE::Dataflow::FContext& Context, const FDataflowOutput* Out) const override;
+	//~ End FDataflowNode interface
+
+	/** Managed array collection to be used to store datas */
+	UPROPERTY(meta = (DataflowInput, DataflowOutput, DisplayName = "Collection", DataflowPassthrough  = "Collection", DataflowRenderGroups = "Surface"))
+	FManagedArrayCollection Collection;
+
+	/** Managed array collection to be used to store datas */
+	UPROPERTY(meta = (DataflowInput, DisplayName = "Source Curves"))
+	FManagedArrayCollection SourceCurves;
+
+	/** Vertex selection to focus the guides generation spatially */
+	UPROPERTY(meta = (DataflowInput, DisplayName = "Curve Selection"))
+	FDataflowCurveSelection CurveSelection;
+
+	/** Max number of guides */
+	UPROPERTY(EditAnywhere, Category="Geometry", meta = (DataflowInput, DisplayName = "Guides Count"))
+	int32 CurveCount = 0;
+
+	/** Flag to check if we can merge guides or not */
+	UPROPERTY(EditAnywhere, Category="Geometry", meta = (DisplayName = "Merge Curves"))
+	bool bMergeCurves = false;
+};
+
diff --git a/Engine/Plugins/Runtime/HairStrands/Source/HairStrandsDataflow/Public/GetGroomAssetNode.h b/Engine/Plugins/Runtime/HairStrands/Source/HairStrandsDataflow/Public/GetGroomAssetNode.h
--- a/Engine/Plugins/Runtime/HairStrands/Source/HairStrandsDataflow/Public/GetGroomAssetNode.h
+++ b/Engine/Plugins/Runtime/HairStrands/Source/HairStrandsDataflow/Public/GetGroomAssetNode.h
@@ -7,13 +7,22 @@
 #include "ChaosLog.h"
 #include "Dataflow/DataflowCore.h"
 #include "Dataflow/DataflowEngine.h"
-#include "GroomCollectionFacades.h"
 
 #include "GetGroomAssetNode.generated.h"
 
-/** Get the groom asset guides */
-USTRUCT(meta = (Experimental, DataflowGroom))
-struct FGetGroomAssetDataflowNode : public FDataflowNode
+UENUM()
+enum class EGroomCollectionType : uint8
+{
+	/** Strands type (Rendering Mesh) */
+	Strands UMETA(DisplayName = "Strands"),
+
+	/** Guides type (Simulation Mesh)*/
+	Guides UMETA(DisplayName = "Guides"),
+};
+
+/** Get the groom asset */
+USTRUCT(meta = (Experimental, DataflowGroom, Deprecated="5.7"))
+struct UE_DEPRECATED(5.7, "Use the newer version of this node instead.") FGetGroomAssetDataflowNode : public FDataflowNode
 {
 	GENERATED_BODY()
 	
@@ -29,7 +38,6 @@ public:
 
 	//~ Begin FDataflowNode interface
 	virtual void Evaluate(UE::Dataflow::FContext& Context, const FDataflowOutput* Out) const override;
-	virtual TArray<UE::Dataflow::FRenderingParameter> GetRenderParametersImpl() const override;
 	//~ End FDataflowNode interface
 
 	/** Input asset to read the guides from */
@@ -44,3 +52,63 @@ public:
 	UPROPERTY(meta = (DataflowOutput, DisplayName = "Collection"))
 	FManagedArrayCollection Collection;
 };
+
+/** Get the groom asset */
+USTRUCT(meta = (Experimental, DataflowGroom))
+struct FGetGroomAssetDataflowNode_v2 : public FDataflowNode
+{
+	GENERATED_BODY()
+	DATAFLOW_NODE_DEFINE_INTERNAL(FGetGroomAssetDataflowNode_v2, "GetGroomAsset", "Groom", "")
+
+public:
+	
+	FGetGroomAssetDataflowNode_v2(const UE::Dataflow::FNodeParameters& InParam, FGuid InGuid = FGuid::NewGuid())
+		: FDataflowNode(InParam, InGuid)
+	{
+		RegisterOutputConnection(&GroomAsset);
+	}
+
+	//~ Begin FDataflowNode interface
+	virtual void Evaluate(UE::Dataflow::FContext& Context, const FDataflowOutput* Out) const override;
+	//~ End FDataflowNode interface
+
+	/** Input asset to read the guides from */
+	UPROPERTY(EditAnywhere, Category = "Groom", meta = (DataflowOutput, DisplayName = "Groom Asset"))
+	TObjectPtr<const UGroomAsset> GroomAsset = nullptr;
+};
+
+/** Transform a groom asset to a collection */
+USTRUCT(meta = (Experimental, DataflowGroom))
+struct FGroomAssetToCollectionDataflowNode : public FDataflowNode
+{
+	GENERATED_BODY()
+	
+	DATAFLOW_NODE_DEFINE_INTERNAL(FGroomAssetToCollectionDataflowNode, "GroomAssetToCollection", "Groom", "")
+
+public:
+	
+	FGroomAssetToCollectionDataflowNode(const UE::Dataflow::FNodeParameters& InParam, FGuid InGuid = FGuid::NewGuid());
+
+private:
+
+	//~ Begin FDataflowNode interface
+	virtual void Evaluate(UE::Dataflow::FContext& Context, const FDataflowOutput* Out) const override;
+	//~ End FDataflowNode interface
+
+	/** Input asset to read the guides from */
+	UPROPERTY(EditAnywhere, Category = "Groom", meta = (DataflowInput, DisplayName = "Groom Asset"))
+	TObjectPtr<const UGroomAsset> GroomAsset = nullptr;
+
+	/** Type of curves to use to fill the groom collection (guides/strands) */
+	UPROPERTY(EditAnywhere, Category = "Groom", meta = (DisplayName = "Curves Type"), meta = (DataflowInput))
+	EGroomCollectionType CurvesType = EGroomCollectionType::Strands;
+
+	/** Curves thickness for geometry generation */
+	UPROPERTY(EditAnywhere, Category = "Geometry", meta = (DataflowInput, DisplayName = "Curves Thickness"))
+	float CurvesThickness = 0.5;
+
+	/** Managed array collection used to store the guides */
+	UPROPERTY(meta = (DataflowOutput, DisplayName = "Collection", DataflowRenderGroups = "Surface"))
+	FManagedArrayCollection Collection;
+};
+
diff --git a/Engine/Plugins/Runtime/HairStrands/Source/HairStrandsDataflow/Public/GetGroomAttributesNodes.h b/Engine/Plugins/Runtime/HairStrands/Source/HairStrandsDataflow/Public/GetGroomAttributesNodes.h
--- a/Engine/Plugins/Runtime/HairStrands/Source/HairStrandsDataflow/Public/GetGroomAttributesNodes.h
+++ b/Engine/Plugins/Runtime/HairStrands/Source/HairStrandsDataflow/Public/GetGroomAttributesNodes.h
@@ -4,9 +4,9 @@
 
 #include "CoreMinimal.h"
 #include "ChaosLog.h"
-#include "GroomCollectionFacades.h"
 #include "Dataflow/DataflowCore.h"
 #include "Dataflow/DataflowConnectionTypes.h"
+#include "GetGroomAssetNode.h"
 
 #include "GetGroomAttributesNodes.generated.h"
 
@@ -22,8 +22,8 @@ enum class EGroomAttributeType : uint8
 };
 
 /** Get the kinematic weights attributes names */
-USTRUCT(meta = (Experimental, DataflowGroom))
-struct FGetGroomAttributesDataflowNode : public FDataflowNode
+USTRUCT(meta = (Experimental, DataflowGroom, Deprecated="5.7"))
+struct UE_DEPRECATED(5.7, "Use the newer version of this node instead.") FGetGroomAttributesDataflowNode : public FDataflowNode
 {
 	GENERATED_BODY()
 	
@@ -54,3 +54,30 @@ public:
 	EGroomAttributeType AttributeType = EGroomAttributeType::KinematicWeights;
 };
 
+
+/** Get the kinematic weights attributes names */
+USTRUCT(meta = (Experimental, DataflowGroom))
+struct FGetCurveAttributesDataflowNode : public FDataflowNode
+{
+	GENERATED_BODY()
+	
+	DATAFLOW_NODE_DEFINE_INTERNAL(FGetCurveAttributesDataflowNode, "GetCurveAttributes", "Groom", "")
+
+public:
+	
+	FGetCurveAttributesDataflowNode(const UE::Dataflow::FNodeParameters& InParam, FGuid InGuid = FGuid::NewGuid());
+
+private:
+
+	//~ Begin FDataflowNode interface
+	virtual void Evaluate(UE::Dataflow::FContext& Context, const FDataflowOutput* Out) const override;
+	//~ End FDataflowNode interface
+
+	/** Attribute key to build */
+	UPROPERTY(meta = (DisplayName = "Attribute Key", DataflowOutput))
+	FCollectionAttributeKey AttributeKey;
+	
+	/** Type of attribute to use */
+	UPROPERTY(EditAnywhere, Category = "Groom", meta = (DisplayName = "Attribute Type"))
+	EGroomAttributeType AttributeType = EGroomAttributeType::KinematicWeights;
+};
diff --git a/Engine/Plugins/Runtime/HairStrands/Source/HairStrandsDataflow/Public/GroomAssetTerminalNode.h b/Engine/Plugins/Runtime/HairStrands/Source/HairStrandsDataflow/Public/GroomAssetTerminalNode.h
--- a/Engine/Plugins/Runtime/HairStrands/Source/HairStrandsDataflow/Public/GroomAssetTerminalNode.h
+++ b/Engine/Plugins/Runtime/HairStrands/Source/HairStrandsDataflow/Public/GroomAssetTerminalNode.h
@@ -14,10 +14,10 @@
 
 class UDataflow;
 
-USTRUCT(meta = (Experimental, DataflowGroom, DataflowTerminal))
-struct FGroomAssetTerminalDataflowNode : public FDataflowTerminalNode
+USTRUCT(meta = (Experimental, DataflowGroom, DataflowTerminal, Deprecated="5.7"))
+struct UE_DEPRECATED(5.7, "Use the newer version of this node instead.") FGroomAssetTerminalDataflowNode : public FDataflowTerminalNode
 {
-	GENERATED_USTRUCT_BODY()
+	GENERATED_BODY()
 	DATAFLOW_NODE_DEFINE_INTERNAL(FGroomAssetTerminalDataflowNode, "GroomAssetTerminal", "Groom", "")
 
 public:
@@ -27,7 +27,40 @@ public:
 	//~ Begin FDataflowTerminalNode interface
 	virtual void SetAssetValue(TObjectPtr<UObject> Asset, UE::Dataflow::FContext& Context) const override;
 	virtual void Evaluate(UE::Dataflow::FContext& Context, const FDataflowOutput* Out) const override;
+	//~ End FDataflowTerminalNode interface
+
+	UPROPERTY(meta = (DataflowInput, DataflowOutput, DisplayName = "Collection", DataflowPassthrough  = "Collection"))
+	FManagedArrayCollection Collection;
+	
+	/** List of attribute keys that will be used to save matching attributes in the collection. */
+	UPROPERTY()
+	TArray<FCollectionAttributeKey> AttributeKeys;
+
+	UPROPERTY(EditAnyWhere, Category = Attributes, EditFixedSize)
+	TArray<FName> AttributeNames;
 
+private :
+
+	TWeakObjectPtr<UDataflow> DataflowAssetWeakPtr;
+
+	static constexpr int32 NumOtherInputs = 1;
+};
+
+USTRUCT(meta = (Experimental, DataflowGroom, DataflowTerminal))
+struct FGroomAssetTerminalDataflowNode_v2 : public FDataflowTerminalNode
+{
+	GENERATED_BODY()
+	DATAFLOW_NODE_DEFINE_INTERNAL(FGroomAssetTerminalDataflowNode_v2, "GroomAssetTerminal", "Groom", "")
+
+public:
+
+	FGroomAssetTerminalDataflowNode_v2(const UE::Dataflow::FNodeParameters& InParam, FGuid InGuid = FGuid::NewGuid());
+
+private :
+
+	//~ Begin FDataflowTerminalNode interface
+	virtual void SetAssetValue(TObjectPtr<UObject> Asset, UE::Dataflow::FContext& Context) const override;
+	virtual void Evaluate(UE::Dataflow::FContext& Context, const FDataflowOutput* Out) const override;
 	virtual TArray<UE::Dataflow::FPin> AddPins() override;
 	virtual bool CanAddPin() const override { return true; }
 	virtual bool CanRemovePin() const override { return !AttributeKeys.IsEmpty(); }
@@ -36,31 +69,38 @@ public:
 	virtual void PostSerialize(const FArchive& Ar) override;
 	//~ End FDataflowTerminalNode interface
 
-	UPROPERTY(meta = (DataflowInput, DataflowOutput, DisplayName = "Collection", DataflowPassthrough  = "Collection"))
-	FManagedArrayCollection Collection;
-	
+	UPROPERTY(meta = (DataflowInput, DisplayName = "Strands Collection"))
+	FManagedArrayCollection StrandsCollection;
+
+	UPROPERTY(meta = (DataflowInput, DisplayName = "Guides Collection"))
+	FManagedArrayCollection GuidesCollection;
+
 	/** List of attribute keys that will be used to save matching attributes in the collection. */
 	UPROPERTY()
 	TArray<FCollectionAttributeKey> AttributeKeys;
-
+	
+	/** Attributes names used for the keys */
 	UPROPERTY(EditAnyWhere, Category = Attributes, EditFixedSize)
 	TArray<FName> AttributeNames;
-
-private :
+	
 	/** Get the connection reference matching an attribute key index */
 	UE::Dataflow::TConnectionReference<FCollectionAttributeKey> GetConnectionReference(int32 Index) const;
 
+	//~ Begin FDataflowTerminalNode interface
 	virtual bool ShouldInvalidateOnPropertyChanged(const FPropertyChangedEvent& InPropertyChangedEvent) const override;
 	virtual void OnPropertyChanged(UE::Dataflow::FContext& Context, const FPropertyChangedEvent& InPropertyChangedEvent) override;
-
 	virtual bool SupportsDropConnectionOnNode(FName TypeName, UE::Dataflow::FPin::EDirection Direction) const override;
 	virtual const FDataflowConnection* OnDropConnectionOnNode(const FDataflowConnection& DroppedConnection) override;
+	//~ End FDataflowTerminalNode interface
 
+	/** Sync input names with the keys */
 	void SyncInputNames();
+
+	/** Generate a unique input name */
 	FName GenerateUniqueInputName(FName BaseName) const;
 
 	TWeakObjectPtr<UDataflow> DataflowAssetWeakPtr;
 
-	static constexpr int32 NumOtherInputs = 1;
+	static constexpr int32 NumOtherInputs = 2;
 };
 
diff --git a/Engine/Plugins/Runtime/HairStrands/Source/HairStrandsDataflow/Public/ResampleGuidesPointsNode.h b/Engine/Plugins/Runtime/HairStrands/Source/HairStrandsDataflow/Public/ResampleGuidesPointsNode.h
--- a/Engine/Plugins/Runtime/HairStrands/Source/HairStrandsDataflow/Public/ResampleGuidesPointsNode.h
+++ b/Engine/Plugins/Runtime/HairStrands/Source/HairStrandsDataflow/Public/ResampleGuidesPointsNode.h
@@ -11,7 +11,7 @@
 #include "ResampleGuidesPointsNode.generated.h"
 
 /** Enum to pick the number of points per guide */
-UENUM()
+UENUM(BlueprintType)
 enum class EGroomNumPoints : uint8
 {
 	/** Default behavior coming from the strands size in the physics settings */
@@ -34,8 +34,8 @@ enum class EGroomNumPoints : uint8
 };
 
 /** Resample the groom guides with a fixed number of points (in physics strands size) */
-USTRUCT(meta = (Experimental, DataflowGroom))
-struct FResampleGuidesPointsDataflowNode : public FDataflowNode
+USTRUCT(meta = (Experimental, DataflowGroom, Deprecated="5.7"))
+struct UE_DEPRECATED(5.7, "Use the newer version of this node instead.") FResampleGuidesPointsDataflowNode : public FDataflowNode
 {
 	GENERATED_BODY()
 	
@@ -64,3 +64,39 @@ public:
 	EGroomNumPoints PointsCount = EGroomNumPoints::Default;
 };
 
+/** Resample the groom guides with a fixed number of points (in physics strands size) */
+USTRUCT(meta = (Experimental, DataflowGroom))
+struct FResampleCurvePointsDataflowNode : public FDataflowNode
+{
+	GENERATED_BODY()
+	
+	DATAFLOW_NODE_DEFINE_INTERNAL(FResampleCurvePointsDataflowNode, "ResampleCurvePoints", "Groom", "")
+
+public:
+	
+	FResampleCurvePointsDataflowNode(const UE::Dataflow::FNodeParameters& InParam, FGuid InGuid = FGuid::NewGuid());
+
+private:
+
+	//~ Begin FDataflowNode interface
+	virtual void Evaluate(UE::Dataflow::FContext& Context, const FDataflowOutput* Out) const override;
+	//~ End FDataflowNode interface
+
+	/** Managed array collection to be used to store datas */
+	UPROPERTY(meta = (DataflowInput, DataflowOutput, DisplayName = "Collection", DataflowPassthrough  = "Collection", DataflowRenderGroups = "Surface"))
+	FManagedArrayCollection Collection;
+
+	/** Curve selection to focus the guides generation spatially */
+	UPROPERTY(meta = (DataflowInput, DisplayName = "Curve Selection"))
+	FDataflowCurveSelection CurveSelection;
+
+	/** Max number of points */
+	UPROPERTY(EditAnywhere, Category="Geometry", meta = (DisplayName = "Points Count"))
+	EGroomNumPoints PointsCount = EGroomNumPoints::Size16;
+	
+	/** Max number of points (to be able to plug into a variable since enum is not supported in dataflow) */
+	UPROPERTY(meta = (DataflowInput, DisplayName = "Points Count"))
+	int32 NumPoints = 16;
+};
+
+
diff --git a/Engine/Plugins/Runtime/HairStrands/Source/HairStrandsDataflow/Public/SmoothGuidesCurvesNode.h b/Engine/Plugins/Runtime/HairStrands/Source/HairStrandsDataflow/Public/SmoothGuidesCurvesNode.h
--- a/Engine/Plugins/Runtime/HairStrands/Source/HairStrandsDataflow/Public/SmoothGuidesCurvesNode.h
+++ b/Engine/Plugins/Runtime/HairStrands/Source/HairStrandsDataflow/Public/SmoothGuidesCurvesNode.h
@@ -12,8 +12,8 @@
 
 
 /** Smooth the guides for more stable simulation */
-USTRUCT(meta = (Experimental, DataflowGroom))
-struct FSmoothGuidesCurvesDataflowNode : public FDataflowNode
+USTRUCT(meta = (Experimental, DataflowGroom, Deprecated="5.7"))
+struct UE_DEPRECATED(5.7, "Use the newer version of this node instead.") FSmoothGuidesCurvesDataflowNode : public FDataflowNode
 {
 	GENERATED_BODY()
 	
@@ -42,3 +42,34 @@ public:
 	float SmoothingFactor = 0.0f;
 };
 
+/** Smooth the guides for more stable simulation */
+USTRUCT(meta = (Experimental, DataflowGroom))
+struct FSmoothCurvePointsDataflowNode : public FDataflowNode
+{
+	GENERATED_BODY()
+	
+	DATAFLOW_NODE_DEFINE_INTERNAL(FSmoothCurvePointsDataflowNode, "SmoothCurvePoints", "Groom", "")
+
+public:
+	
+	FSmoothCurvePointsDataflowNode(const UE::Dataflow::FNodeParameters& InParam, FGuid InGuid = FGuid::NewGuid());
+
+private:
+
+	//~ Begin FDataflowNode interface
+	virtual void Evaluate(UE::Dataflow::FContext& Context, const FDataflowOutput* Out) const override;
+	//~ End FDataflowNode interface
+
+	/** Managed array collection to be used to store data */
+	UPROPERTY(meta = (DataflowInput, DataflowOutput, DisplayName = "Collection", DataflowPassthrough  = "Collection", DataflowRenderGroups = "Surface"))
+	FManagedArrayCollection Collection;
+
+	/** Curve selection to focus the guides smoothing spatially */
+	UPROPERTY(meta = (DataflowInput, DisplayName = "Curve Selection"))
+	FDataflowCurveSelection CurveSelection;
+
+	/** Smoothing factor between 0 and 1 */
+	UPROPERTY(EditAnywhere, Category="Geometry", meta = (DataflowInput, DisplayName = "Smoothing Factor", ClampMin="0", ClampMax="1", UIMin="0", UIMax="1"))
+	float SmoothingFactor = 0.0f;
+};
+
