Geometry facade for grooms and cards + new rendering + use of curve selection

diff --git a/Engine/Plugins/Runtime/HairStrands/Source/HairStrandsEditor/Private/GroomActions.cpp b/Engine/Plugins/Runtime/HairStrands/Source/HairStrandsEditor/Private/GroomActions.cpp
--- a/Engine/Plugins/Runtime/HairStrands/Source/HairStrandsEditor/Private/GroomActions.cpp
+++ b/Engine/Plugins/Runtime/HairStrands/Source/HairStrandsEditor/Private/GroomActions.cpp
@@ -40,6 +40,7 @@
 #include "HairStrandsFactory.h"
 #include "ToolMenus.h"
 #include "Dataflow/DataflowEditor.h"
+#include "Dialog/SMessageDialog.h"
 
 #include "ThumbnailRendering/SceneThumbnailInfo.h"
 
@@ -537,6 +538,76 @@ void ExecuteCreateStrandsTextures(const FToolMenuContext& InContext)
 	}
 }
 
+	
+////////////////////////////////////////////////////////////////////////////////////////////////////////////////////////
+// Dataflow Asset
+
+bool CanCreateDataflowAsset(const FToolMenuContext& InContext)
+{
+	const UContentBrowserAssetContextMenuContext* CBContext = UContentBrowserAssetContextMenuContext::FindContextWithAssets(InContext);
+	for (const FAssetData& SelectedAsset : CBContext->SelectedAssets)
+	{
+		if (SelectedAsset.IsValid())
+		{
+			return true;
+		}
+	}
+	return false;
+}
+
+void ExecuteCreateDataflowAsset(const FToolMenuContext& InContext)
+{
+	FString GroomDataflowTemplate;
+	auto OnButtonClicked = [&GroomDataflowTemplate](FString TemplatePath)
+		{
+			GroomDataflowTemplate = TemplatePath;
+		};
+
+	TArray<SMessageDialog::FButton> Buttons = UE::Groom::BuildGroomDataflowTemplateButtons(OnButtonClicked);
+	
+	// Select the template
+	TSharedRef<SMessageDialog> SelectTemplateMessageDialog = SNew(SMessageDialog)
+		.Title(FText(LOCTEXT("SelectTemplateTitle", "Select a Groom Dataflow Template")))
+		.Message(LOCTEXT("SelectTemplateMessage", "Select a template for this groom asset:"))
+		.Buttons(Buttons);
+
+	// result will be set in GroomDataflowTemplate
+	SelectTemplateMessageDialog->ShowModal();
+	UDataflow* const Template = GroomDataflowTemplate.IsEmpty() ? nullptr : LoadObject<UDataflow>(nullptr, *GroomDataflowTemplate);
+	
+	const UContentBrowserAssetContextMenuContext* CBContext = UContentBrowserAssetContextMenuContext::FindContextWithAssets(InContext);
+	for (UGroomAsset* GroomAsset : CBContext->LoadSelectedObjects<UGroomAsset>())
+	{
+		if (GroomAsset && GroomAsset->IsValid())
+		{
+			// Create a new Dataflow asset
+			const FString DataflowPath = FPackageName::GetLongPackagePath(GroomAsset->GetOutermost()->GetName());
+			const FString GroomAssetName = GroomAsset->GetName();
+			FString DataflowName = FString(TEXT("DF_")) + (GroomAssetName.StartsWith(TEXT("GA_")) ? GroomAssetName.RightChop(3) : GroomAssetName);
+			FString DataflowPackageName = FPaths::Combine(DataflowPath, DataflowName);
+			if (FindPackage(nullptr, *DataflowPackageName))
+			{
+				// If a Dataflow asset already exists with this name, make a unique name from it to avoid clobbering it
+				DataflowPackageName = MakeUniqueObjectName(nullptr, UPackage::StaticClass(), FName(DataflowPackageName)).ToString();
+				DataflowName = FPaths::GetBaseFilename(DataflowPackageName);
+			}
+			UPackage* const DataflowPackage = CreatePackage(*DataflowPackageName);
+
+			// Load the dataflow template into the groom asset
+			
+			if (UDataflow* const Dataflow = Template ? DuplicateObject(Template, DataflowPackage, FName(DataflowName)) : nullptr)
+			{
+				Dataflow->MarkPackageDirty();
+
+				// Notify the asset registry
+				FAssetRegistryModule::AssetCreated(Dataflow);
+
+				// Set the Dataflow to the groom asset
+				GroomAsset->GetDataflowSettings().SetDataflowAsset(Dataflow);
+			}
+		}
+	}
+}
 
 ////////////////////////////////////////////////////////////////////////////////////////////////////////////////////////
 // Actions registration
@@ -551,7 +622,7 @@ static FDelayedAutoRegisterHelper DelayedAutoRegister(EDelayedRegisterRunPhase::
 		Section.AddDynamicEntry(NAME_None, FNewToolMenuSectionDelegate::CreateLambda([](FToolMenuSection& InSection)
 		{
 			{
-				const TAttribute<FText> Label = LOCTEXT("RebuildGroom", "Rebuild");
+				const TAttribute<FText> Label = LOCTEXT("RebuildGroomAsset", "Rebuild Groom Asset");
 				const TAttribute<FText> ToolTip = LOCTEXT("RebuildGroomTooltip", "Rebuild the groom with new build settings");
 				const FSlateIcon Icon = FSlateIcon(FAppStyle::GetAppStyleSetName(), "ContentBrowser.AssetActions");
 
@@ -562,7 +633,7 @@ static FDelayedAutoRegisterHelper DelayedAutoRegister(EDelayedRegisterRunPhase::
 			}
 
 			{
-				const TAttribute<FText> Label = LOCTEXT("CreateBindingAsset", "Create Binding");
+				const TAttribute<FText> Label = LOCTEXT("CreateBindingAsset", "Create Binding Asset");
 				const TAttribute<FText> ToolTip = LOCTEXT("CreateBindingAssetTooltip", "Create a binding asset between a skeletal mesh and a groom asset");
 				const FSlateIcon Icon = FSlateIcon(FAppStyle::GetAppStyleSetName(), "ContentBrowser.AssetActions");
 
@@ -593,6 +664,17 @@ static FDelayedAutoRegisterHelper DelayedAutoRegister(EDelayedRegisterRunPhase::
 				UIAction.CanExecuteAction = FToolMenuCanExecuteAction::CreateStatic(&CanCreateStrandsTextures);
 				InSection.AddMenuEntry("GroomAsset_CreateStrandsTextures", Label, ToolTip, Icon, UIAction);
 			}
+
+			{
+				const TAttribute<FText> Label = LOCTEXT("CreateDataflowAsset", "Create Dataflow Asset");
+				const TAttribute<FText> ToolTip = LOCTEXT("CreateDataflowAssetooltip", "Create a dataflow asset from a list of templates");
+				const FSlateIcon Icon = FSlateIcon(FAppStyle::GetAppStyleSetName(), "ContentBrowser.AssetActions");
+
+				FToolUIAction UIAction;
+				UIAction.ExecuteAction = FToolMenuExecuteAction::CreateStatic(&ExecuteCreateDataflowAsset);
+				UIAction.CanExecuteAction = FToolMenuCanExecuteAction::CreateStatic(&CanCreateDataflowAsset);
+				InSection.AddMenuEntry("GroomAsset_CreateDataflowAsset", Label, ToolTip, Icon, UIAction);					
+			}
 		}));
 	}));
 });
diff --git a/Engine/Plugins/Runtime/HairStrands/Source/HairStrandsEditor/Private/HairStrandsEditor.cpp b/Engine/Plugins/Runtime/HairStrands/Source/HairStrandsEditor/Private/HairStrandsEditor.cpp
--- a/Engine/Plugins/Runtime/HairStrands/Source/HairStrandsEditor/Private/HairStrandsEditor.cpp
+++ b/Engine/Plugins/Runtime/HairStrands/Source/HairStrandsEditor/Private/HairStrandsEditor.cpp
@@ -158,7 +158,7 @@ void FGroomEditor::StartupModule()
 
 	// default no dataflow template
 	UE::Groom::RegisterGroomDataflowTemplate({ "SelectNoDataflowTemplate", "No Dataflow", "Don't add a Dataflow to this Groom Asset. Useful when a Dataflow already exists.", "", false });
-	UE::Groom::RegisterGroomDataflowTemplatePath(TEXT("/HairStrands/Dataflow"));
+	UE::Groom::RegisterGroomDataflowTemplatePath(TEXT("/HairStrands/Dataflow/Templates"));
 }
 
 void FGroomEditor::ShutdownModule()
diff --git a/Engine/Plugins/Runtime/HairStrands/Source/HairStrandsEditor/Private/HairStrandsFactory.cpp b/Engine/Plugins/Runtime/HairStrands/Source/HairStrandsEditor/Private/HairStrandsFactory.cpp
--- a/Engine/Plugins/Runtime/HairStrands/Source/HairStrandsEditor/Private/HairStrandsFactory.cpp
+++ b/Engine/Plugins/Runtime/HairStrands/Source/HairStrandsEditor/Private/HairStrandsFactory.cpp
@@ -141,6 +141,11 @@ namespace UE::Groom
 		TArray<FGroomDataflowTemplateData>	TemplateData;
 		TArray<FString> DynamicallyLoadedTemplatePaths;
 	};
+	
+	TArray<SMessageDialog::FButton> BuildGroomDataflowTemplateButtons(TFunction<void(FString)> OnButtonClicked)
+	{
+		return FGroomDataflowTemplateManager::GetInstance().BuildDataflowTemplateButtons(OnButtonClicked);
+	}
 
 	void RegisterGroomDataflowTemplate(const FGroomDataflowTemplateData& DataflowTemplate)
 	{
diff --git a/Engine/Plugins/Runtime/HairStrands/Source/HairStrandsEditor/Public/HairStrandsFactory.h b/Engine/Plugins/Runtime/HairStrands/Source/HairStrandsEditor/Public/HairStrandsFactory.h
--- a/Engine/Plugins/Runtime/HairStrands/Source/HairStrandsEditor/Public/HairStrandsFactory.h
+++ b/Engine/Plugins/Runtime/HairStrands/Source/HairStrandsEditor/Public/HairStrandsFactory.h
@@ -2,6 +2,7 @@
 
 #pragma once
 
+#include "Dialog/SMessageDialog.h"
 #include "Factories/Factory.h"
 #include "UObject/ObjectMacros.h"
 
@@ -30,6 +31,9 @@ namespace UE::Groom
 		bool bIsPrimaryTemplate = false;
 	};
 
+	/** Build dataflow template buttons */
+	HAIRSTRANDSEDITOR_API TArray<SMessageDialog::FButton> BuildGroomDataflowTemplateButtons(TFunction<void(FString)> OnButtonClicked);
+
 	/** Register dataflow template to the manager */
 	HAIRSTRANDSEDITOR_API void RegisterGroomDataflowTemplate(const FGroomDataflowTemplateData& DataflowTemplate);
 
diff --git a/Engine/Source/Runtime/Experimental/Chaos/Private/GeometryCollection/Facades/CollectionCurveFacade.cpp b/Engine/Source/Runtime/Experimental/Chaos/Private/GeometryCollection/Facades/CollectionCurveFacade.cpp
--- a/Engine/Source/Runtime/Experimental/Chaos/Private/GeometryCollection/Facades/CollectionCurveFacade.cpp
+++ b/Engine/Source/Runtime/Experimental/Chaos/Private/GeometryCollection/Facades/CollectionCurveFacade.cpp
@@ -3,7 +3,7 @@
 #include "GeometryCollection/Facades/CollectionCurveFacade.h"
 #include "GeometryCollection/GeometryCollection.h"
 
-namespace UE::GeometryCollection::Facades
+namespace GeometryCollection::Facades
 {
 	// Curves facade  Groups
 	const FName FCollectionCurveGeometryFacade::CurvesGroup("Curves");
diff --git a/Engine/Source/Runtime/Experimental/Chaos/Private/GeometryCollection/Facades/CollectionTransformSelectionFacade.cpp b/Engine/Source/Runtime/Experimental/Chaos/Private/GeometryCollection/Facades/CollectionTransformSelectionFacade.cpp
--- a/Engine/Source/Runtime/Experimental/Chaos/Private/GeometryCollection/Facades/CollectionTransformSelectionFacade.cpp
+++ b/Engine/Source/Runtime/Experimental/Chaos/Private/GeometryCollection/Facades/CollectionTransformSelectionFacade.cpp
@@ -1609,7 +1609,7 @@ namespace GeometryCollection::Facades
 	{
 		TArray<int32> OutSelection;
 
-		UE::GeometryCollection::Facades::FCollectionCurveGeometryFacade CurveFacade(ConstCollection);
+		GeometryCollection::Facades::FCollectionCurveGeometryFacade CurveFacade(ConstCollection);
 		if (TransformToGeometryIndexAttribute.IsValid() && CurveFacade.IsValid())
 		{
 			const TManagedArray<int32>& TransformToGeometryIndices = TransformToGeometryIndexAttribute.Get();
@@ -1641,7 +1641,7 @@ namespace GeometryCollection::Facades
 	{
 		TArray<int32> OutSelection;
 
-		UE::GeometryCollection::Facades::FCollectionCurveGeometryFacade CurveFacade(ConstCollection);
+		GeometryCollection::Facades::FCollectionCurveGeometryFacade CurveFacade(ConstCollection);
 		if (CurveFacade.IsValid())
 		{
 			const TArray<int32>& CurveOffsets = CurveFacade.GetGeometryCurveOffsets();
@@ -1667,7 +1667,7 @@ namespace GeometryCollection::Facades
 	{
 		TArray<int32> OutSelection;
 
-		UE::GeometryCollection::Facades::FCollectionCurveGeometryFacade CurveFacade(ConstCollection);
+		GeometryCollection::Facades::FCollectionCurveGeometryFacade CurveFacade(ConstCollection);
 		if (CurveFacade.IsValid())
 		{
 			const int32 NumCurves = CurveFacade.GetNumCurves();
@@ -1713,7 +1713,7 @@ namespace GeometryCollection::Facades
 	{
 		TArray<int32> OutSelection;
 
-		UE::GeometryCollection::Facades::FCollectionCurveGeometryFacade CurveFacade(ConstCollection);
+		GeometryCollection::Facades::FCollectionCurveGeometryFacade CurveFacade(ConstCollection);
 		if (CurveFacade.IsValid())
 		{
 			const int32 NumCurves = CurveFacade.GetNumCurves();
@@ -1759,7 +1759,7 @@ namespace GeometryCollection::Facades
 	{
 		TArray<int32> OutSelection;
 
-		UE::GeometryCollection::Facades::FCollectionCurveGeometryFacade CurveFacade(ConstCollection);
+		GeometryCollection::Facades::FCollectionCurveGeometryFacade CurveFacade(ConstCollection);
 		if (CurveFacade.IsValid())
 		{
 			const int32 NumCurves = CurveFacade.GetNumCurves();
@@ -1785,7 +1785,7 @@ namespace GeometryCollection::Facades
 	{
 		TArray<int32> OutSelection;
 
-		UE::GeometryCollection::Facades::FCollectionCurveGeometryFacade CurveFacade(ConstCollection);
+		GeometryCollection::Facades::FCollectionCurveGeometryFacade CurveFacade(ConstCollection);
 		if (CurveFacade.IsValid())
 		{
 			const int32 NumCurves = CurveFacade.GetNumCurves();
@@ -1811,7 +1811,7 @@ namespace GeometryCollection::Facades
 	{
 		TArray<int32> OutSelection;
 
-		UE::GeometryCollection::Facades::FCollectionCurveGeometryFacade CurveFacade(ConstCollection);
+		GeometryCollection::Facades::FCollectionCurveGeometryFacade CurveFacade(ConstCollection);
 		if (CurveFacade.IsValid())
 		{
 			const TArray<int32>& CurveOffsets = CurveFacade.GetGeometryCurveOffsets();
@@ -1857,7 +1857,7 @@ namespace GeometryCollection::Facades
 	{
 		TArray<int32> OutSelection;
 
-		UE::GeometryCollection::Facades::FCollectionCurveGeometryFacade CurveFacade(ConstCollection);
+		GeometryCollection::Facades::FCollectionCurveGeometryFacade CurveFacade(ConstCollection);
 		if (CurveFacade.IsValid() && TransformIndexAttribute.IsValid())
 		{
 			const TManagedArray<int32>& TransformIndices = TransformIndexAttribute.Get();
diff --git a/Engine/Source/Runtime/Experimental/Chaos/Public/GeometryCollection/Facades/CollectionCurveFacade.h b/Engine/Source/Runtime/Experimental/Chaos/Public/GeometryCollection/Facades/CollectionCurveFacade.h
--- a/Engine/Source/Runtime/Experimental/Chaos/Public/GeometryCollection/Facades/CollectionCurveFacade.h
+++ b/Engine/Source/Runtime/Experimental/Chaos/Public/GeometryCollection/Facades/CollectionCurveFacade.h
@@ -1,4 +1,4 @@
-﻿// Copyright Epic Games, Inc. All Rights Reserved.
+// Copyright Epic Games, Inc. All Rights Reserved.
 
 #pragma once
 
@@ -6,7 +6,7 @@
 #include "GeometryCollection/ManagedArrayCollection.h"
 #include "GeometryCollection/GeometryCollection.h"
 	
-namespace UE::GeometryCollection::Facades
+namespace GeometryCollection::Facades
 {
 	class FCollectionCurveGeometryFacade
 	{
