[Backout] - CL45497446 - backout due to Main CIS issue

diff --git a/Engine/Plugins/Runtime/HairStrands/Source/HairStrandsDataflow/Private/AttachGuidesRootsNode.cpp b/Engine/Plugins/Runtime/HairStrands/Source/HairStrandsDataflow/Private/AttachGuidesRootsNode.cpp
--- a/Engine/Plugins/Runtime/HairStrands/Source/HairStrandsDataflow/Private/AttachGuidesRootsNode.cpp
+++ b/Engine/Plugins/Runtime/HairStrands/Source/HairStrandsDataflow/Private/AttachGuidesRootsNode.cpp
@@ -1,95 +1,60 @@
 // Copyright Epic Games, Inc. All Rights Reserved.
 
 #include "AttachGuidesRootsNode.h"
+#include "GroomCollectionFacades.h"
 #include "Dataflow/DataflowObjectInterface.h"
-#include "GeometryCollection/Facades/CollectionCurveFacade.h"
-#include "GeometryCollection/Facades/CollectionVertexBoneWeightsFacade.h"
 
 #include UE_INLINE_GENERATED_CPP_BY_NAME(AttachGuidesRootsNode)
 
-void FAttachGuidesRootsDataflowNode::Evaluate(UE::Dataflow::FContext& Context, const FDataflowOutput* Out) const
-{}
-
-void FAttachCurveRootsDataflowNode::Evaluate(UE::Dataflow::FContext& Context, const FDataflowOutput* Out) const
+namespace UE::Groom::Private
 {
-	if (Out->IsA<FManagedArrayCollection>(&Collection))
-	{
-		FManagedArrayCollection GeometryCollection = GetValue<FManagedArrayCollection>(Context, &Collection);
-		GeometryCollection::Facades::FCollectionCurveGeometryFacade CurvesFacade(GeometryCollection);
-
-		if(CurvesFacade.IsValid())
-		{
-			GeometryCollection::Facades::FVertexBoneWeightsFacade SkinningFacade(GeometryCollection, false);
-			SkinningFacade.DefineSchema();
-			
-			const bool bValidSelection = CurveSelection.IsValidForCollection(CurvesFacade.GetManagedArrayCollection());
-			
-			const FDataflowCurveSelection& DataflowSelection = GetValue<FDataflowCurveSelection>(Context, &CurveSelection);
-			
-			ParallelFor(CurvesFacade.GetNumCurves(), [&SkinningFacade, &CurvesFacade, DataflowSelection, bValidSelection](int32 CurveIndex)
-			{
-				int32 PointIndex = (CurveIndex == 0) ? 0 : CurvesFacade.GetCurvePointOffsets()[CurveIndex-1];
-				if(!bValidSelection || (bValidSelection && DataflowSelection.IsSelected(CurveIndex)))
-				{
-					SkinningFacade.ModifyKinematicWeight(2*PointIndex,1.0f);
-					SkinningFacade.ModifyKinematicWeight(2*PointIndex+1, 1.0f);
-				}
-					
-				PointIndex += 1;
-				if(!bValidSelection || (bValidSelection && DataflowSelection.IsSelected(CurveIndex)))
-				{
-					SkinningFacade.ModifyKinematicWeight(2*PointIndex, 1.0f);
-					SkinningFacade.ModifyKinematicWeight(2*PointIndex+1, 1.0f);
-				}
-			}, EParallelForFlags::None);
-		}
-		
-		SetValue(Context, MoveTemp(GeometryCollection), &Collection);
-	}
-	else if (Out->IsA<FCollectionAttributeKey>(&KinematicWeightsKey))
+	FCollectionAttributeKey GetKinematicWeightsKey()
 	{
 		FCollectionAttributeKey Key;
-		Key.Group =  FGeometryCollection::VerticesGroup.ToString();
-		Key.Attribute = GeometryCollection::Facades::FVertexBoneWeightsFacade::KinematicWeightAttributeName.ToString();
-		
-		SetValue(Context, Key, &KinematicWeightsKey);
+		Key.Group =  UE::Groom::FGroomGuidesFacade::VerticesGroup.ToString();
+		Key.Attribute = UE::Groom::FGroomGuidesFacade::PointKinematicWeightsAttribute.ToString();
+		return Key;
 	}
 }
 
-void FBuildCurveWeightsDataflowNode::Evaluate(UE::Dataflow::FContext& Context, const FDataflowOutput* Out) const
+void FAttachGuidesRootsDataflowNode::Evaluate(UE::Dataflow::FContext& Context, const FDataflowOutput* Out) const
 {
 	if (Out->IsA<FManagedArrayCollection>(&Collection))
 	{
-		FManagedArrayCollection GeometryCollection = GetValue<FManagedArrayCollection>(Context, &Collection);
-		GeometryCollection::Facades::FCollectionCurveGeometryFacade CurvesFacade(GeometryCollection);
+		FManagedArrayCollection GroomCollection = GetValue<FManagedArrayCollection>(Context, &Collection);
+		
+		UE::Groom::FGroomGuidesFacade GuidesFacade(GroomCollection);
 
-		if(CurvesFacade.IsValid())
+		if(GuidesFacade.IsValid())
 		{
-			const FCollectionAttributeKey WeightKey = GetValue<FCollectionAttributeKey>(Context, &WeightsAttribute);
-			const FDataflowCurveSelection& DataflowSelection = GetValue<FDataflowCurveSelection>(Context, &CurveSelection);
+			TArray<float> KinematicWeights = GuidesFacade.GetPointKinematicWeights();
 			
-			if(TManagedArray<float>* WeightsValue = GeometryCollection.FindAttribute<float>(FName(WeightKey.Attribute), FName(WeightKey.Group)))
+			const int32 LocalIndex = GroupIndex;
+			ParallelFor(GuidesFacade.GetNumPoints(), [&KinematicWeights, &GuidesFacade, LocalIndex](int32 PointIndex)
 			{
-				const bool bValidSelection = CurveSelection.IsValidForCollection(CurvesFacade.GetManagedArrayCollection());;
-				
-				ParallelFor(CurvesFacade.GetNumPoints(), [this, &WeightsValue, &CurvesFacade, &DataflowSelection, bValidSelection](int32 PointIndex)
+				const int32 CurveIndex = GuidesFacade.GetPointCurveIndices()[PointIndex];
+				const int32 ObjectIndex = GuidesFacade.GetCurveObjectIndices()[CurveIndex];
+
+				if(LocalIndex == INDEX_NONE || ObjectIndex == LocalIndex)
 				{
-					const int32 CurveIndex = CurvesFacade.GetPointCurveIndices()[PointIndex];
-					if(!bValidSelection || (bValidSelection && DataflowSelection.IsSelected(CurveIndex)))
-					{
-						const int32 NextPoint = CurvesFacade.GetCurvePointOffsets()[CurveIndex];
-						const int32 PrevPoint = (CurveIndex == 0) ? 0 : CurvesFacade.GetCurvePointOffsets()[CurveIndex-1];
-						
-						const float PointCoord = static_cast<float>(PointIndex - PrevPoint) / (NextPoint - 1 - PrevPoint);
+					const int32 PointOffset = (CurveIndex == 0) ? 0 : GuidesFacade.GetCurvePointOffsets()[CurveIndex-1];
+					const int32 LocalPoint = PointIndex - PointOffset;
 
-						(*WeightsValue)[2 * PointIndex] = CurveWeights.GetRichCurveConst()->Eval(PointCoord);
-						(*WeightsValue)[2 * PointIndex + 1] = (*WeightsValue)[2 * PointIndex];
+					if(LocalPoint < 2)
+					{
+						KinematicWeights[2*PointIndex] = 1.0f;
+						KinematicWeights[2*PointIndex+1] = 1.0f;
 					}
-				}, EParallelForFlags::None);
-			}
+				}
+			}, EParallelForFlags::None);
+			GuidesFacade.SetPointKinematicWeights(KinematicWeights);
 		}
 		
-		SetValue(Context, MoveTemp(GeometryCollection), &Collection);
+		SetValue(Context, MoveTemp(GroomCollection), &Collection);
+	}
+	else if (Out->IsA<FCollectionAttributeKey>(&KinematicWeightsKey))
+	{
+		SetValue(Context, UE::Groom::Private::GetKinematicWeightsKey(), &KinematicWeightsKey);
 	}
 }
 
diff --git a/Engine/Plugins/Runtime/HairStrands/Source/HairStrandsDataflow/Private/BuildGroomSkinningNodes.cpp b/Engine/Plugins/Runtime/HairStrands/Source/HairStrandsDataflow/Private/BuildGroomSkinningNodes.cpp
--- a/Engine/Plugins/Runtime/HairStrands/Source/HairStrandsDataflow/Private/BuildGroomSkinningNodes.cpp
+++ b/Engine/Plugins/Runtime/HairStrands/Source/HairStrandsDataflow/Private/BuildGroomSkinningNodes.cpp
@@ -1,14 +1,13 @@
 // Copyright Epic Games, Inc. All Rights Reserved.
 
 #include "BuildGroomSkinningNodes.h"
+#include "GroomCollectionFacades.h"
 #include "MeshDescriptionToDynamicMesh.h"
 #include "SkeletalMeshAttributes.h"
 #include "SkeletalMeshLODRenderDataToDynamicMesh.h"
 #include "DynamicMesh/DynamicMesh3.h"
 #include "Dataflow/DataflowObjectInterface.h"
 #include "DynamicMesh/MeshTransforms.h"
-#include "GeometryCollection/Facades/CollectionVertexBoneWeightsFacade.h"
-#include "GeometryCollection/Facades/CollectionMeshFacade.h"
 #include "Operations/TransferBoneWeights.h"
 #include "Rendering/SkeletalMeshRenderData.h"
 
@@ -16,7 +15,7 @@
 
 namespace UE::Groom::Private
 {
-	static bool SkeletalMeshToDynamicMesh(const USkeletalMesh* SkeletalMesh, const int32 LodIndex, UE::Geometry::FDynamicMesh3& ToDynamicMesh)
+	FORCEINLINE bool SkeletalMeshToDynamicMesh(const USkeletalMesh* SkeletalMesh, const int32 LodIndex, UE::Geometry::FDynamicMesh3& ToDynamicMesh)
 	{
 		if (SkeletalMesh->HasMeshDescription(LodIndex))
 		{
@@ -55,54 +54,175 @@ namespace UE::Groom::Private
 
 		return true;
 	}
-		
-	static void BuildGeometrySkinWeights(GeometryCollection::Facades::FCollectionMeshFacade& MeshFacade, GeometryCollection::Facades::FVertexBoneWeightsFacade& SkinningFacade,
-		const FDataflowVertexSelection& VertexSelection, const TObjectPtr<USkeletalMesh> SkeletalMesh, const int32 LODIndex, const FTransform& RelativeTransform)
+
+	template<typename FacadeType>
+	static void BuildSkinningData(FManagedArrayCollection& GroomCollection, const int32 GroupIndex, const TObjectPtr<USkeletalMesh> SkeletalMesh, const int32 LODIndex,
+		const FTransform& RelativeTransform)
 	{
-		if(MeshFacade.IsValid() && SkeletalMesh && SkeletalMesh->IsValidLODIndex(LODIndex))
+		FacadeType GroomFacade(GroomCollection);
+
+		if(GroomFacade.IsValid())
 		{
-			const bool bValidSelection = VertexSelection.IsValidForCollection(SkinningFacade.GetManagedArrayCollection());
+			TManagedArray<TObjectPtr<UObject>>& ObjectSkeletalMeshes =
+				GroomCollection.AddAttribute<TObjectPtr<UObject>>(UE::Groom::FGroomGuidesFacade::ObjectSkeletalMeshesAttribute, FacadeType::ObjectsGroup);
+
+			TManagedArray<int32>& ObjectMeshLODs =
+				GroomCollection.AddAttribute<int32>(UE::Groom::FGroomGuidesFacade::ObjectMeshLODsAttribute, FacadeType::ObjectsGroup);
 
-			TMap<FName, FBoneIndexType> TargetBoneToIndex;
-			TargetBoneToIndex.Reserve(SkeletalMesh->GetRefSkeleton().GetRawBoneNum());
-			for (int32 BoneIdx = 0; BoneIdx < SkeletalMesh->GetRefSkeleton().GetRawBoneNum(); ++BoneIdx)
+			TManagedArray<FIntVector4>& VertexBoneIndices =
+				GroomCollection.AddAttribute<FIntVector4>(UE::Groom::FGroomGuidesFacade::PointBoneIndicesAttribute, FacadeType::VerticesGroup);
+
+			TManagedArray<FVector4f>& VertexBoneWeights =
+				GroomCollection.AddAttribute<FVector4f>(UE::Groom::FGroomGuidesFacade::PointBoneWeightsAttribute, FacadeType::VerticesGroup);
+			
+			if(GroupIndex == INDEX_NONE)
 			{
-				TargetBoneToIndex.Add(SkeletalMesh->GetRefSkeleton().GetRawRefBoneInfo()[BoneIdx].Name, BoneIdx);
+				for(int32 LocalIndex = 0; LocalIndex < ObjectSkeletalMeshes.Num(); ++LocalIndex)
+				{
+					ObjectSkeletalMeshes[LocalIndex] = SkeletalMesh;
+					ObjectMeshLODs[LocalIndex] = LODIndex;
+				}
 			}
-			UE::Geometry::FDynamicMesh3 DynamicMesh;
-			if (UE::Groom::Private::SkeletalMeshToDynamicMesh(SkeletalMesh, LODIndex, DynamicMesh))
+			else if(ObjectSkeletalMeshes.IsValidIndex(GroupIndex) && ObjectMeshLODs.IsValidIndex(GroupIndex))
 			{
-				MeshTransforms::ApplyTransform(DynamicMesh, RelativeTransform, true);
+				ObjectSkeletalMeshes[GroupIndex] = SkeletalMesh;
+				ObjectMeshLODs[GroupIndex] = LODIndex;
+			}
+			if(SkeletalMesh && SkeletalMesh->IsValidLODIndex(LODIndex))
+			{
+				TMap<FName, FBoneIndexType> TargetBoneToIndex;
+				TargetBoneToIndex.Reserve(SkeletalMesh->GetRefSkeleton().GetRawBoneNum());
+				for (int32 BoneIdx = 0; BoneIdx < SkeletalMesh->GetRefSkeleton().GetRawBoneNum(); ++BoneIdx)
+				{
+					TargetBoneToIndex.Add(SkeletalMesh->GetRefSkeleton().GetRawRefBoneInfo()[BoneIdx].Name, BoneIdx);
+				}
+				UE::Geometry::FDynamicMesh3 DynamicMesh;
+				if (UE::Groom::Private::SkeletalMeshToDynamicMesh(SkeletalMesh, LODIndex, DynamicMesh))
+				{
+					MeshTransforms::ApplyTransform(DynamicMesh, RelativeTransform, true);
 
-				UE::Geometry::FTransferBoneWeights TransferBoneWeights(&DynamicMesh, FSkeletalMeshAttributes::DefaultSkinWeightProfileName);
-				TransferBoneWeights.bUseParallel = true;
-				TransferBoneWeights.MaxNumInfluences = 4;
-				TransferBoneWeights.TransferMethod = UE::Geometry::FTransferBoneWeights::ETransferBoneWeightsMethod::ClosestPointOnSurface;
+					UE::Geometry::FTransferBoneWeights TransferBoneWeights(&DynamicMesh, FSkeletalMeshAttributes::DefaultSkinWeightProfileName);
+					TransferBoneWeights.bUseParallel = true;
+					TransferBoneWeights.MaxNumInfluences = 4;
+					TransferBoneWeights.TransferMethod = UE::Geometry::FTransferBoneWeights::ETransferBoneWeightsMethod::ClosestPointOnSurface;
 
-				if (TransferBoneWeights.Validate() == UE::Geometry::EOperationValidationResult::Ok)
-				{
-					for(int32 GeometryIndex = 0; GeometryIndex < SkinningFacade.NumGeometry(); ++GeometryIndex)
+					if (TransferBoneWeights.Validate() == UE::Geometry::EOperationValidationResult::Ok)
 					{
-						SkinningFacade.ModifyGeometryBinding(GeometryIndex, SkeletalMesh, LODIndex);
-						
-						const TArrayView<const FVector3f> VertexPositions = MeshFacade.GetVertexPositions(GeometryIndex);
-						const TArray<int32> VertexIndices = MeshFacade.GetVertexIndices(GeometryIndex);
-						
-						const int32 NumVertices = VertexPositions.Num();
-						ParallelFor(NumVertices, [&TransferBoneWeights, &TargetBoneToIndex, &VertexPositions, &VertexIndices, &SkinningFacade,
-							&VertexSelection, bValidSelection](const int32 LocalIndex)
+						const int32 LocalIndex = GroupIndex;
+						ParallelFor(GroomFacade.GetNumPoints(), [&TransferBoneWeights, &TargetBoneToIndex, &GroomFacade, LocalIndex, &VertexBoneIndices, &VertexBoneWeights](
+							int32 PointIndex)
 						{
-							const int32 VertexIndex = VertexIndices[LocalIndex];
+							const int32 CurveIndex = GroomFacade.GetPointCurveIndices()[PointIndex];
+							const int32 ObjectIndex = GroomFacade.GetCurveObjectIndices()[CurveIndex];
 
-							if(!bValidSelection || (bValidSelection && (VertexSelection.IsSelected(VertexIndex))))
+							if(LocalIndex == INDEX_NONE || ObjectIndex == LocalIndex)
 							{
 								TArray<int32> BoneIndices;
 								TArray<float> BoneWeights;
-									
-								TransferBoneWeights.TransferWeightsToPoint(BoneIndices, BoneWeights,VertexPositions[LocalIndex], &TargetBoneToIndex);
-								SkinningFacade.ModifyBoneWeight(VertexIndex, BoneIndices, BoneWeights);
+								TransferBoneWeights.TransferWeightsToPoint(BoneIndices, BoneWeights,
+									GroomFacade.GetPointRestPositions()[PointIndex],&TargetBoneToIndex);
+
+								FIntVector4 PointBoneIndices(INDEX_NONE);
+								FVector4f PointBoneWeights(0.0f);
+
+								for(int32 BoneIdx = 0; BoneIdx < BoneIndices.Num(); ++BoneIdx)
+								{
+									PointBoneIndices[BoneIdx] = BoneIndices[BoneIdx];
+									PointBoneWeights[BoneIdx] = BoneWeights[BoneIdx];
+								}
+								VertexBoneIndices[2*PointIndex] = PointBoneIndices;
+								VertexBoneIndices[2*PointIndex+1] = PointBoneIndices;
+								
+								VertexBoneWeights[2*PointIndex] = PointBoneWeights;
+								VertexBoneWeights[2*PointIndex+1] = PointBoneWeights;
 							}
-						}, TransferBoneWeights.bUseParallel ? EParallelForFlags::None : EParallelForFlags::ForceSingleThread);
+
+						}, TransferBoneWeights.bUseParallel  ? EParallelForFlags::None : EParallelForFlags::ForceSingleThread);
+					}
+				}
+			}
+		}
+	}
+	
+	template<typename FacadeType>
+	FORCEINLINE void ExtractSkinningData(FManagedArrayCollection& GroomCollection, const FCollectionAttributeKey& BoneIndicesKey, const FCollectionAttributeKey& BoneWeightsKey)
+	{
+		if (!BoneIndicesKey.Attribute.IsEmpty() && !BoneWeightsKey.Attribute.IsEmpty())
+		{
+			FacadeType GroomFacade(GroomCollection);
+
+			if(GroomFacade.IsValid())
+			{
+				const TManagedArray<FIntVector4>* BoneIndices = GroomCollection.FindAttributeTyped<FIntVector4>(FName(UE::Groom::FGroomGuidesFacade::PointBoneIndicesAttribute), FacadeType::VerticesGroup);
+				const TManagedArray<FVector4f>* BoneWeights = GroomCollection.FindAttributeTyped<FVector4f>(FName(UE::Groom::FGroomGuidesFacade::PointBoneWeightsAttribute), FacadeType::VerticesGroup);
+				
+				if(BoneIndices && BoneWeights)
+				{
+					TManagedArray<TArray<int32>>& IndicesArray = GroomCollection.AddAttribute<TArray<int32>>(FName(BoneIndicesKey.Attribute), FName(BoneIndicesKey.Group));
+					TManagedArray<TArray<float>>& WeightsArray = GroomCollection.AddAttribute<TArray<float>>(FName(BoneWeightsKey.Attribute), FName(BoneWeightsKey.Group));
+
+					const int32 NumVertices = GroomFacade.GetNumVertices();
+					for(int32 VertexIndex = 0; VertexIndex < NumVertices; ++VertexIndex)
+					{
+						const FIntVector4& PointIndices = (*BoneIndices)[VertexIndex];
+						const FVector4f& PointWeights = (*BoneWeights)[VertexIndex];
+
+						IndicesArray[VertexIndex].Reset(4);
+						WeightsArray[VertexIndex].Reset(4);
+						for(int32 WeightIndex = 0; WeightIndex < 4; ++WeightIndex)
+						{
+							if(PointIndices[WeightIndex] != INDEX_NONE)
+							{
+								IndicesArray[VertexIndex].Add(PointIndices[WeightIndex]);
+								WeightsArray[VertexIndex].Add(PointWeights[WeightIndex]);
+							}
+						}
+					}
+				}
+			}
+		}
+	}
+
+	template<typename FacadeType>
+	FORCEINLINE void ReportSkinningData(FManagedArrayCollection& GroomCollection, const FCollectionAttributeKey& BoneIndicesKey, const FCollectionAttributeKey& BoneWeightsKey)
+	{
+		if (!BoneIndicesKey.Attribute.IsEmpty() && !BoneWeightsKey.Attribute.IsEmpty())
+		{
+			FacadeType GroomFacade(GroomCollection);
+
+			if(GroomFacade.IsValid())
+			{
+				const TManagedArray<TArray<int32>>* IndicesArray = GroomCollection.FindAttributeTyped<TArray<int32>>(FName(BoneIndicesKey.Attribute), FName(BoneIndicesKey.Group));
+				const TManagedArray<TArray<float>>* WeightsArray = GroomCollection.FindAttributeTyped<TArray<float>>(FName(BoneWeightsKey.Attribute), FName(BoneWeightsKey.Group));
+
+				if(IndicesArray && WeightsArray)
+				{
+					TManagedArray<FIntVector4>& BoneIndices = GroomCollection.AddAttribute<FIntVector4>(FName(UE::Groom::FGroomGuidesFacade::PointBoneIndicesAttribute), FacadeType::VerticesGroup);
+					TManagedArray<FVector4f>& BoneWeights = GroomCollection.AddAttribute<FVector4f>(FName(UE::Groom::FGroomGuidesFacade::PointBoneWeightsAttribute), FacadeType::VerticesGroup);
+				
+					const int32 NumVertices = GroomFacade.GetNumVertices();
+					for(int32 VertexIndex = 0; VertexIndex < NumVertices; ++VertexIndex)
+					{
+						FIntVector4 PointIndices(INDEX_NONE);
+						FVector4f PointWeights(0.0f);
+
+						const int32 NumWeights = FMath::Min(4, (*IndicesArray)[VertexIndex].Num());
+						float TotalWeight = 0.0;
+						for(int32 BoneIdx = 0; BoneIdx < NumWeights; ++BoneIdx)
+						{
+							PointIndices[BoneIdx] = (*IndicesArray)[VertexIndex][BoneIdx];
+							PointWeights[BoneIdx] = (*WeightsArray)[VertexIndex][BoneIdx];
+							TotalWeight += PointWeights[BoneIdx];
+						}
+						if(TotalWeight != 0.0)
+						{
+							for(int32 BoneIdx = 0; BoneIdx < NumWeights; ++BoneIdx)
+							{
+								PointWeights[BoneIdx] /= TotalWeight;
+							}
+						}
+						BoneIndices[VertexIndex] = PointIndices;
+						BoneWeights[VertexIndex] = PointWeights;
 					}
 				}
 			}
@@ -122,67 +242,62 @@ FTransferSkinWeightsGroomNode::FTransferSkinWeightsGroomNode(const UE::Dataflow:
 }
 
 void FTransferSkinWeightsGroomNode::Evaluate(UE::Dataflow::FContext& Context, const FDataflowOutput* Out) const
-{}
-
-FTransferGeometrySkinWeightsDataflowNode::FTransferGeometrySkinWeightsDataflowNode(const UE::Dataflow::FNodeParameters& InParam, FGuid InGuid)
-	: FDataflowNode(InParam, InGuid)
-{
-	RegisterInputConnection(&Collection);
-	RegisterInputConnection(&VertexSelection);
-	RegisterInputConnection(&SkeletalMesh);
-	RegisterInputConnection(&LODIndex).SetCanHidePin(true).SetPinIsHidden(true);
-	RegisterInputConnection(&RelativeTransform).SetCanHidePin(true).SetPinIsHidden(true);
-	
-	RegisterOutputConnection(&Collection, &Collection);
-	RegisterOutputConnection(&BoneIndicesKey);
-	RegisterOutputConnection(&BoneWeightsKey);
-}
-
-void FTransferGeometrySkinWeightsDataflowNode::Evaluate(UE::Dataflow::FContext& Context, const FDataflowOutput* Out) const
 {
 	if (Out->IsA(&Collection))
 	{
-		FManagedArrayCollection GeometryCollection = GetValue(Context, &Collection);
-		GeometryCollection::Facades::FCollectionMeshFacade MeshFacade(GeometryCollection);
+		FManagedArrayCollection GroomCollection = GetValue(Context, &Collection);
 
-		if(MeshFacade.IsValid())
+		if(CurvesType == EGroomCollectionType::Guides)
 		{
-			GeometryCollection::Facades::FVertexBoneWeightsFacade SkinningFacade(GeometryCollection, false);
-			SkinningFacade.DefineSchema();
-			
-			UE::Groom::Private::BuildGeometrySkinWeights(MeshFacade, SkinningFacade, GetValue<FDataflowVertexSelection>(Context, &VertexSelection),
-				GetValue<TObjectPtr<USkeletalMesh>>(Context, &SkeletalMesh), GetValue<int32>(Context, &LODIndex),
-				GetValue<FTransform>(Context, &RelativeTransform));
+			UE::Groom::Private::BuildSkinningData<UE::Groom::FGroomGuidesFacade>(GroomCollection, GroupIndex, SkeletalMesh, LODIndex, RelativeTransform);
 		}
-
-		SetValue(Context, MoveTemp(GeometryCollection), &Collection);
+		else
+		{
+			UE::Groom::Private::BuildSkinningData<UE::Groom::FGroomStrandsFacade>(GroomCollection, GroupIndex, SkeletalMesh, LODIndex, RelativeTransform);
+		}
+		SetValue(Context, MoveTemp(GroomCollection), &Collection);
 	}
-	else if (Out->IsA<FCollectionAttributeKey>(&BoneIndicesKey))
+	else if (Out->IsA(&BoneIndicesKey))
 	{
 		SetValue(Context, GetBoneIndicesKey(), &BoneIndicesKey);
 	}
-	else if (Out->IsA<FCollectionAttributeKey>(&BoneWeightsKey))
+	else if (Out->IsA(&BoneWeightsKey))
 	{
 		SetValue(Context, GetBoneWeightsKey(), &BoneWeightsKey);
 	}
+	else if (Out->IsA(&SkeletalMesh))
+	{
+		SafeForwardInput(Context, &SkeletalMesh, &SkeletalMesh);
+	}
 }
 
-FCollectionAttributeKey FTransferGeometrySkinWeightsDataflowNode::GetBoneIndicesKey()
+FCollectionAttributeKey FTransferSkinWeightsGroomNode::GetBoneIndicesKey() const
 {
 	FCollectionAttributeKey Key;
-	Key.Group = FGeometryCollection::VerticesGroup.ToString();
-	Key.Attribute = GeometryCollection::Facades::FVertexBoneWeightsFacade::BoneIndicesAttributeName.ToString();
-	
+	Key.Group = (CurvesType == EGroomCollectionType::Guides) ? UE::Groom::FGroomGuidesFacade::VerticesGroup.ToString() :
+		UE::Groom::FGroomStrandsFacade::VerticesGroup.ToString();
+	Key.Attribute = UE::Groom::FGroomGuidesFacade::PointBoneIndicesAttribute.ToString();
 	return Key;
 }
 
-FCollectionAttributeKey FTransferGeometrySkinWeightsDataflowNode::GetBoneWeightsKey()
+FCollectionAttributeKey FTransferSkinWeightsGroomNode::GetBoneWeightsKey() const
 {
 	FCollectionAttributeKey Key;
-	Key.Group = FGeometryCollection::VerticesGroup.ToString();
-	Key.Attribute = GeometryCollection::Facades::FVertexBoneWeightsFacade::BoneWeightsAttributeName.ToString();
+	Key.Group = (CurvesType == EGroomCollectionType::Guides) ? UE::Groom::FGroomGuidesFacade::VerticesGroup.ToString() :
+		UE::Groom::FGroomStrandsFacade::VerticesGroup.ToString();
+	Key.Attribute = UE::Groom::FGroomGuidesFacade::PointBoneWeightsAttribute.ToString();
 	return Key;
 }
 
-
+TArray<UE::Dataflow::FRenderingParameter> FTransferSkinWeightsGroomNode::GetRenderParametersImpl() const
+{
+	if(CurvesType == EGroomCollectionType::Guides)
+	{
+		return { {TEXT("GuidesRender"), FName("FGroomCollection"), {TEXT("Collection")}}};
+	}
+	else
+	{
+		return { {TEXT("StrandsRender"), FName("FGroomCollection"), {TEXT("Collection")}}};
+	}
+}
 
diff --git a/Engine/Plugins/Runtime/HairStrands/Source/HairStrandsDataflow/Private/BuildGroomSplineSkinningNode.cpp b/Engine/Plugins/Runtime/HairStrands/Source/HairStrandsDataflow/Private/BuildGroomSplineSkinningNode.cpp
--- a/Engine/Plugins/Runtime/HairStrands/Source/HairStrandsDataflow/Private/BuildGroomSplineSkinningNode.cpp
+++ b/Engine/Plugins/Runtime/HairStrands/Source/HairStrandsDataflow/Private/BuildGroomSplineSkinningNode.cpp
@@ -3,52 +3,11 @@
 #include "BuildGroomSplineSkinningNode.h"
 #include "Algo/MaxElement.h"
 #include "Algo/Count.h"
-#include "Dataflow/DataflowSelection.h"
-#include "GeometryCollection/Facades/CollectionMeshFacade.h"
-#include "GeometryCollection/Facades/CollectionVertexBoneWeightsFacade.h"
-#include "GeometryCollection/Facades/CollectionCurveFacade.h"
 
 #include UE_INLINE_GENERATED_CPP_BY_NAME(BuildGroomSplineSkinningNode)
 
 namespace UE::Groom::Private
 {
-	const FName FCollectionSplineAttributes::VertexSplineParamAttribute("SplineParam");
-	const FName FCollectionSplineAttributes::VertexSplineBonesAttribute("SplineBones");
-
-	static FCollectionAttributeKey GetBoneIndicesKey()
-	{
-		FCollectionAttributeKey Key;
-		Key.Group = FGeometryCollection::VerticesGroup.ToString();
-		Key.Attribute = GeometryCollection::Facades::FVertexBoneWeightsFacade::BoneIndicesAttributeName.ToString();
-	
-		return Key;
-	}
-
-	static FCollectionAttributeKey GetBoneWeightsKey()
-	{
-		FCollectionAttributeKey Key;
-		Key.Group = FGeometryCollection::VerticesGroup.ToString();
-		Key.Attribute = GeometryCollection::Facades::FVertexBoneWeightsFacade::BoneWeightsAttributeName.ToString();
-		return Key;
-	}
-
-	static FCollectionAttributeKey GetSplineParamKey()
-	{
-		FCollectionAttributeKey Key;
-		Key.Group = FGeometryCollection::VerticesGroup.ToString();
-		Key.Attribute = UE::Groom::Private::FCollectionSplineAttributes::VertexSplineParamAttribute.ToString();
-	
-		return Key;
-	}
-
-	static FCollectionAttributeKey GetSplineBonesKey()
-	{
-		FCollectionAttributeKey Key;
-		Key.Group = FGeometryCollection::VerticesGroup.ToString();
-		Key.Attribute = UE::Groom::Private::FCollectionSplineAttributes::VertexSplineBonesAttribute.ToString();
-		return Key;
-	}
-	
 	struct FCatmullRomSpline
 	{
 		TArray<FVector> CVs;
@@ -191,289 +150,328 @@ namespace UE::Groom::Private
 		}
 	}
 	
-	static void BuildSplineSkinningWeights(FManagedArrayCollection& GeometryCollection, const GeometryCollection::Facades::FCollectionMeshFacade& MeshFacade,
-		const GeometryCollection::Facades::FCollectionCurveGeometryFacade& CurveFacade, const FDataflowVertexSelection& VertexSelection, 
-		const TArray<UE::Groom::Private::FCatmullRomSpline>& Splines, const FCollectionAttributeKey& SplineParamKey, const FCollectionAttributeKey& SplineBonesKey)
+	template<typename FacadeType>
+	static void BuildSplineSkinningWeights(FManagedArrayCollection& GroomCollection, int GroupIndex, const TObjectPtr<USkeletalMesh> SkeletalMesh, int LODIndex,
+		const TArray<UE::Groom::Private::FCatmullRomSpline>& Splines, FCollectionAttributeKey SplineParamKey, FCollectionAttributeKey SplineBoneKey)
 	{
-		if(MeshFacade.IsValid() && CurveFacade.IsValid() && !SplineParamKey.Attribute.IsEmpty() && !SplineBonesKey.Attribute.IsEmpty())
-		{
-			const bool bValidSelection = VertexSelection.IsValidForCollection(GeometryCollection);
+		FacadeType GroomFacade(GroomCollection);
 
-			const bool ExistingSplineParams = GeometryCollection.HasAttribute(FName(SplineParamKey.Attribute), FName(SplineParamKey.Group)) &&
-				GeometryCollection.HasAttribute(FName(SplineBonesKey.Attribute), FName(SplineBonesKey.Group));
+		if (GroomFacade.IsValid() && (GroupIndex == INDEX_NONE || GroupIndex < GroomFacade.GetNumObjects()))
+		{
+			// Set groom meshes
+			TManagedArray<TObjectPtr<UObject>>& ObjectSkeletalMeshes =
+				GroomCollection.AddAttribute<TObjectPtr<UObject>>(UE::Groom::FGroomGuidesFacade::ObjectSkeletalMeshesAttribute, FacadeType::ObjectsGroup);
 
-			TManagedArray<float>& SplineParams = GeometryCollection.AddAttribute<float>(FName(SplineParamKey.Attribute), FName(SplineParamKey.Group));
-			TManagedArray<FIntVector2>& SplineBones = GeometryCollection.AddAttribute<FIntVector2>(FName(SplineBonesKey.Attribute), FName(SplineBonesKey.Group));
+			TManagedArray<int32>& ObjectMeshLODs =
+				GroomCollection.AddAttribute<int32>(UE::Groom::FGroomGuidesFacade::ObjectMeshLODsAttribute, FacadeType::ObjectsGroup);
 
-			if (!ExistingSplineParams)
+			if (GroupIndex == INDEX_NONE)
 			{
-				for (int32 VertexIndex = 0; VertexIndex < SplineParams.Num(); ++VertexIndex)
+				for (int32 LocalIndex = 0; LocalIndex < ObjectSkeletalMeshes.Num(); ++LocalIndex)
 				{
-					SplineParams[VertexIndex] = 0;
-					SplineBones[VertexIndex] = FIntVector2(INDEX_NONE, INDEX_NONE);
+					ObjectSkeletalMeshes[LocalIndex] = SkeletalMesh;
+					ObjectMeshLODs[LocalIndex] = LODIndex;
 				}
 			}
+			else if (ObjectSkeletalMeshes.IsValidIndex(GroupIndex) && ObjectMeshLODs.IsValidIndex(GroupIndex))
+			{
+				ObjectSkeletalMeshes[GroupIndex] = SkeletalMesh;
+				ObjectMeshLODs[GroupIndex] = LODIndex;
+			}
+
+			// Compute skinning weights
+			TManagedArray<float>& VerticesSplineParams = GroomCollection.AddAttribute<float>(FName(SplineParamKey.Attribute), FName(SplineParamKey.Group));
+			TManagedArray<FIntVector2>& VerticesSplineBones = GroomCollection.AddAttribute<FIntVector2>(FName(SplineBoneKey.Attribute), FName(SplineBoneKey.Group));
+			TManagedArray<float>& VerticesSplineSegments = GroomCollection.AddAttribute<float>(FName("SplineSegment_Build"), FName(SplineParamKey.Group));
+			TManagedArray<float>& VerticesSplineSegmentParams = GroomCollection.AddAttribute<float>(FName("SplineSegmentParam_Build"), FName(SplineParamKey.Group));
+
+			const TArray<FVector3f>& Positions = GroomFacade.GetPointRestPositions();
+			const TArray<int32>& CurvePointOffsets = GroomFacade.GetCurvePointOffsets();
 
-			const int32 NumCurves = CurveFacade.GetNumCurves();
-			ParallelFor(NumCurves, [&Splines, &SplineParams, &SplineBones, &CurveFacade, &VertexSelection, bValidSelection](int32 CurveIndex)
+			const int FirstCurveIndex = GroupIndex > 0 ? GroomFacade.GetObjectCurveOffsets()[GroupIndex - 1] : 0;
+			const int LastCurveIndex = GroupIndex != INDEX_NONE ? GroomFacade.GetObjectCurveOffsets()[GroupIndex] : GroomFacade.GetNumCurves();
+			const int32 NumCurves = LastCurveIndex - FirstCurveIndex;
+
+			ParallelFor(NumCurves, [&FirstCurveIndex, &Positions, &CurvePointOffsets, &Splines, &VerticesSplineParams, &VerticesSplineBones, &VerticesSplineSegments, &VerticesSplineSegmentParams](int32 CurveIndex)
 			{
 				TArray<int32> SplineCounts;
 				SplineCounts.Init(0, Splines.Num());
-
-				const int32 FirstPoint = CurveIndex > 0 ? CurveFacade.GetCurvePointOffsets()[CurveIndex-1] : 0;
-				const int32 LastPoint = CurveFacade.GetCurvePointOffsets()[CurveIndex];
+					
+				const int32 FirstPoint = FirstCurveIndex + CurveIndex > 0 ? CurvePointOffsets[FirstCurveIndex + CurveIndex - 1] : 0;
+				const int32 LastPoint = CurvePointOffsets[FirstCurveIndex + CurveIndex];
 
 				for (int32 PointIndex = FirstPoint; PointIndex < LastPoint; ++PointIndex)
 				{
-					if(!bValidSelection || (bValidSelection && (VertexSelection.IsSelected(2*PointIndex) || VertexSelection.IsSelected(2*PointIndex+1))))
+					const FVector Position(Positions[PointIndex]);
+					int32 ClosestSplineIndex = 0;
+					float MinParam=0, MinDistance = UE_MAX_FLT;
+
+					for (int32 SplineIndex = 0; SplineIndex < Splines.Num(); ++SplineIndex)
 					{
-						const FVector PointPosition = FVector(CurveFacade.GetPointRestPositions()[PointIndex]);
-						int32 ClosestSplineIndex = 0;
-					
-						float MinDistance = UE_MAX_FLT;
-						float MinParam = 0.0f; 
+						float Param, Distance;
+						Splines[SplineIndex].ClosestParam(Position, Param, Distance);
 
-						for (int32 SplineIndex = 0; SplineIndex < Splines.Num(); ++SplineIndex)
+						if (Distance < MinDistance)
 						{
-							float Param, Distance;
-							Splines[SplineIndex].ClosestParam(PointPosition, Param, Distance);
-
-							if (Distance < MinDistance)
-							{
-								MinParam = Param;
-								MinDistance = Distance;
-								ClosestSplineIndex = SplineIndex;
+							MinParam = Param;
+							MinDistance = Distance;
+							ClosestSplineIndex = SplineIndex;
 
-							}
 						}
-				
-						SplineParams[PointIndex * 2] = MinParam;
-						SplineParams[PointIndex * 2 + 1] = MinParam;
-						SplineBones[PointIndex * 2] = Splines[ClosestSplineIndex].Bones;
-						SplineBones[PointIndex * 2 + 1] = Splines[ClosestSplineIndex].Bones;
-
-						++SplineCounts[ClosestSplineIndex];
 					}
-				}
+					
+					VerticesSplineParams[PointIndex * 2] = MinParam;
+					VerticesSplineParams[PointIndex * 2 + 1] = MinParam;
+					VerticesSplineBones[PointIndex * 2] = Splines[ClosestSplineIndex].Bones;
+					VerticesSplineBones[PointIndex * 2 + 1] = Splines[ClosestSplineIndex].Bones;
 
+					++SplineCounts[ClosestSplineIndex];
+				}
+				
 				// Ensure that all CVs on a strand are bound to the same spline
 				if (Algo::Count(SplineCounts, 0) < SplineCounts.Num() - 1)
 				{
 					// Use the most commonly bound spline
-					const int32 SplineIndex = Algo::MaxElement(SplineCounts) - &SplineCounts[0];
+					int32 SplineIndex = Algo::MaxElement(SplineCounts) - &SplineCounts[0];
 
 					for (int32 PointIndex = FirstPoint; PointIndex < LastPoint; ++PointIndex)
 					{
-						if(!bValidSelection || (bValidSelection && (VertexSelection.IsSelected(2*PointIndex) || VertexSelection.IsSelected(2*PointIndex+1))))
+						if (VerticesSplineBones[PointIndex * 2] != Splines[SplineIndex].Bones)
 						{
-							if (SplineBones[PointIndex * 2] != Splines[SplineIndex].Bones)
-							{
-								const FVector PointPosition = FVector(CurveFacade.GetPointRestPositions()[PointIndex]);
-								float MinDistance = UE_MAX_FLT;
-								float MinParam = 0.0f;
-								Splines[SplineIndex].ClosestParam(PointPosition, MinParam, MinDistance);
+							const FVector Position(Positions[PointIndex]);
+							float MinDistance = UE_MAX_FLT;
+							float Param;
+							Splines[SplineIndex].ClosestParam(Position, Param, MinDistance);
 							
-								SplineParams[PointIndex * 2] = MinParam;
-								SplineParams[PointIndex * 2 + 1] = MinParam;
-								SplineBones[PointIndex * 2] = Splines[SplineIndex].Bones;
-								SplineBones[PointIndex * 2 + 1] = Splines[SplineIndex].Bones;
-							}
+							VerticesSplineParams[PointIndex * 2] = Param;
+							VerticesSplineParams[PointIndex * 2 + 1] = Param;
+							VerticesSplineBones[PointIndex * 2] = Splines[SplineIndex].Bones;
+							VerticesSplineBones[PointIndex * 2 + 1] = Splines[SplineIndex].Bones;
 						}
 					}
 				}
 			});
 		}
 	}
-	
-	static void ConvertFromSplineToLinearWeights(FManagedArrayCollection& GeometryCollection, const GeometryCollection::Facades::FCollectionMeshFacade& MeshFacade,
-		GeometryCollection::Facades::FVertexBoneWeightsFacade& SkinningFacade, const FDataflowVertexSelection& VertexSelection, const FCollectionAttributeKey& SplineParamKey, const FCollectionAttributeKey& SplineBonesKey)
+
+	template<typename FacadeType>
+	static void ConvertFromSplineToLinearWeights(FManagedArrayCollection & GroomCollection, int32 GroomGroupIndex, const FCollectionAttributeKey & SplineParamKey, const FCollectionAttributeKey& SplineBonesKey,
+		const FCollectionAttributeKey & BoneIndicesKey, const FCollectionAttributeKey & BoneWeightsKey)
 	{
-		if(MeshFacade.IsValid() && SkinningFacade.IsValid() && !SplineParamKey.Attribute.IsEmpty() && !SplineBonesKey.Attribute.IsEmpty())
+		const TManagedArray<float>* SplineParams = GroomCollection.FindAttributeTyped<float>(FName(SplineParamKey.Attribute), FName(SplineParamKey.Group));
+		const TManagedArray<FIntVector2>* SplineBones = GroomCollection.FindAttributeTyped<FIntVector2>(FName(SplineBonesKey.Attribute), FName(SplineBonesKey.Group));
+
+		if (SplineParams && SplineBones)
 		{
-			const bool bValidSelection = VertexSelection.IsValidForCollection(SkinningFacade.GetManagedArrayCollection());
-			
-			const TManagedArray<float>* SplineParams = GeometryCollection.FindAttributeTyped<float>(FName(SplineParamKey.Attribute), FName(SplineParamKey.Group));
-			const TManagedArray<FIntVector2>* SplineBones = GeometryCollection.FindAttributeTyped<FIntVector2>(FName(SplineBonesKey.Attribute), FName(SplineBonesKey.Group));
+			FacadeType GroomFacade(GroomCollection);
 
-			if (SplineParams && SplineBones)
+			if (GroomFacade.IsValid() && (GroomGroupIndex == INDEX_NONE || GroomGroupIndex < GroomFacade.GetNumObjects()))
 			{
-				TArray<int32> BoneIndices;
-				TArray<float> BoneWeights;
-				
-				for(int32 GeometryIndex = 0; GeometryIndex < SkinningFacade.NumGeometry(); ++GeometryIndex)
+				bool ExistingBoneIndices = !GroomCollection.HasAttribute(FName(BoneIndicesKey.Attribute), FName(BoneIndicesKey.Group));
+				TManagedArray<FIntVector4>& BoneIndices = GroomCollection.AddAttribute<FIntVector4>(FName(BoneIndicesKey.Attribute), FName(BoneIndicesKey.Group));
+				TManagedArray<FVector4f>& BoneWeights = GroomCollection.AddAttribute<FVector4f>(FName(BoneWeightsKey.Attribute), FName(BoneWeightsKey.Group));
+
+				if (!ExistingBoneIndices)
 				{
-					const TArray<int32> VertexIndices = MeshFacade.GetVertexIndices(GeometryIndex);
-					for (const int32& VertexIndex : VertexIndices)
+					for (int32 VertexIndex = 0; VertexIndex < GroomFacade.GetNumPoints() * 2; ++VertexIndex)
 					{
-						if(!bValidSelection || (bValidSelection && (VertexSelection.IsSelected(VertexIndex))))
-						{
-							const float SplineParam = (*SplineParams)[VertexIndex];
-							const int32 SegmentIndex = trunc(SplineParam);
-							const float SegmentParam = SplineParam - SegmentIndex;
-							
-							if (SegmentIndex < (*SplineBones)[VertexIndex][1])
-							{
-								BoneIndices = {SegmentIndex, SegmentIndex + 1};
-								BoneWeights = {1.0f - SegmentParam, SegmentParam};
-							}
-							else
-							{
-								BoneIndices = {SegmentIndex};
-								BoneWeights = {1.0f};
-							}
+						BoneIndices[VertexIndex] = FIntVector4(-1, -1, -1, -1);
+						BoneWeights[VertexIndex] = FVector4f(0, 0, 0, 0);
+					}
+				}
+
+				TArray<int32> GroupIndices;
+				if (GroomGroupIndex == INDEX_NONE)
+				{
+					for (int32 LocalIndex = 0; LocalIndex < GroomFacade.GetNumObjects(); ++LocalIndex)
+					{
+						GroupIndices.Add(LocalIndex);
+					}
+				}
+				else
+				{
+					GroupIndices.Add(GroomGroupIndex);
+				}
+
+				for (int32 GroupIndex : GroupIndices)
+				{
+					const int FirstCurveIndex = GroupIndex > 0 ? GroomFacade.GetObjectCurveOffsets()[GroupIndex - 1] : 0;
+					const int LastCurveIndex = GroomFacade.GetObjectCurveOffsets()[GroupIndex];
+					const int32 FirstVertex = FirstCurveIndex > 0 ? GroomFacade.GetCurvePointOffsets()[FirstCurveIndex - 1] * 2 : 0;
+					const int32 LastVertex = GroomFacade.GetCurvePointOffsets()[LastCurveIndex-1] * 2;
 
-							SkinningFacade.ModifyBoneWeight(VertexIndex, BoneIndices, BoneWeights);
+					for (int32 VertexIndex = FirstVertex; VertexIndex < LastVertex; ++VertexIndex)
+					{
+						float SplineParam = (*SplineParams)[VertexIndex];
+
+						int Segment = trunc(SplineParam);
+						float SegmentParam = SplineParam - Segment;
+
+						if (Segment < (*SplineBones)[VertexIndex][1])
+						{
+							BoneIndices[VertexIndex] = FIntVector4(Segment, Segment + 1, -1 , -1);
+							BoneWeights[VertexIndex] = FVector4f(1.0 - SegmentParam, SegmentParam, 0, 0);
+						}
+						else
+						{
+							BoneIndices[VertexIndex] = FIntVector4(Segment, -1, -1, -1);
+							BoneWeights[VertexIndex] = FVector4f(1, 0, 0, 0);
 						}
 					}
 				}
 			}
 		}
 	}
-	
-	static void ConvertFromLinearToSplineWeights(FManagedArrayCollection& GeometryCollection, GeometryCollection::Facades::FCollectionMeshFacade& MeshFacade,
-		GeometryCollection::Facades::FVertexBoneWeightsFacade& SkinningFacade, const FDataflowVertexSelection& VertexSelection,
-		const FCollectionAttributeKey& SplineParamKey, const FCollectionAttributeKey& SplineBonesKey)
-	{
-		if(MeshFacade.IsValid() && SkinningFacade.IsValid() && !SplineParamKey.Attribute.IsEmpty() && !SplineBonesKey.Attribute.IsEmpty())
-		{
-			const bool bValidSelection = VertexSelection.IsValidForCollection(SkinningFacade.GetManagedArrayCollection());
 
-			const TArray<TObjectPtr<UObject>>& SkeletalMeshs = SkinningFacade.GetSkeletalMeshes().GetConstArray();
-			const TArray<int32>& GeometryLods = SkinningFacade.GetGeometryLODs().GetConstArray();
-			
-			const TArray<TArray<int32>>& BoneIndices = SkinningFacade.GetBoneIndices().GetConstArray();
-			const TArray<TArray<float>>& BoneWeights = SkinningFacade.GetBoneWeights().GetConstArray();
+	template<typename FacadeType>
+	static void ConvertFromLinearToSplineWeights(FManagedArrayCollection & GroomCollection, int32 GroomGroupIndex, const FCollectionAttributeKey & SplineParamKey, const FCollectionAttributeKey & SplineBonesKey)
+	{
+		const TManagedArray<FIntVector4>* BoneIndices = GroomCollection.FindAttributeTyped<FIntVector4>(FName(UE::Groom::FGroomGuidesFacade::PointBoneIndicesAttribute), FacadeType::VerticesGroup);
+		const TManagedArray<FVector4f>* BoneWeights = GroomCollection.FindAttributeTyped<FVector4f>(FName(UE::Groom::FGroomGuidesFacade::PointBoneWeightsAttribute), FacadeType::VerticesGroup);
+		
+		TManagedArray<TObjectPtr<UObject>>* ObjectSkeletalMeshes =
+			GroomCollection.FindAttributeTyped<TObjectPtr<UObject>>(UE::Groom::FGroomGuidesFacade::ObjectSkeletalMeshesAttribute, FacadeType::ObjectsGroup);
+		TManagedArray<int32>* ObjectMeshLODs =
+			GroomCollection.FindAttributeTyped<int32>(UE::Groom::FGroomGuidesFacade::ObjectMeshLODsAttribute, FacadeType::ObjectsGroup);
 
-			const bool ExistingSplineParams = GeometryCollection.HasAttribute(FName(SplineParamKey.Attribute), FName(SplineParamKey.Group)) &&
-				GeometryCollection.HasAttribute(FName(SplineBonesKey.Attribute), FName(SplineBonesKey.Group));
-			
-			TManagedArray<float>& SplineParams = GeometryCollection.AddAttribute<float>(FName(SplineParamKey.Attribute), FName(SplineParamKey.Group));
-			TManagedArray<FIntVector2>& SplineBones = GeometryCollection.AddAttribute<FIntVector2>(FName(SplineBonesKey.Attribute), FName(SplineBonesKey.Group));
+		if (BoneIndices && BoneWeights && ObjectSkeletalMeshes && ObjectMeshLODs && !SplineParamKey.Attribute.IsEmpty() && !SplineBonesKey.Attribute.IsEmpty())
+		{
+			FacadeType GroomFacade(GroomCollection);
 
-			if (!ExistingSplineParams)
+			if (GroomFacade.IsValid() && (GroomGroupIndex == INDEX_NONE || GroomGroupIndex < GroomFacade.GetNumObjects()))
 			{
-				for (int32 VertexIndex = 0; VertexIndex < SplineParams.Num(); ++VertexIndex)
+				TManagedArray<float>& VertexSplineParams = GroomCollection.AddAttribute<float>(FName(SplineParamKey.Attribute), FName(SplineParamKey.Group));
+				TManagedArray<FIntVector2>& VertexSplineBones = GroomCollection.AddAttribute<FIntVector2>(FName(SplineBonesKey.Attribute), FName(SplineBonesKey.Group));
+
+				TArray<int32> GroupIndices;
+				if (GroomGroupIndex == INDEX_NONE)
 				{
-					SplineParams[VertexIndex] = 0;
-					SplineBones[VertexIndex] = FIntVector2(INDEX_NONE, INDEX_NONE);
+					for (int32 LocalIndex = 0; LocalIndex < GroomFacade.GetNumObjects(); ++LocalIndex)
+					{
+						GroupIndices.Add(LocalIndex);
+					}
 				}
-			}
-		
-			for(int32 GeometryIndex = 0; GeometryIndex < SkinningFacade.NumGeometry(); ++GeometryIndex)
-			{
-				if (SkeletalMeshs.IsValidIndex(GeometryIndex) && GeometryLods.IsValidIndex(GeometryIndex))
+				else
 				{
-					if (TObjectPtr<USkeletalMesh> SkeletalMesh = Cast<USkeletalMesh>(SkeletalMeshs[GeometryIndex]))
+					GroupIndices.Add(GroomGroupIndex);
+				}
+
+				for (int32 GroupIndex : GroupIndices)
+				{
+					if (ObjectSkeletalMeshes->IsValidIndex(GroupIndex) && ObjectMeshLODs->IsValidIndex(GroupIndex))
 					{
-						const TArray<int32> VertexIndices = MeshFacade.GetVertexIndices(GeometryIndex);
-					
-						// Find all used bones for this group						
-						TSet<int32> UsedBoneIndices;
-						for (const int32& VertexIndex : VertexIndices)
+						if (TObjectPtr<USkeletalMesh> SkeletalMesh = Cast<USkeletalMesh>((*ObjectSkeletalMeshes)[GroupIndex]))
 						{
-							if(!bValidSelection || (bValidSelection && (VertexSelection.IsSelected(VertexIndex))))
+							const int FirstCurveIndex = GroupIndex > 0 ? GroomFacade.GetObjectCurveOffsets()[GroupIndex - 1] : 0;
+							const int LastCurveIndex = GroomFacade.GetObjectCurveOffsets()[GroupIndex];
+							const int32 FirstVertex = FirstCurveIndex > 0 ? GroomFacade.GetCurvePointOffsets()[FirstCurveIndex - 1] * 2 : 0;
+							const int32 LastVertex = GroomFacade.GetCurvePointOffsets()[LastCurveIndex-1]*2;
+
+							// Find all used bones for this group						
+							TSet<int32> UsedBoneIndices;
+							for (int32 VertexIndex = FirstVertex; VertexIndex < LastVertex; ++VertexIndex)
 							{
-								for (int32 InfluenceIndex = 0; InfluenceIndex < BoneIndices[VertexIndex].Num(); ++InfluenceIndex)
+								for (int32 InfluenceIndex = 0; InfluenceIndex < 4; ++InfluenceIndex)
 								{
-									if (BoneIndices[VertexIndex][InfluenceIndex] != INDEX_NONE)
+									FIntVector4 DebugInfluences = (*BoneIndices)[VertexIndex];
+									if ((*BoneIndices)[VertexIndex][InfluenceIndex] != INDEX_NONE)
 									{
-										UsedBoneIndices.Add(BoneIndices[VertexIndex][InfluenceIndex]);
+										UsedBoneIndices.Add((*BoneIndices)[VertexIndex][InfluenceIndex]);
 									}
 								}
 							}
-						}
-
-						// Find root bones
-						TArray<int32> RootBones;
-						TMap<int32, int32> BoneSplineMap;
-						
-						const FReferenceSkeleton& Skeleton = SkeletalMesh->GetRefSkeleton();
-						for (int32 BoneIndex : UsedBoneIndices)
-						{
-							int32 ParentIndex = Skeleton.GetParentIndex(BoneIndex);
-							if (ParentIndex == INDEX_NONE || !UsedBoneIndices.Contains(ParentIndex))
+							
+							// Find root bones
+							TArray<int32> RootBones;
+							TMap<int32, int32> BoneSplineMap;
+							
+							const FReferenceSkeleton& Skeleton = SkeletalMesh->GetRefSkeleton();
+							for (int32 BoneIndex : UsedBoneIndices)
 							{
-								RootBones.Add(BoneIndex);
+								int32 ParentIndex = Skeleton.GetParentIndex(BoneIndex);
+								if (ParentIndex == INDEX_NONE || !UsedBoneIndices.Contains(ParentIndex))
+								{
+									RootBones.Add(BoneIndex);
+								}
 							}
-						}
-
-						// Find end bones
-						TArray<int32> EndBones;						
-						TArray<int32> Children;
-
-						for (int32 SplineIndex = 0; SplineIndex < RootBones.Num(); ++SplineIndex)
-						{
-							BoneSplineMap.Add(RootBones[SplineIndex], SplineIndex);
-							int32 BoneIndex = RootBones[SplineIndex];
+							
+							// Find end bones
+							TArray<int32> EndBones;						
+							TArray<int32> Children;
 
-							while (Skeleton.GetDirectChildBones(BoneIndex, Children))
+							for (int32 SplineIndex = 0; SplineIndex < RootBones.Num(); ++SplineIndex)
 							{
-								BoneSplineMap.Add(Children[0], SplineIndex);
-								BoneIndex = Children[0];
-								if (Children.Num() > 1)
+								BoneSplineMap.Add(RootBones[SplineIndex], SplineIndex);
+								int32 BoneIndex = RootBones[SplineIndex];
+
+								while (Skeleton.GetDirectChildBones(BoneIndex, Children))
 								{
-									// TODO: Add support for branching splines
+									BoneSplineMap.Add(Children[0], SplineIndex);
+									BoneIndex = Children[0];
+									if (Children.Num() > 1)
+									{
+										// TODO: Add support for branching splines
+									}
 								}
+
+								EndBones.Add(BoneIndex);
 							}
-							EndBones.Add(BoneIndex);
-						}
-						
-						for (const int32& VertexIndex : VertexIndices)
-						{
-							if(!bValidSelection || (bValidSelection && (VertexSelection.IsSelected(VertexIndex))))
+							
+							
+							for (int32 VertexIndex = FirstVertex; VertexIndex < LastVertex; ++VertexIndex)
 							{
-								const TArray<int32>& VertexBones = BoneIndices[VertexIndex];
-								const TArray<float>& VertexWeights = BoneWeights[VertexIndex];
+								const FIntVector4& Influences = (*BoneIndices)[VertexIndex];
+								const FVector4f& Weights = (*BoneWeights)[VertexIndex];
 
 								int32 MaxWeightIndex = INDEX_NONE;
 								float MaxWeight = 0;
-								for (int32 BoneIndex = 0; BoneIndex < BoneIndices[VertexIndex].Num(); ++BoneIndex)
+								for (int32 BoneIndex = 0; BoneIndex < 4; ++BoneIndex)
 								{
-									if (VertexWeights[BoneIndex] > MaxWeight)
+									if (Weights[BoneIndex] > MaxWeight)
 									{
-										MaxWeight = VertexWeights[BoneIndex];
+										MaxWeight = Weights[BoneIndex];
 										MaxWeightIndex = BoneIndex;
 									}
 								}
 								
 								if (MaxWeight > 0)
 								{
-									const int32 SplineIndex = BoneSplineMap[VertexBones[MaxWeightIndex]];
 									float ParentWeight = 0, ChildWeight = 0;
 									
-									for (int32 InfluenceIndex = 0; InfluenceIndex < VertexBones.Num(); ++InfluenceIndex)
+									for (int32 InfluenceIndex = 0; InfluenceIndex < 4; ++InfluenceIndex)
 									{
-										if (VertexBones[MaxWeightIndex] > RootBones[SplineIndex] && VertexBones[InfluenceIndex] == VertexBones[MaxWeightIndex] - 1)
+										if (Influences[MaxWeightIndex] > VertexSplineBones[VertexIndex][0] && Influences[InfluenceIndex] == Influences[MaxWeightIndex]-1)
 										{
-											ParentWeight = VertexWeights[InfluenceIndex];
+											ParentWeight = Weights[InfluenceIndex];
 										}
 
-										if (VertexBones[MaxWeightIndex] < EndBones[SplineIndex] && VertexBones[InfluenceIndex] == VertexBones[MaxWeightIndex] + 1)
+										if (Influences[MaxWeightIndex] < VertexSplineBones[VertexIndex][1] && Influences[InfluenceIndex] == Influences[MaxWeightIndex] + 1)
 										{
-											ChildWeight = VertexWeights[InfluenceIndex];
+											ChildWeight = Weights[InfluenceIndex];
 										}
 									}
 									
 									if (ParentWeight > ChildWeight)
 									{
-										int Segment = VertexBones[MaxWeightIndex] - 1; 
+										int Segment = Influences[MaxWeightIndex] - 1; 
 										float SegmentParam = MaxWeight / (MaxWeight + ParentWeight);
-										SplineParams[VertexIndex] = Segment + SegmentParam;
+										VertexSplineParams[VertexIndex] = Segment + SegmentParam;
 									}
 									else
 									{
-										int Segment = VertexBones[MaxWeightIndex];
+										int Segment = Influences[MaxWeightIndex];
 										float SegmentParam = (1 - MaxWeight) / (MaxWeight + ChildWeight);
-										SplineParams[VertexIndex] = Segment + SegmentParam;
+										VertexSplineParams[VertexIndex] = Segment + SegmentParam;
 									}		
 									
-									SplineBones[VertexIndex] = FIntVector2(RootBones[SplineIndex], EndBones[SplineIndex]);
+									const int32 SplineIndex = BoneSplineMap[Influences[MaxWeightIndex]];
+									VertexSplineBones[VertexIndex] = FIntVector2(RootBones[SplineIndex], EndBones[SplineIndex]);
 								}
 								else if (!RootBones.IsEmpty())
 								{
 									// Rigidly attach to first spline root if there are no weights for this vertex
-									SplineBones[VertexIndex] = FIntVector2(RootBones[0], EndBones[0]);
-									SplineParams[VertexIndex] = 0;
+									VertexSplineBones[VertexIndex] = FIntVector2(RootBones[0], EndBones[0]);
+									VertexSplineParams[VertexIndex] = 0;
 								}
 							}
 						}
@@ -482,7 +480,8 @@ namespace UE::Groom::Private
 			}
 		}
 	}
-}
+}	
+
 
 FBuildGroomSplineSkinWeightsNode::FBuildGroomSplineSkinWeightsNode(const UE::Dataflow::FNodeParameters& InParam, FGuid InGuid)
 	: FDataflowNode(InParam, InGuid)
@@ -496,237 +495,240 @@ FBuildGroomSplineSkinWeightsNode::FBuildGroomSplineSkinWeightsNode(const UE::Dat
 }
 
 void FBuildGroomSplineSkinWeightsNode::Evaluate(UE::Dataflow::FContext& Context, const FDataflowOutput* Out) const
-{
-}
-
-FConvertSplineToLinearSkinWeightsNode::FConvertSplineToLinearSkinWeightsNode(const UE::Dataflow::FNodeParameters& InParam, FGuid InGuid)
-	: FDataflowNode(InParam, InGuid)
-{
-	RegisterInputConnection(&Collection);
-	RegisterInputConnection(&SplineParamKey);
-	RegisterInputConnection(&SplineBonesKey);
-	RegisterOutputConnection(&Collection, &Collection);
-	RegisterOutputConnection(&BoneIndicesKey);
-	RegisterOutputConnection(&BoneWeightsKey);
-}
-
-void FConvertSplineToLinearSkinWeightsNode::Evaluate(UE::Dataflow::FContext& Context, const FDataflowOutput* Out) const
-{
-}
-
-FConvertLinearToSplineSkinWeightsNode::FConvertLinearToSplineSkinWeightsNode(const UE::Dataflow::FNodeParameters& InParam, FGuid InGuid)
-	: FDataflowNode(InParam, InGuid)
-{
-	RegisterInputConnection(&Collection);
-	RegisterOutputConnection(&Collection, &Collection);
-	RegisterOutputConnection(&SplineParamKey);
-	RegisterOutputConnection(&SplineBoneKey);
-}
-
-void FConvertLinearToSplineSkinWeightsNode::Evaluate(UE::Dataflow::FContext& Context, const FDataflowOutput* Out) const
-{
-}
-
-FBuildSplineSkinWeightsDataflowNode::FBuildSplineSkinWeightsDataflowNode(const UE::Dataflow::FNodeParameters& InParam, FGuid InGuid)
-	: FDataflowNode(InParam, InGuid)
-{
-	RegisterInputConnection(&Collection);
-	RegisterInputConnection(&VertexSelection);
-	RegisterInputConnection(&SkeletalMesh);
-	RegisterInputConnection(&LODIndex).SetCanHidePin(true).SetPinIsHidden(true);
-	RegisterInputConnection(&RootBones).SetCanHidePin(true).SetPinIsHidden(true);
-	RegisterInputConnection(&SamplesPerSegment).SetCanHidePin(true).SetPinIsHidden(true);
-
-	RegisterOutputConnection(&Collection, &Collection);
-	RegisterOutputConnection(&SplineParamKey);
-	RegisterOutputConnection(&SplineBonesKey);
-}
-
-void FBuildSplineSkinWeightsDataflowNode::Evaluate(UE::Dataflow::FContext& Context, const FDataflowOutput* Out) const
 {
 	if (Out->IsA<FManagedArrayCollection>(&Collection))
 	{
-		FManagedArrayCollection GeometryCollection = GetValue<FManagedArrayCollection>(Context, &Collection);
-		GeometryCollection::Facades::FCollectionMeshFacade MeshFacade(GeometryCollection);
-		GeometryCollection::Facades::FCollectionCurveGeometryFacade CurveFacade(GeometryCollection);
-		
-		if(MeshFacade.IsValid() && CurveFacade.IsValid())
-		{
-			GeometryCollection::Facades::FVertexBoneWeightsFacade SkinningFacade(GeometryCollection, false);
-			SkinningFacade.DefineSchema();
+		FManagedArrayCollection GroomCollection = GetValue<FManagedArrayCollection>(Context, &Collection);
+		FCollectionAttributeKey ParamKey = GetSplineParamKey();
+		FCollectionAttributeKey BoneKey = GetSplineBonesKey();
 
-			const TObjectPtr<USkeletalMesh> InputSkelMesh = GetValue(Context, &SkeletalMesh);
-			const int32 InputLOD = GetValue(Context, &LODIndex);
+		if (!ParamKey.Attribute.IsEmpty() && !BoneKey.Attribute.IsEmpty() && SkeletalMesh && SkeletalMesh->IsValidLODIndex(LODIndex))
+		{
+			TArray<UE::Groom::Private::FCatmullRomSpline> Splines;
+			UE::Groom::Private::BuildSplines(SkeletalMesh, RootBones, SamplesPerSegment, Splines);
 
-			if(InputSkelMesh && InputSkelMesh->IsValidLODIndex(InputLOD))
+			if (CurvesType == EGroomCollectionType::Guides)
 			{
-				const FDataflowVertexSelection InputSelection = GetValue<FDataflowVertexSelection>(Context, &VertexSelection);
-				const bool bValidSelection = InputSelection.IsValidForCollection(GeometryCollection);
-				
-				TArray<UE::Groom::Private::FCatmullRomSpline> Splines;
-				UE::Groom::Private::BuildSplines(InputSkelMesh,
-					GetValue(Context, &RootBones), GetValue(Context, &SamplesPerSegment), Splines);
-
-				UE::Groom::Private::BuildSplineSkinningWeights(GeometryCollection, MeshFacade, CurveFacade,
-						InputSelection, Splines, GetSplineParamKey(Context), GetSplineBonesKey(Context));
-				
-				for(int32 GeometryIndex = 0; GeometryIndex < SkinningFacade.NumGeometry(); ++GeometryIndex)
-				{
-					bool bValidGeometry = false;
-					for(int32 VertexIndex :  MeshFacade.GetVertexIndices(GeometryIndex))
-					{
-						if(!bValidSelection || (bValidSelection && InputSelection.IsSelected(VertexIndex)))
-						{
-							bValidGeometry = true;
-							break;
-						}
-					}
-					if(bValidGeometry)
-					{
-						SkinningFacade.ModifyGeometryBinding(GeometryIndex, InputSkelMesh, InputLOD);
-					}
-				}
+				UE::Groom::Private::BuildSplineSkinningWeights<UE::Groom::FGroomGuidesFacade>(GroomCollection, GroupIndex, SkeletalMesh, LODIndex, Splines, ParamKey, BoneKey);
 			}
+			else
+			{
+				UE::Groom::Private::BuildSplineSkinningWeights<UE::Groom::FGroomStrandsFacade>(GroomCollection, GroupIndex, SkeletalMesh, LODIndex, Splines, ParamKey, BoneKey);
+			}
+
 		}
 
-		SetValue(Context, MoveTemp(GeometryCollection), &Collection);
+		SetValue(Context, MoveTemp(GroomCollection), &Collection);
 	}
 	else if (Out->IsA<FCollectionAttributeKey>(&SplineParamKey))
 	{
-		SetValue(Context, GetSplineParamKey(Context), &SplineParamKey);
+		SetValue(Context, GetSplineParamKey(), &SplineParamKey);
 	}
-	else if (Out->IsA<FCollectionAttributeKey>(&SplineBonesKey))
+	else if (Out->IsA<FCollectionAttributeKey>(&SplineBoneKey))
 	{
-		SetValue(Context, GetSplineBonesKey(Context), &SplineBonesKey);
+		SetValue(Context, GetSplineBonesKey(), &SplineBoneKey);
 	}
 }
 
-FCollectionAttributeKey FBuildSplineSkinWeightsDataflowNode::GetSplineParamKey(UE::Dataflow::FContext& Context) const
+FCollectionAttributeKey FBuildGroomSplineSkinWeightsNode::GetSplineParamKey() const
 {
 	FCollectionAttributeKey Key;
-	Key.Group = FGeometryCollection::VerticesGroup.ToString();
-	Key.Attribute = UE::Groom::Private::FCollectionSplineAttributes::VertexSplineParamAttribute.ToString();
+	Key.Group = (CurvesType == EGroomCollectionType::Guides) ? UE::Groom::FGroomGuidesFacade::VerticesGroup.ToString() :
+		UE::Groom::FGroomStrandsFacade::VerticesGroup.ToString();
+	Key.Attribute = SplineParamName;
 	return Key;
 }
 
-FCollectionAttributeKey FBuildSplineSkinWeightsDataflowNode::GetSplineBonesKey(UE::Dataflow::FContext& Context) const
+FCollectionAttributeKey FBuildGroomSplineSkinWeightsNode::GetSplineBonesKey() const
 {
 	FCollectionAttributeKey Key;
-	Key.Group = FGeometryCollection::VerticesGroup.ToString();
-	Key.Attribute = UE::Groom::Private::FCollectionSplineAttributes::VertexSplineBonesAttribute.ToString();
+	Key.Group = (CurvesType == EGroomCollectionType::Guides) ? UE::Groom::FGroomGuidesFacade::VerticesGroup.ToString() :
+		UE::Groom::FGroomStrandsFacade::VerticesGroup.ToString();
+	Key.Attribute = SplineBonesParamName;
 	return Key;
 }
 
-FSplineToLinearSkinWeightsDataflowNode::FSplineToLinearSkinWeightsDataflowNode(const UE::Dataflow::FNodeParameters& InParam, FGuid InGuid)
+TArray<UE::Dataflow::FRenderingParameter> FBuildGroomSplineSkinWeightsNode::GetRenderParametersImpl() const
+{
+	if (CurvesType == EGroomCollectionType::Guides)
+	{
+		return { {TEXT("GuidesRender"), FName("FGroomCollection"), {TEXT("Collection")}} };
+	}
+	else
+	{
+		return { {TEXT("StrandsRender"), FName("FGroomCollection"), {TEXT("Collection")}} };
+	}
+}
+
+FConvertSplineToLinearSkinWeightsNode::FConvertSplineToLinearSkinWeightsNode(const UE::Dataflow::FNodeParameters& InParam, FGuid InGuid)
 	: FDataflowNode(InParam, InGuid)
 {
 	RegisterInputConnection(&Collection);
-	RegisterInputConnection(&VertexSelection);
-	RegisterInputConnection(&SplineParamKey).SetCanHidePin(true).SetPinIsHidden(true);
-	RegisterInputConnection(&SplineBonesKey).SetCanHidePin(true).SetPinIsHidden(true);
-	
+	RegisterInputConnection(&SplineParamKey);
+	RegisterInputConnection(&SplineBonesKey);
 	RegisterOutputConnection(&Collection, &Collection);
 	RegisterOutputConnection(&BoneIndicesKey);
 	RegisterOutputConnection(&BoneWeightsKey);
 }
 
-void FSplineToLinearSkinWeightsDataflowNode::Evaluate(UE::Dataflow::FContext& Context, const FDataflowOutput* Out) const
+void FConvertSplineToLinearSkinWeightsNode::Evaluate(UE::Dataflow::FContext& Context, const FDataflowOutput* Out) const
 {
 	if (Out->IsA<FManagedArrayCollection>(&Collection))
 	{
-		FManagedArrayCollection GeometryCollection = GetValue<FManagedArrayCollection>(Context, &Collection);
+		FManagedArrayCollection GroomCollection = GetValue<FManagedArrayCollection>(Context, &Collection);
+		FCollectionAttributeKey SplineParam = GetSplineParamKey(Context);
+		FCollectionAttributeKey SplineBones = GetSplineBonesKey(Context);
 
-		GeometryCollection::Facades::FCollectionMeshFacade MeshFacade(GeometryCollection);
-		
-		if(MeshFacade.IsValid())
+		if (!SplineParam.Attribute.IsEmpty() && !SplineBones.Attribute.IsEmpty())
 		{
-			GeometryCollection::Facades::FVertexBoneWeightsFacade SkinningFacade(GeometryCollection, false);
-			if(!SkinningFacade.IsValid())
+			if (CurvesType == EGroomCollectionType::Guides)
+			{
+				UE::Groom::Private::ConvertFromSplineToLinearWeights<UE::Groom::FGroomGuidesFacade>(GroomCollection, GroupIndex, SplineParam, SplineBones, GetBoneIndicesKey(), GetBoneWeightsKey());
+			}
+			else
 			{
-				SkinningFacade.DefineSchema();
+				UE::Groom::Private::ConvertFromSplineToLinearWeights<UE::Groom::FGroomStrandsFacade>(GroomCollection, GroupIndex, SplineParam, SplineBones, GetBoneIndicesKey(), GetBoneWeightsKey());
 			}
-			UE::Groom::Private::ConvertFromSplineToLinearWeights(GeometryCollection, MeshFacade, SkinningFacade,
-				GetValue<FDataflowVertexSelection>(Context, &VertexSelection), GetSplineParamKey(Context), GetSplineBonesKey(Context));
 		}
-		SetValue(Context, MoveTemp(GeometryCollection), &Collection);
+		SetValue(Context, MoveTemp(GroomCollection), &Collection);
 	}
 	else if (Out->IsA<FCollectionAttributeKey>(&BoneIndicesKey))
 	{
-		SetValue(Context, UE::Groom::Private::GetBoneIndicesKey(), &BoneIndicesKey);
+		SetValue(Context, GetBoneIndicesKey(), &BoneIndicesKey);
 	}
 	else if (Out->IsA<FCollectionAttributeKey>(&BoneWeightsKey))
 	{
-		SetValue(Context, UE::Groom::Private::GetBoneWeightsKey(), &BoneWeightsKey);
+		SetValue(Context, GetBoneWeightsKey(), &BoneWeightsKey);
 	}
 }
 
-FCollectionAttributeKey FSplineToLinearSkinWeightsDataflowNode::GetSplineParamKey(UE::Dataflow::FContext& Context) const
+FCollectionAttributeKey FConvertSplineToLinearSkinWeightsNode::GetBoneIndicesKey() const
+{
+	FCollectionAttributeKey Key;
+	Key.Group = (CurvesType == EGroomCollectionType::Guides) ? UE::Groom::FGroomGuidesFacade::VerticesGroup.ToString() :
+		UE::Groom::FGroomStrandsFacade::VerticesGroup.ToString();
+	Key.Attribute = UE::Groom::FGroomGuidesFacade::PointBoneIndicesAttribute.ToString();
+	return Key;
+}
+
+FCollectionAttributeKey FConvertSplineToLinearSkinWeightsNode::GetBoneWeightsKey() const
+{
+	FCollectionAttributeKey Key;
+	Key.Group = (CurvesType == EGroomCollectionType::Guides) ? UE::Groom::FGroomGuidesFacade::VerticesGroup.ToString() :
+		UE::Groom::FGroomStrandsFacade::VerticesGroup.ToString();
+	Key.Attribute = UE::Groom::FGroomGuidesFacade::PointBoneWeightsAttribute.ToString();
+	return Key;
+}
+
+
+
+FCollectionAttributeKey FConvertSplineToLinearSkinWeightsNode::GetSplineParamKey(UE::Dataflow::FContext& Context) const
 {
 	FCollectionAttributeKey Key = GetValue(Context, &SplineParamKey, SplineParamKey);
-	
+
+	// If nothing set use default value
 	if (Key.Attribute.IsEmpty() && Key.Group.IsEmpty())
 	{
-		Key.Group = FGeometryCollection::VerticesGroup.ToString();
-		Key.Attribute = UE::Groom::Private::FCollectionSplineAttributes::VertexSplineParamAttribute.ToString();
+		Key.Group = (CurvesType == EGroomCollectionType::Guides) ? UE::Groom::FGroomGuidesFacade::VerticesGroup.ToString() :
+			UE::Groom::FGroomStrandsFacade::VerticesGroup.ToString();
+		Key.Attribute = SplineParamName;
 	}
 
 	return Key;
 }
 
-FCollectionAttributeKey FSplineToLinearSkinWeightsDataflowNode::GetSplineBonesKey(UE::Dataflow::FContext& Context) const
+FCollectionAttributeKey FConvertSplineToLinearSkinWeightsNode::GetSplineBonesKey(UE::Dataflow::FContext& Context) const
 {
 	FCollectionAttributeKey Key = GetValue(Context, &SplineBonesKey, SplineBonesKey);
 
 	if (Key.Attribute.IsEmpty() && Key.Group.IsEmpty())
 	{
-		Key.Group = FGeometryCollection::VerticesGroup.ToString();
-		Key.Attribute = UE::Groom::Private::FCollectionSplineAttributes::VertexSplineBonesAttribute.ToString();
+		Key.Group = (CurvesType == EGroomCollectionType::Guides) ? UE::Groom::FGroomGuidesFacade::VerticesGroup.ToString() :
+			UE::Groom::FGroomStrandsFacade::VerticesGroup.ToString();
+		Key.Attribute = SplineBonesName;
 	}
 
 	return Key;
 }
 
-FLinearToSplineSkinWeightsDataflowNode::FLinearToSplineSkinWeightsDataflowNode(const UE::Dataflow::FNodeParameters& InParam, FGuid InGuid)
+TArray<UE::Dataflow::FRenderingParameter> FConvertSplineToLinearSkinWeightsNode::GetRenderParametersImpl() const
+{
+	if (CurvesType == EGroomCollectionType::Guides)
+	{
+		return { {TEXT("GuidesRender"), FName("FGroomCollection"), {TEXT("Collection")}} };
+	}
+	else
+	{
+		return { {TEXT("StrandsRender"), FName("FGroomCollection"), {TEXT("Collection")}} };
+	}
+}
+
+FConvertLinearToSplineSkinWeightsNode::FConvertLinearToSplineSkinWeightsNode(const UE::Dataflow::FNodeParameters& InParam, FGuid InGuid)
 	: FDataflowNode(InParam, InGuid)
 {
 	RegisterInputConnection(&Collection);
-	RegisterInputConnection(&VertexSelection);
-	RegisterInputConnection(&BoneIndicesKey).SetCanHidePin(true).SetPinIsHidden(true);
-	RegisterInputConnection(&BoneWeightsKey).SetCanHidePin(true).SetPinIsHidden(true);
-	
 	RegisterOutputConnection(&Collection, &Collection);
 	RegisterOutputConnection(&SplineParamKey);
-	RegisterOutputConnection(&SplineBonesKey);
+	RegisterOutputConnection(&SplineBoneKey);
 }
 
-void FLinearToSplineSkinWeightsDataflowNode::Evaluate(UE::Dataflow::FContext& Context, const FDataflowOutput* Out) const
+void FConvertLinearToSplineSkinWeightsNode::Evaluate(UE::Dataflow::FContext& Context, const FDataflowOutput* Out) const
 {
 	using namespace UE::Dataflow;
 
 	if (Out->IsA<FManagedArrayCollection>(&Collection))
 	{
-		FManagedArrayCollection GeometryCollection = GetValue<FManagedArrayCollection>(Context, &Collection);
-		
-		GeometryCollection::Facades::FCollectionMeshFacade MeshFacade(GeometryCollection);
-		GeometryCollection::Facades::FVertexBoneWeightsFacade SkinningFacade(GeometryCollection);
-		
-		if(MeshFacade.IsValid() && SkinningFacade.IsValid())
+		FManagedArrayCollection GroomCollection = GetValue<FManagedArrayCollection>(Context, &Collection);
+
+		if (CurvesType == EGroomCollectionType::Guides)
+		{
+			UE::Groom::Private::ConvertFromLinearToSplineWeights<UE::Groom::FGroomGuidesFacade>(GroomCollection, GroupIndex, GetSplineParamKey(), GetSplineBonesKey());
+		}
+		else
 		{
-			UE::Groom::Private::ConvertFromLinearToSplineWeights(GeometryCollection, MeshFacade, SkinningFacade,
-				GetValue<FDataflowVertexSelection>(Context, &VertexSelection), UE::Groom::Private::GetSplineParamKey(), UE::Groom::Private::GetSplineBonesKey());
+			UE::Groom::Private::ConvertFromLinearToSplineWeights<UE::Groom::FGroomStrandsFacade>(GroomCollection, GroupIndex, GetSplineParamKey(), GetSplineBonesKey());
 		}
 
-		SetValue(Context, MoveTemp(GeometryCollection), &Collection);
+		SetValue(Context, MoveTemp(GroomCollection), &Collection);
 	}
 	else if (Out->IsA<FCollectionAttributeKey>(&SplineParamKey))
 	{
-		SetValue(Context, UE::Groom::Private::GetSplineParamKey(), &SplineParamKey);
+		SetValue(Context, GetSplineParamKey(), &SplineParamKey);
 	}
-	else if (Out->IsA<FCollectionAttributeKey>(&SplineBonesKey))
+	else if (Out->IsA<FCollectionAttributeKey>(&SplineBoneKey))
 	{
-		SetValue(Context, UE::Groom::Private::GetSplineBonesKey(), &SplineBonesKey);
+		SetValue(Context, GetSplineBonesKey(), &SplineBoneKey);
 	}
+
 }
 
+TArray<UE::Dataflow::FRenderingParameter> FConvertLinearToSplineSkinWeightsNode::GetRenderParametersImpl() const
+{
+	if (CurvesType == EGroomCollectionType::Guides)
+	{
+		return { {TEXT("GuidesRender"), FName("FGroomCollection"), {TEXT("Collection")}} };
+	}
+	else
+	{
+		return { {TEXT("StrandsRender"), FName("FGroomCollection"), {TEXT("Collection")}} };
+	}
+}
 
+FCollectionAttributeKey FConvertLinearToSplineSkinWeightsNode::GetSplineParamKey() const
+{
+	FCollectionAttributeKey Key;
+	Key.Group = (CurvesType == EGroomCollectionType::Guides) ? UE::Groom::FGroomGuidesFacade::VerticesGroup.ToString() :
+		UE::Groom::FGroomStrandsFacade::VerticesGroup.ToString();
+	Key.Attribute = SplineParamName;
+	return Key;
+}
+
+FCollectionAttributeKey FConvertLinearToSplineSkinWeightsNode::GetSplineBonesKey() const
+{
+	FCollectionAttributeKey Key;
+	Key.Group = (CurvesType == EGroomCollectionType::Guides) ? UE::Groom::FGroomGuidesFacade::VerticesGroup.ToString() :
+		UE::Groom::FGroomStrandsFacade::VerticesGroup.ToString();
+	Key.Attribute = SplineBonesName;
+	return Key;
+}
 
diff --git a/Engine/Plugins/Runtime/HairStrands/Source/HairStrandsDataflow/Private/BuildGuidesLODsNode.cpp b/Engine/Plugins/Runtime/HairStrands/Source/HairStrandsDataflow/Private/BuildGuidesLODsNode.cpp
--- a/Engine/Plugins/Runtime/HairStrands/Source/HairStrandsDataflow/Private/BuildGuidesLODsNode.cpp
+++ b/Engine/Plugins/Runtime/HairStrands/Source/HairStrandsDataflow/Private/BuildGuidesLODsNode.cpp
@@ -1,21 +1,21 @@
 // Copyright Epic Games, Inc. All Rights Reserved.
 
 #include "BuildGuidesLODsNode.h"
+#include "GroomCollectionFacades.h"
 #include "Dataflow/DataflowObjectInterface.h"
-#include "GeometryCollection/Facades/CollectionCurveFacade.h"
 
 #include UE_INLINE_GENERATED_CPP_BY_NAME(BuildGuidesLODsNode)
 
 namespace UE::Groom::Private
 {
-	FORCEINLINE float ComputeGuidesMetric(const TArray<FVector3f>& PointRestPositions, const TArray<int32>& CurvePointOffsets,const uint32 GuideIndexA, const uint32 GuideIndexB, const float& GuideLengthA, const float& GuideLengthB,
+	FORCEINLINE float ComputeGuidesMetric(const FGroomGuidesFacade& GuidesFacade, const uint32 GuideIndexA, const uint32 GuideIndexB, const float& GuideLengthA, const float& GuideLengthB,
 										 const float ShapeWeight, const float ProximityWeight)
 	{
-		const uint32 PointOffsetA = (GuideIndexA > 0) ? CurvePointOffsets[GuideIndexA-1] : 0;
-		const uint32 PointOffsetB = (GuideIndexB > 0) ? CurvePointOffsets[GuideIndexB-1] : 0;
+		const uint32 PointOffsetA = (GuideIndexA > 0) ? GuidesFacade.GetCurvePointOffsets()[GuideIndexA-1] : 0;
+		const uint32 PointOffsetB = (GuideIndexB > 0) ? GuidesFacade.GetCurvePointOffsets()[GuideIndexB-1] : 0;
 		
-		const uint32 NumPointsA = CurvePointOffsets[GuideIndexA] - PointOffsetA;
-		const uint32 NumPointsB = CurvePointOffsets[GuideIndexB] - PointOffsetB;
+		const uint32 NumPointsA = GuidesFacade.GetCurvePointOffsets()[GuideIndexA] - PointOffsetA;
+		const uint32 NumPointsB = GuidesFacade.GetCurvePointOffsets()[GuideIndexB] - PointOffsetB;
 
 		if(NumPointsA == NumPointsB)
 		{
@@ -23,11 +23,11 @@ namespace UE::Groom::Private
 			float ShapeMetric = 0.0f;
 			for(uint32 PointIndex = 0; PointIndex < NumPointsA; ++PointIndex)
 			{
-				const FVector3f& GuidePositionA = PointRestPositions[PointIndex+PointOffsetA];
-				const FVector3f& GuidePositionB = PointRestPositions[PointIndex+PointOffsetB];
+				const FVector3f& GuidePositionA = GuidesFacade.GetPointRestPositions()[PointIndex+PointOffsetA];
+				const FVector3f& GuidePositionB = GuidesFacade.GetPointRestPositions()[PointIndex+PointOffsetB];
 				
 				ProximityMetric += (GuidePositionB - GuidePositionA).Size();
-				ShapeMetric += (GuidePositionB - PointRestPositions[PointOffsetB] - GuidePositionA + PointRestPositions[PointOffsetA]).Size();
+				ShapeMetric += (GuidePositionB - GuidesFacade.GetPointRestPositions()[PointOffsetB] - GuidePositionA + GuidesFacade.GetPointRestPositions()[PointOffsetA]).Size();
 			}
 
 			const float MetricScale = 1.0f / (NumPointsA * 0.5f * (GuideLengthA+GuideLengthB));
@@ -35,124 +35,89 @@ namespace UE::Groom::Private
 		}
 		return 0.0f;
 	}
-	
-	FORCEINLINE void BuildInternalLODs(const int32 NumObjects, const int32 NumCurves, const TArray<int32>& GeometryCurveOffsets,
-		const TArray<int32>& CurvePointOffsets, const TArray<FVector3f>& PointRestPositions, TArray<int32>& CurveParentIndices, TArray<int32>& CurveLodIndices,
-			const FDataflowCurveSelection& CurveSelection, const bool bValidSelection)
+
+	FORCEINLINE void BuildGuidesLODs(FGroomGuidesFacade& GuidesFacade)
 	{
-		CurveParentIndices.Init(INDEX_NONE, NumCurves);
-		CurveLodIndices.Init(INDEX_NONE, NumCurves);
+		const uint32 NumObjects = GuidesFacade.GetNumObjects();
+		
+		TArray<int32> CurveParentIndices; CurveParentIndices.Init(INDEX_NONE, GuidesFacade.GetNumCurves());
+		TArray<int32> CurveLodIndices; CurveLodIndices.Init(INDEX_NONE, GuidesFacade.GetNumCurves());
 		
 		uint32 CurveOffset = 0, PointOffset = 0;
-		for(int32 ObjectIndex = 0; ObjectIndex < NumObjects; ++ObjectIndex)
+		for(uint32 ObjectIndex = 0; ObjectIndex < NumObjects; ++ObjectIndex)
 		{ 
-			TArray<float> GuidesLengths; GuidesLengths.Init(0.0f, GeometryCurveOffsets[ObjectIndex] - CurveOffset);
-			for(int32 GuideIndex = CurveOffset; GuideIndex < GeometryCurveOffsets[ObjectIndex]; ++GuideIndex)
+			TArray<float> GuidesLengths; GuidesLengths.Init(0.0f, GuidesFacade.GetObjectCurveOffsets()[ObjectIndex] - CurveOffset);
+			for(int32 GuideIndex = CurveOffset; GuideIndex < GuidesFacade.GetObjectCurveOffsets()[ObjectIndex]; ++GuideIndex)
 			{
-				for(int32 PointIndex = PointOffset+1; PointIndex < CurvePointOffsets[GuideIndex]; ++PointIndex)
+				for(int32 PointIndex = PointOffset+1; PointIndex < GuidesFacade.GetCurvePointOffsets()[GuideIndex]; ++PointIndex)
 				{
-					GuidesLengths[GuideIndex-CurveOffset] += (PointRestPositions[PointIndex]-PointRestPositions[PointIndex-1]).Size();
+					GuidesLengths[GuideIndex-CurveOffset] += (GuidesFacade.GetPointRestPositions()[PointIndex]-GuidesFacade.GetPointRestPositions()[PointIndex-1]).Size();
 				}
-				PointOffset = CurvePointOffsets[GuideIndex];
+				PointOffset = GuidesFacade.GetCurvePointOffsets()[GuideIndex];
 			}
 			const int32 NumObjectLODs = FMath::CeilLogTwo(GuidesLengths.Num());
-			ParallelFor(GuidesLengths.Num(), [CurveOffset, &GuidesLengths, &PointRestPositions, &CurvePointOffsets, &CurveParentIndices, &CurveLodIndices, &NumObjectLODs, bValidSelection, &CurveSelection](int32 CurveIndex)
+			ParallelFor(GuidesLengths.Num(), [CurveOffset, &GuidesLengths, &GuidesFacade, &CurveParentIndices, &CurveLodIndices, &NumObjectLODs](int32 CurveIndex)
 			{
 				const uint32 GuideIndex = CurveIndex + CurveOffset;
+				const uint32 GuideLod = NumObjectLODs - 1 - FMath::FloorLog2(CurveIndex);
+				const uint32 LodOffset = (GuideLod == NumObjectLODs - 1) ? 0 : 1 << FMath::FloorLog2(CurveIndex);
 
-				if(!bValidSelection || (bValidSelection && (CurveSelection.IsSelected(GuideIndex))))
+				float MinMetric = FLT_MAX;
+				int32 MinParent = INDEX_NONE;
+				for(uint32 ParentIndex = CurveOffset; ParentIndex < CurveOffset+LodOffset; ++ParentIndex)
 				{
-					const uint32 GuideLod = NumObjectLODs - 1 - FMath::FloorLog2(CurveIndex);
-					const uint32 LodOffset = (GuideLod == NumObjectLODs - 1) ? 0 : 1 << FMath::FloorLog2(CurveIndex);
-
-					float MinMetric = FLT_MAX;
-					int32 MinParent = INDEX_NONE;
-					for(uint32 ParentIndex = CurveOffset; ParentIndex < CurveOffset+LodOffset; ++ParentIndex)
+					const float ParentMetric = 1.0f - ComputeGuidesMetric(GuidesFacade,
+						GuideIndex, ParentIndex, GuidesLengths[CurveIndex], GuidesLengths[ParentIndex-CurveOffset],  1.0f, 1.0f);
+					if (ParentMetric < MinMetric)
 					{
-						const float ParentMetric = 1.0f - ComputeGuidesMetric(PointRestPositions, CurvePointOffsets,
-							GuideIndex, ParentIndex, GuidesLengths[CurveIndex], GuidesLengths[ParentIndex-CurveOffset],  1.0f, 1.0f);
-						if (ParentMetric < MinMetric)
-						{
-							MinMetric = ParentMetric;
-							MinParent = ParentIndex;
-						}
+						MinMetric = ParentMetric;
+						MinParent = ParentIndex;
 					}
-					CurveParentIndices[GuideIndex] = MinParent;
-					CurveLodIndices[GuideIndex] = GuideLod;
 				}
+				CurveParentIndices[GuideIndex] = MinParent;
+				CurveLodIndices[GuideIndex] = GuideLod;
 
 				//UE_LOG(LogTemp, Log, TEXT("Guide Index = %d, Num Guides = %d, Num LODs = %d, Guide LOD = %d, LOD Offset = %d, Parent Index = %d"),
 				//	CurveIndex, GuidesLengths.Num(), NumObjectLODs, GuideLod, LodOffset, MinParent);
 			}, EParallelForFlags::None);
 
-			CurveOffset = GeometryCurveOffsets[ObjectIndex];
+			GuidesFacade.SetCurveParentIndices(CurveParentIndices);
+			GuidesFacade.SetCurveLodIndices(CurveLodIndices);
+			CurveOffset = GuidesFacade.GetObjectCurveOffsets()[ObjectIndex];
 		}
 	}
-	
-	FORCEINLINE void BuildCurvesLODs(GeometryCollection::Facades::FCollectionCurveGeometryFacade& CurvesFacade, GeometryCollection::Facades::FCollectionCurveHierarchyFacade& HierarchyFacade, const FDataflowCurveSelection& CurveSelection)
-	{
-		const uint32 NumObjects = CurvesFacade.GetNumGeometry();
-		
-		TArray<int32> CurveParentIndices; CurveParentIndices.Init(INDEX_NONE, CurvesFacade.GetNumCurves());
-		TArray<int32> CurveLodIndices; CurveLodIndices.Init(INDEX_NONE, CurvesFacade.GetNumCurves());
-
-		const bool bValidSelection = CurveSelection.IsValidForCollection(CurvesFacade.GetManagedArrayCollection());
-
-		BuildInternalLODs(CurvesFacade.GetNumGeometry(), CurvesFacade.GetNumCurves(), CurvesFacade.GetGeometryCurveOffsets(),
-			CurvesFacade.GetCurvePointOffsets(), CurvesFacade.GetPointRestPositions(), CurveParentIndices, CurveLodIndices, CurveSelection, bValidSelection);
-		
-		HierarchyFacade.SetCurveParentIndices(CurveParentIndices);
-		HierarchyFacade.SetCurveLodIndices(CurveLodIndices);
-	}
 
 	FCollectionAttributeKey GetCurveParentsKey()
     {
     	FCollectionAttributeKey Key;
-    	Key.Group =  GeometryCollection::Facades::FCollectionCurveGeometryFacade::CurvesGroup.ToString();
-    	Key.Attribute = GeometryCollection::Facades::FCollectionCurveHierarchyFacade::CurveParentIndicesAttribute.ToString();
+    	Key.Group =  UE::Groom::FGroomGuidesFacade::CurvesGroup.ToString();
+    	Key.Attribute = UE::Groom::FGroomGuidesFacade::CurveParentIndicesAttribute.ToString();
     	return Key;
     }
 
 	FCollectionAttributeKey GetCurveLodsKey()
 	{
 		FCollectionAttributeKey Key;
-		Key.Group =  GeometryCollection::Facades::FCollectionCurveGeometryFacade::CurvesGroup.ToString();
-		Key.Attribute = GeometryCollection::Facades::FCollectionCurveHierarchyFacade::CurveLodIndicesAttribute.ToString();
+		Key.Group =  UE::Groom::FGroomGuidesFacade::CurvesGroup.ToString();
+		Key.Attribute = UE::Groom::FGroomGuidesFacade::CurveLodIndicesAttribute.ToString();
 		return Key;
 	}
 }
 
 void FBuildGuidesLODsDataflowNode::Evaluate(UE::Dataflow::FContext& Context, const FDataflowOutput* Out) const
-{
-}
-
-FBuildCurveLODsDataflowNode::FBuildCurveLODsDataflowNode(const UE::Dataflow::FNodeParameters& InParam, FGuid InGuid)
-	: FDataflowNode(InParam, InGuid)
-{
-	RegisterInputConnection(&Collection);
-	RegisterInputConnection(&CurveSelection);
-	RegisterOutputConnection(&Collection, &Collection);
-	RegisterOutputConnection(&CurveParentsKey);
-	RegisterOutputConnection(&CurveLodsKey);
-}
-
-void FBuildCurveLODsDataflowNode::Evaluate(UE::Dataflow::FContext& Context, const FDataflowOutput* Out) const
 {
 	if (Out->IsA<FManagedArrayCollection>(&Collection))
 	{
-		FManagedArrayCollection GeometryCollection = GetValue<FManagedArrayCollection>(Context, &Collection);
+		FManagedArrayCollection GroomCollection = GetValue<FManagedArrayCollection>(Context, &Collection);
 		
-		GeometryCollection::Facades::FCollectionCurveGeometryFacade CurvesFacade(GeometryCollection);
+		UE::Groom::FGroomGuidesFacade GuidesFacade(GroomCollection);
 
-		if(CurvesFacade.IsValid())
+		if(GuidesFacade.IsValid())
 		{
-			GeometryCollection::Facades::FCollectionCurveHierarchyFacade HierarchyFacade(GeometryCollection);
-			HierarchyFacade.DefineSchema();
-			
-			UE::Groom::Private::BuildCurvesLODs(CurvesFacade, HierarchyFacade, GetValue(Context, &CurveSelection));
+			UE::Groom::Private::BuildGuidesLODs(GuidesFacade);
 		}
 		
-		SetValue(Context, MoveTemp(GeometryCollection), &Collection);
+		SetValue(Context, MoveTemp(GroomCollection), &Collection);
 	}
 	else if (Out->IsA<FCollectionAttributeKey>(&CurveParentsKey))
 	{
diff --git a/Engine/Plugins/Runtime/HairStrands/Source/HairStrandsDataflow/Private/GenerateGuidesCurvesNode.cpp b/Engine/Plugins/Runtime/HairStrands/Source/HairStrandsDataflow/Private/GenerateGuidesCurvesNode.cpp
--- a/Engine/Plugins/Runtime/HairStrands/Source/HairStrandsDataflow/Private/GenerateGuidesCurvesNode.cpp
+++ b/Engine/Plugins/Runtime/HairStrands/Source/HairStrandsDataflow/Private/GenerateGuidesCurvesNode.cpp
@@ -3,145 +3,110 @@
 #include "GenerateGuidesCurvesNode.h"
 
 #include "GroomBindingBuilder.h"
+#include "GroomCollectionFacades.h"
 #include "Dataflow/DataflowObjectInterface.h"
-#include "GeometryCollection/Facades/CollectionCurveFacade.h"
 
 #include UE_INLINE_GENERATED_CPP_BY_NAME(GenerateGuidesCurvesNode)
 
 namespace UE::Groom::Private
 {
-	FORCEINLINE void GenerateStrandsCurves(const GeometryCollection::Facades::FCollectionCurveGeometryFacade& CurvesFacade, const FDataflowCurveSelection& CurveSelection, const uint32 GuidesCount, TArray<uint32>& SampleIndices)
+	FORCEINLINE uint32 SampleStrandsCurves(const UE::Groom::FGroomStrandsFacade& StrandsFacade, const uint32 GuidesCount, const uint32 GuidesOffset,
+		const uint32 ObjectIndex, const uint32 PrevStrands, const uint32 NextStrands, TArray<uint32>& SampleIndices)
 	{
-		TArray<FVector3f> RootPositions;
-		RootPositions.Reserve(CurvesFacade.GetNumCurves());
-
-		TArray<bool> ValidPoints;
-		ValidPoints.Reserve(CurvesFacade.GetNumCurves());
+		const uint32 NumStrands = NextStrands-PrevStrands;
+		uint32 NumGuides = (ObjectIndex == StrandsFacade.GetNumObjects()-1) ? GuidesCount - GuidesOffset :
+			GuidesCount * static_cast<float>(NumStrands)/StrandsFacade.GetNumCurves();
+		NumGuides = FMath::Max(1u, NumGuides);
 
-		const bool bValidSelection = CurveSelection.IsValidForCollection(CurvesFacade.GetManagedArrayCollection());;
-			
-		int32 PointOffset = 0;
-		for (int32 CurveIndex = 0; CurveIndex < CurvesFacade.GetNumCurves(); ++CurveIndex)
+		TArray<FVector3f> RootPositions;
+		RootPositions.Init(FVector3f::ZeroVector, NumStrands);
+		for(uint32 StrandIndex = 0; StrandIndex < NumStrands; ++StrandIndex)
 		{
-			RootPositions.Add(CurvesFacade.GetPointRestPositions()[PointOffset]);
-			ValidPoints.Add( (!bValidSelection || (bValidSelection && CurveSelection.IsSelected(CurveIndex))));
-			
-			PointOffset = CurvesFacade.GetCurvePointOffsets()[CurveIndex];
+			const uint32 RootIndex = (StrandIndex+PrevStrands == 0) ? 0 : StrandsFacade.GetCurvePointOffsets()[StrandIndex+PrevStrands-1];
+			RootPositions[StrandIndex] = StrandsFacade.GetPointRestPositions()[RootIndex];
 		}
 
-		GroomBinding_RBFWeighting::FPointsSampler PointsSampler(ValidPoints, RootPositions.GetData(), GuidesCount);
+		TArray<bool> ValidPoints;
+		ValidPoints.Init(true, NumStrands);
+				
+		GroomBinding_RBFWeighting::FPointsSampler PointsSampler(ValidPoints, RootPositions.GetData(), NumGuides);
+
 		SampleIndices = PointsSampler.SampleIndices;
+		return NumGuides;
 	}
 
-	FORCEINLINE void GenerateStrandsPoints(const GeometryCollection::Facades::FCollectionCurveGeometryFacade& StrandsFacade, const TArray<uint32>& SampleIndices,
-		const bool bMergeGuides, GeometryCollection::Facades::FCollectionCurveGeometryFacade& GuidesFacade)
+	FORCEINLINE void BuildGuidesCurves(const UE::Groom::FGroomStrandsFacade& StrandsFacade, const uint32 PrevStrands, const TArray<uint32>& SampleIndices,
+		TArray<FVector3f>& PointRestPositions, TArray<int32>& ObjectCurveOffsets, TArray<int32>& CurvePointOffsets, TArray<int32>& CurveStrandIndices)
 	{
-		TArray<TArray<FVector3f>> GeometryPointPositions;
-		TArray<TArray<int32>> GeometryCurvePoints;
-		TArray<TArray<int32>> GeometrySourceIndices;
-			
-		GeometryPointPositions.SetNum(StrandsFacade.GetNumGeometry());
-		GeometryCurvePoints.SetNum(StrandsFacade.GetNumGeometry());
-		GeometrySourceIndices.SetNum(StrandsFacade.GetNumGeometry());
-		
-		if(bMergeGuides && GuidesFacade.GetNumGeometry() == StrandsFacade.GetNumGeometry())
+		for(uint32 SampleIndex : SampleIndices)
 		{
-			int32 PointOffset = 0;
-			int32 CurveOffset = 0;
-			for (int32 GeometryIndex = 0; GeometryIndex < GuidesFacade.GetNumGeometry(); ++GeometryIndex)
-			{
-				for (int32 CurveIndex = CurveOffset; CurveIndex < GuidesFacade.GetGeometryCurveOffsets()[GeometryIndex]; ++CurveIndex)
-				{
-					for (int32 PointIndex = PointOffset; PointIndex < GuidesFacade.GetCurvePointOffsets()[CurveIndex]; ++PointIndex)
-					{
-						GeometryPointPositions[GeometryIndex].Add(GuidesFacade.GetPointRestPositions()[PointIndex]);
-					}
-					GeometryCurvePoints[GeometryIndex].Add(GuidesFacade.GetCurvePointOffsets()[CurveIndex]-PointOffset);
-					GeometrySourceIndices[GeometryIndex].Add(GuidesFacade.GetCurveSourceIndices()[CurveIndex]);
-					PointOffset = GuidesFacade.GetCurvePointOffsets()[CurveIndex];
-				}
-				CurveOffset = GuidesFacade.GetGeometryCurveOffsets()[GeometryIndex];
-			}
-		}
-			
-		for(const uint32& SampleIndex : SampleIndices)
-		{
-			const int32 PrevPoint = (SampleIndex == 0) ? 0 : StrandsFacade.GetCurvePointOffsets()[SampleIndex-1];
-			const int32 NextPoint = StrandsFacade.GetCurvePointOffsets()[SampleIndex];
+			const uint32 GuideIndex = SampleIndex + PrevStrands;
+			const int32 PointBegin = (GuideIndex == 0) ? 0 : StrandsFacade.GetCurvePointOffsets()[GuideIndex-1];
+			const int32 PointEnd = StrandsFacade.GetCurvePointOffsets()[GuideIndex];
 
-			const int32 GeometryIndex = StrandsFacade.GetCurveGeometryIndices()[SampleIndex];
-			for(int32 PointIndex = PrevPoint; PointIndex < NextPoint; ++PointIndex)
+			for(int32 PointIndex = PointBegin; PointIndex < PointEnd; ++PointIndex)
 			{
-				GeometryPointPositions[GeometryIndex].Add(StrandsFacade.GetPointRestPositions()[PointIndex]);
+				PointRestPositions.Add(StrandsFacade.GetPointRestPositions()[PointIndex]);
 			}
-			GeometryCurvePoints[GeometryIndex].Add(NextPoint-PrevPoint);
-			GeometrySourceIndices[GeometryIndex].Add(SampleIndex);
-		}
-
-		TArray<FVector3f> PointRestPositions;
-		TArray<int32> GeometryCurveOffsets, CurvePointOffsets, CurveSourceIndices;
-		TArray<FString> GeometryGroupNames;
-		
-		int32 CurveOffset = 0;
-		int32 PointOffset = 0;
-		for (int32 GeometryIndex = 0; GeometryIndex < StrandsFacade.GetNumGeometry(); ++GeometryIndex)
-		{
-			PointRestPositions.Append(GeometryPointPositions[GeometryIndex]);
-			for (int32 CurveIndex = 0; CurveIndex < GeometryCurvePoints[GeometryIndex].Num(); ++CurveIndex)
-			{
-				PointOffset += GeometryCurvePoints[GeometryIndex][CurveIndex];
-				CurvePointOffsets.Add(PointOffset);
-				CurveSourceIndices.Add(GeometrySourceIndices[GeometryIndex][CurveIndex]);
-			}
-			CurveOffset += GeometryCurvePoints[GeometryIndex].Num();
-			GeometryCurveOffsets.Add(CurveOffset);
-				
-			const FString GroupName = StrandsFacade.GetGeometryGroupNames()[GeometryIndex].Replace(TEXT("Strands"), TEXT("Guides"));
-			GeometryGroupNames.Add(GroupName);
+			CurvePointOffsets.Add(PointRestPositions.Num());
+			CurveStrandIndices.Add(SampleIndex);
 		}
-		GuidesFacade.InitCurvesCollection(PointRestPositions, CurvePointOffsets, 
-		GeometryCurveOffsets, GeometryGroupNames,  StrandsFacade.GetGeometryCurveThickness(), CurveSourceIndices);
+		ObjectCurveOffsets.Add(CurvePointOffsets.Num());
 	}
-
-	
 }
 
 void FGenerateGuidesCurvesDataflowNode::Evaluate(UE::Dataflow::FContext& Context, const FDataflowOutput* Out) const
-{}
-
-FGenerateCurveGeometryDataflowNode::FGenerateCurveGeometryDataflowNode(const UE::Dataflow::FNodeParameters& InParam, FGuid InGuid)
-	: FDataflowNode(InParam, InGuid)
-{
-	RegisterInputConnection(&Collection);
-	RegisterInputConnection(&SourceCurves);
-	RegisterInputConnection(&CurveSelection);
-	RegisterInputConnection(&CurveCount).SetCanHidePin(true).SetPinIsHidden(true);
-	RegisterOutputConnection(&Collection, &Collection);
-}
-
-void FGenerateCurveGeometryDataflowNode::Evaluate(UE::Dataflow::FContext& Context, const FDataflowOutput* Out) const
 {
 	if (Out->IsA<FManagedArrayCollection>(&Collection))
 	{
-		FManagedArrayCollection TargetCollection = GetValue<FManagedArrayCollection>(Context, &Collection);
-		const FManagedArrayCollection& SourceCollection = GetValue<FManagedArrayCollection>(Context, &SourceCurves);
-		const int32 CountValue = FMath::Max(1, GetValue<int32>(Context, &CurveCount));
-		const FDataflowCurveSelection& DataflowSelection = GetValue<FDataflowCurveSelection>(Context, &CurveSelection);
-		
-		GeometryCollection::Facades::FCollectionCurveGeometryFacade SourceFacade(SourceCollection);
+		FManagedArrayCollection GroomCollection = GetValue<FManagedArrayCollection>(Context, &Collection);
 		
-		if(SourceFacade.IsValid() && CountValue > 0)
+		UE::Groom::FGroomStrandsFacade StrandsFacade(GroomCollection);
+		if(StrandsFacade.IsValid())
 		{
-			TArray<uint32> SampleIndices;
-			UE::Groom::Private::GenerateStrandsCurves(SourceFacade, DataflowSelection, CountValue, SampleIndices);
+			TArray<FVector3f> PointRestPositions;
+			TArray<int32> ObjectCurveOffsets, CurvePointOffsets, CurveStrandIndices;
+			TArray<FString> ObjectGroupNames;
+			
+			uint32 PrevOffset = 0, GuidesOffset = 0;
+			for (int32 ObjectIndex = 0; ObjectIndex < StrandsFacade.GetNumObjects(); ++ObjectIndex)
+			{
+				const uint32 NextOffset = StrandsFacade.GetObjectCurveOffsets()[ObjectIndex];
+
+				TArray<uint32> SampleIndices;
+				
+				// Sample guides among the input strands
+				const uint32 NumGuides = UE::Groom::Private::SampleStrandsCurves(StrandsFacade, GuidesCount, GuidesOffset, ObjectIndex, PrevOffset, NextOffset, SampleIndices);
+
+				// Build the sampled guides
+				UE::Groom::Private::BuildGuidesCurves(StrandsFacade, PrevOffset, SampleIndices, PointRestPositions, ObjectCurveOffsets, CurvePointOffsets, CurveStrandIndices);
+				
+				const FString GroupName = StrandsFacade.GetObjectGroupNames()[ObjectIndex].Replace(
+					*UE::Groom::FGroomStrandsFacade::GroupPrefix.ToString(), *UE::Groom::FGroomGuidesFacade::GroupPrefix.ToString());
+				
+				ObjectGroupNames.Add(GroupName);
+				
+				GuidesOffset += NumGuides;
+				PrevOffset = NextOffset;
+			}
+			
+			UE::Groom::FGroomGuidesFacade GuidesFacade(GroomCollection);
+			const TArray<int32> ObjectPointSamples = GuidesFacade.GetObjectPointSamples();
+
+			// Init the groom collection
+			GuidesFacade.InitGroomCollection(PointRestPositions, CurvePointOffsets, ObjectCurveOffsets, ObjectGroupNames);
 
-			GeometryCollection::Facades::FCollectionCurveGeometryFacade TargetFacade(TargetCollection);
-			UE::Groom::Private::GenerateStrandsPoints(SourceFacade, SampleIndices, bMergeCurves, TargetFacade);
+			// Set the curve strand indices
+			GuidesFacade.SetCurveStrandIndices(CurveStrandIndices);
+
+			// Set the point samples if already defined and if the size is matching
+			if(ObjectPointSamples.Num() == GuidesFacade.GetNumObjects())
+			{
+				GuidesFacade.SetObjectPointSamples(ObjectPointSamples);
+			}
 		}
 		
-		SetValue(Context, MoveTemp(TargetCollection), &Collection);
+		SetValue(Context, MoveTemp(GroomCollection), &Collection);
 	}
 }
 
-
-
diff --git a/Engine/Plugins/Runtime/HairStrands/Source/HairStrandsDataflow/Private/GetGroomAssetNode.cpp b/Engine/Plugins/Runtime/HairStrands/Source/HairStrandsDataflow/Private/GetGroomAssetNode.cpp
--- a/Engine/Plugins/Runtime/HairStrands/Source/HairStrandsDataflow/Private/GetGroomAssetNode.cpp
+++ b/Engine/Plugins/Runtime/HairStrands/Source/HairStrandsDataflow/Private/GetGroomAssetNode.cpp
@@ -1,110 +1,108 @@
-// Copyright Epic Games, Inc. All Rights Reserved.
+﻿// Copyright Epic Games, Inc. All Rights Reserved.
 
 #include "GetGroomAssetNode.h"
 #include "GroomEdit.h"
 #include "GroomInstance.h"
+#include "GroomCollectionFacades.h"
 #include "Dataflow/DataflowInputOutput.h"
 #include "Dataflow/DataflowObjectInterface.h"
-#include "GeometryCollection/Facades/CollectionCurveFacade.h"
-#include "GeometryCollection/Facades/CollectionRenderingFacade.h"
 
 #include UE_INLINE_GENERATED_CPP_BY_NAME(GetGroomAssetNode)
 
 namespace UE::Groom::Private
 {
-	FORCEINLINE void BuildCurveCollection(FManagedArrayCollection& GroomCollection, FEditableGroom& EditGroom, const FString& GroomName, const EGroomCollectionType CurvesType, const float CurveThickness)
+	FORCEINLINE const UGroomAsset* FillEditableAsset(Dataflow::FContext& Context, const UGroomAsset* NodeGroom, FEditableGroom& EditGroom)
+	{
+		const UGroomAsset* GroomAsset = NodeGroom;
+		if (!GroomAsset)
+		{
+			if (const Dataflow::FEngineContext* EngineContext = Context.AsType<Dataflow::FEngineContext>())
+			{
+				GroomAsset = Cast<UGroomAsset>(EngineContext->Owner);
+			}
+		}
+
+		if (GroomAsset)
+		{
+			ConvertFromGroomAsset(const_cast<UGroomAsset*>(GroomAsset), &EditGroom, false, false, false);
+		}
+		return GroomAsset;
+	}
+	
+	template<typename FacadeType>
+	FORCEINLINE void BuildGroomCollection(FManagedArrayCollection& GroomCollection, FEditableGroom& EditGroom, const FString& GroomName)
 	{
 		TArray<FVector3f> PointRestPositions;
-		TArray<int32> GeometryCurveOffsets;
+		TArray<int32> ObjectCurveOffsets;
 		TArray<int32> CurvePointOffsets;
-		TArray<FString> GeometryGroupNames;
-		TArray<float> GeometryCurveThickness;
-		TArray<int32> CurveSourceIndices;
+		TArray<FString> ObjectGroupNames;
 
 		uint32 GroupIndex = 0;
 		for(FEditableGroomGroup& GroomGroup : EditGroom.Groups)
 		{
-			if(CurvesType == EGroomCollectionType::Guides)
-			{
-				for (const FEditableHairGuide& EditType : GroomGroup.Guides)
-				{
-					for(auto& GuidePoint : EditType.ControlPoints)
-					{
-						PointRestPositions.Add(GuidePoint.Position);
-					}
-					CurvePointOffsets.Add(PointRestPositions.Num());
-					CurveSourceIndices.Add(0);
-				}
-			}
-			else
+			for( const typename FacadeType::FEditableType& EditType : FacadeType::GetEditableGroom(GroomGroup))
 			{
-				for (const FEditableHairStrand& EditType : GroomGroup.Strands)
+				for(auto& GuidePoint : EditType.ControlPoints)
 				{
-					for(auto& GuidePoint : EditType.ControlPoints)
-					{
-						PointRestPositions.Add(GuidePoint.Position);
-					}
-					CurvePointOffsets.Add(PointRestPositions.Num());
-					CurveSourceIndices.Add(0);
+					PointRestPositions.Add(GuidePoint.Position);
 				}
+				CurvePointOffsets.Add(PointRestPositions.Num());
 			}
-			GeometryCurveOffsets.Add(CurvePointOffsets.Num());
+			ObjectCurveOffsets.Add(CurvePointOffsets.Num());
 			
-			FString GroupName = (CurvesType == EGroomCollectionType::Guides) ? GroomName + TEXT("_Guides") : GroomName + TEXT("_Strands");
-			GroupName += TEXT("_") + GroomGroup.GroupName.ToString();
-			GeometryGroupNames.Add(GroupName);
-			GeometryCurveThickness.Add(CurveThickness);
+			const FString GroupName = GroomName + TEXT("_") + FacadeType::GroupPrefix.ToString();
+			ObjectGroupNames.Add(GroupName);
 			++GroupIndex;
 		}
-		GeometryCollection::Facades::FCollectionCurveGeometryFacade CurvesFacade(GroomCollection);
-		CurvesFacade.InitCurvesCollection(PointRestPositions, CurvePointOffsets, GeometryCurveOffsets, GeometryGroupNames, GeometryCurveThickness, CurveSourceIndices);
+		FacadeType CurvesFacade(GroomCollection);
+		CurvesFacade.InitGroomCollection(PointRestPositions, CurvePointOffsets, ObjectCurveOffsets, ObjectGroupNames);
 	}
 }
 
 void FGetGroomAssetDataflowNode::Evaluate(UE::Dataflow::FContext& Context, const FDataflowOutput* Out) const
-{
-}
-
-void FGetGroomAssetDataflowNode_v2::Evaluate(UE::Dataflow::FContext& Context, const FDataflowOutput* Out) const
-{
-	if (Out->IsA<TObjectPtr<const UGroomAsset>>(&GroomAsset))
-	{
-		TObjectPtr<const UGroomAsset> LocalAsset = GroomAsset;
-		if (!LocalAsset)
-		{
-			if (const UE::Dataflow::FEngineContext* EngineContext = Context.AsType<UE::Dataflow::FEngineContext>())
-			{
-				LocalAsset = Cast<UGroomAsset>(EngineContext->Owner);
-			}
-		}
-		SetValue(Context, LocalAsset, &GroomAsset);
-	}
-}
-
-FGroomAssetToCollectionDataflowNode::FGroomAssetToCollectionDataflowNode(const UE::Dataflow::FNodeParameters& InParam, FGuid InGuid)
-		: FDataflowNode(InParam, InGuid)
-{
-	RegisterInputConnection(&GroomAsset);
-	RegisterInputConnection(&CurvesThickness).SetCanHidePin(true).SetPinIsHidden(true);
-	RegisterInputConnection(&CurvesType).SetCanHidePin(true).SetPinIsHidden(true);
-	RegisterOutputConnection(&Collection);
-}
-
-void FGroomAssetToCollectionDataflowNode::Evaluate(UE::Dataflow::FContext& Context, const FDataflowOutput* Out) const
 {
 	if (Out->IsA<FManagedArrayCollection>(&Collection))
 	{
 		FManagedArrayCollection GroomCollection;
 		FEditableGroom EditGroom;
-		if(TObjectPtr<const UGroomAsset> InputAsset = GetValue<TObjectPtr<const UGroomAsset>>(Context, &GroomAsset))
+		if(const UGroomAsset* LocalGroom = UE::Groom::Private::FillEditableAsset(Context, GroomAsset, EditGroom))
 		{
-			const float GeometryThickness = GetValue<float>(Context, &CurvesThickness);
-			ConvertFromGroomAsset(const_cast<UGroomAsset*>(InputAsset.Get()), &EditGroom, false, false, false);
+			// Fill the guides facade
+			UE::Groom::Private::BuildGroomCollection<UE::Groom::FGroomGuidesFacade>(
+					GroomCollection, EditGroom, LocalGroom->GetName());
 
-			UE::Groom::Private::BuildCurveCollection(GroomCollection, EditGroom, InputAsset->GetName(), CurvesType, GeometryThickness);
+			// Fill the strands facade
+			UE::Groom::Private::BuildGroomCollection<UE::Groom::FGroomStrandsFacade>(
+					GroomCollection, EditGroom, LocalGroom->GetName());
+
+			// Add guides points samples for future resampling
+			TArray<int32> ObjectPointSamples, CurveStrandIndices;
+			const TArray<FHairGroupsPhysics>& GroupsPhysics = LocalGroom->GetHairGroupsPhysics();
+			for(const FHairGroupsPhysics& GroupPhysics : GroupsPhysics)
+			{
+				ObjectPointSamples.Add(static_cast<uint8>(GroupPhysics.StrandsParameters.StrandsSize));
+			}
+			
+			UE::Groom::FGroomGuidesFacade GuidesFacade(GroomCollection);
+			if(GroomCollection.NumElements(UE::Groom::FGroomGuidesFacade::ObjectsGroup) == 0)
+			{
+				GroomCollection.AddElements(ObjectPointSamples.Num(), UE::Groom::FGroomGuidesFacade::ObjectsGroup);
+			}
+			GuidesFacade.SetObjectPointSamples(ObjectPointSamples);
 		}
 
 		SetValue(Context, MoveTemp(GroomCollection), &Collection);
 	}
 }
 
+TArray<UE::Dataflow::FRenderingParameter> FGetGroomAssetDataflowNode::GetRenderParametersImpl() const
+{
+	if(CurvesType == EGroomCollectionType::Guides)
+	{
+		return { {TEXT("GuidesRender"), FName("FGroomCollection"), {TEXT("Collection")}}};
+	}
+	else
+	{
+		return { {TEXT("StrandsRender"), FName("FGroomCollection"), {TEXT("Collection")}}};
+	}
+}
