- Make groom Niagara component naming consistent
- This allows sequencer to correctly bind to the simulation for caching

diff --git a/Engine/Plugins/Runtime/HairStrands/Source/HairStrandsCore/Private/GroomComponent.cpp b/Engine/Plugins/Runtime/HairStrands/Source/HairStrandsCore/Private/GroomComponent.cpp
--- a/Engine/Plugins/Runtime/HairStrands/Source/HairStrandsCore/Private/GroomComponent.cpp
+++ b/Engine/Plugins/Runtime/HairStrands/Source/HairStrandsCore/Private/GroomComponent.cpp
@@ -1538,6 +1538,8 @@ void UGroomComponent::ReleaseHairSimulation(const int32 GroupIndex)
 				NiagaraComponent->UnregisterComponent();
 			}
 			NiagaraComponent->DestroyComponent();
+			NiagaraComponent->MarkAsGarbage();
+			NiagaraComponent->Rename(nullptr, GetTransientPackage(), REN_NonTransactional | REN_DoNotDirty | REN_DontCreateRedirectors);
 			NiagaraComponent = nullptr;
 		}
 	}
@@ -1552,7 +1554,12 @@ void UGroomComponent::CreateHairSimulation(const int32 GroupIndex, const int32 L
 		{
 			if (!NiagaraComponent)
 			{
-				NiagaraComponent = NewObject<UNiagaraComponent>(this, NAME_None, RF_Transient);
+				FNameBuilder NSName;
+				NSName.Append(TEXT("NiagaraComponent_"));
+				GetFName().AppendString(NSName);
+				NSName.Appendf(TEXT("_%d_%d"), GroupIndex, LODIndex);
+
+				NiagaraComponent = NewObject<UNiagaraComponent>(this, NSName.ToString(), RF_Transient);
 				NiagaraComponent->bUseAttachParentBound = true;
 			}
 			if (GetWorld() && GetWorld()->bIsWorldInitialized)
