- Fix hair reading from potentially stale / invalid data when reading from the sim cache
- Fix for not updating the per instance data resulting in stale data being used

diff --git a/Engine/Plugins/Runtime/HairStrands/Source/HairStrandsCore/Private/Niagara/NiagaraDataInterfaceHairStrands.cpp b/Engine/Plugins/Runtime/HairStrands/Source/HairStrandsCore/Private/Niagara/NiagaraDataInterfaceHairStrands.cpp
--- a/Engine/Plugins/Runtime/HairStrands/Source/HairStrandsCore/Private/Niagara/NiagaraDataInterfaceHairStrands.cpp
+++ b/Engine/Plugins/Runtime/HairStrands/Source/HairStrandsCore/Private/Niagara/NiagaraDataInterfaceHairStrands.cpp
@@ -1015,7 +1015,7 @@ FMatrix44f ComputeWorldTransform(const FNDIHairStrandsData* InstanceData, const
 
 static void InterpolateGroomGuides(
 	FRDGBuilder& GraphBuilder, 
-	FNiagaraDataBuffer* ParticlesBuffer,
+	FNiagaraDataBuffer& ParticlesBuffer,
 	const uint32 NodePositionComponent,  
 	const uint32 RestPositionComponent, 
 	FNDIHairStrandsBuffer* HairStrandsBuffer,
@@ -1041,8 +1041,8 @@ static void InterpolateGroomGuides(
 		PassParameters->StrandsSize = StrandsSize;
 		PassParameters->NumPoints = InstanceRDG.NumPoints;
 
-		PassParameters->NiagaraFloatBuffer = ParticlesBuffer->GetGPUBufferFloat().SRV;
-		PassParameters->NiagaraFloatStride = ParticlesBuffer->GetFloatStride() / sizeof(float);
+		PassParameters->NiagaraFloatBuffer = ParticlesBuffer.GetGPUBufferFloat().SRV;
+		PassParameters->NiagaraFloatStride = ParticlesBuffer.GetFloatStride() / sizeof(float);
 		PassParameters->NodePositionComponent = NodePositionComponent;
 		PassParameters->RestPositionComponent = RestPositionComponent;
 		PassParameters->WorldToLocal = WorldToLocal;
@@ -1081,54 +1081,62 @@ static void InterpolateGroomGuides(
 
 void UNiagaraDataInterfaceHairStrands::SimCachePostReadFrame(void* OptionalPerInstanceData, FNiagaraSystemInstance* SystemInstance)
 {
+	FNDIHairStrandsData* PerIntanceData = static_cast<FNDIHairStrandsData*>(OptionalPerInstanceData);
+	PerInstanceTick(PerIntanceData, SystemInstance, 0.0f);
+
+	const FNiagaraVariable NodePositionVariable(FNiagaraTypeDefinition::GetVec3Def(), FName("NodePosition"));
+	const FNiagaraVariable RestPositionVariable(FNiagaraTypeDefinition::GetVec3Def(), FName("RestPosition"));
+
 	for (TSharedRef<FNiagaraEmitterInstance, ESPMode::ThreadSafe> EmitterInst : SystemInstance->GetEmitters())
 	{
-		FNiagaraDataSet& EmitterDatas  = EmitterInst->GetParticleData();
+		FNiagaraDataSet& EmitterDataSet = EmitterInst->GetParticleData();
+		const FNiagaraDataSetCompiledData& CompiledDataSet = EmitterDataSet.GetCompiledData();
 
-		const FNiagaraVariable NodePositionVariable(FNiagaraTypeDefinition::GetVec3Def(), FName("NodePosition"));
-		const int32 NodePositionIndex = EmitterDatas.GetCompiledData().Variables.IndexOfByKey(NodePositionVariable);
+		// Validate we have the components
+		const int32 NodePositionIndex = CompiledDataSet.Variables.IndexOfByKey(NodePositionVariable);
+		const int32 RestPositionIndex = CompiledDataSet.Variables.IndexOfByKey(RestPositionVariable);
+		if (!CompiledDataSet.VariableLayouts.IsValidIndex(NodePositionIndex) || !CompiledDataSet.VariableLayouts.IsValidIndex(RestPositionIndex))
+		{
+			continue;
+		}
 
-		const FNiagaraVariable RestPositionVariable(FNiagaraTypeDefinition::GetVec3Def(), FName("RestPosition"));
-		const int32 RestPositionIndex = EmitterDatas.GetCompiledData().Variables.IndexOfByKey(RestPositionVariable);
-				
-		if (RestPositionIndex != INDEX_NONE && NodePositionIndex != INDEX_NONE)
+		// Skip compressed systems
+		if (CompiledDataSet.VariableLayouts[NodePositionIndex].GetNumFloatComponents() != 3 || CompiledDataSet.VariableLayouts[RestPositionIndex].GetNumFloatComponents() != 3)
 		{
-			if( EmitterDatas.GetCompiledData().VariableLayouts.IsValidIndex(NodePositionIndex) &&
-				EmitterDatas.GetCompiledData().VariableLayouts.IsValidIndex(RestPositionIndex))
-			{
-				const uint32 NodePositionComponent = EmitterDatas.GetCompiledData().VariableLayouts[NodePositionIndex].GetFloatComponentStart();
-				const uint32 RestPositionComponent = EmitterDatas.GetCompiledData().VariableLayouts[RestPositionIndex].GetFloatComponentStart();
-				FNiagaraDataBuffer* ParticlesBuffer = EmitterDatas.GetCurrentData();
-				
-				FNDIHairStrandsData ProxyData;
-				ProxyData.CopyDatas(static_cast<FNDIHairStrandsData*>(OptionalPerInstanceData));
+			continue;
+		}
+
+		// Queue command to interpolate the data
+		const int NodePositionComponent = CompiledDataSet.VariableLayouts[NodePositionIndex].GetFloatComponentStart();
+		const int RestPositionComponent = CompiledDataSet.VariableLayouts[RestPositionIndex].GetFloatComponentStart();
+
+		FNDIHairStrandsData ProxyData;
+		ProxyData.CopyDatas(PerIntanceData);
 				
-				ENQUEUE_RENDER_COMMAND(NiagaraInterpolateGroomSimCache) (
-					[ProxyData, ParticlesBuffer,
-						NodePositionComponent, RestPositionComponent](FRHICommandListImmediate& RHICmdList)
-					{
-						FMemMark MemMark(FMemStack::Get());
-						FRDGBuilder GraphBuilder(RHICmdList);
+		ENQUEUE_RENDER_COMMAND(NiagaraInterpolateGroomSimCache) (
+			[ProxyData, DataSet_RT=&EmitterDataSet, NodePositionComponent, RestPositionComponent](FRHICommandListImmediate& RHICmdList)
+			{
+				FMemMark MemMark(FMemStack::Get());
+				FRDGBuilder GraphBuilder(RHICmdList);
 
-						const TRefCountPtr<FHairGroupInstance> HairGroupInstance = GetRefCountedHairGroupInstance(ProxyData.HairGroupInstSource.Get(), ProxyData.HairGroupIndex);
-						
-						if(!GEnableProxyInstanceTransform)
-						{
-							const FMatrix44f WorldToLocal = ComputeWorldTransform(&ProxyData, HairGroupInstance).Inverse();
-							InterpolateGroomGuides(GraphBuilder, ParticlesBuffer, NodePositionComponent, RestPositionComponent,
-								ProxyData.HairStrandsBuffer, HairGroupInstance, ProxyData.StrandsSize, WorldToLocal);
-						}
-						else
-						{
-							const FMatrix44f WorldToLocal = ComputeWorldTransform(&ProxyData, ProxyData.HairGroupInstance).Inverse();
-							InterpolateGroomGuides(GraphBuilder, ParticlesBuffer, NodePositionComponent, RestPositionComponent,
-								ProxyData.HairStrandsBuffer, ProxyData.HairGroupInstance, ProxyData.StrandsSize, WorldToLocal);
-						}
-						GraphBuilder.Execute();
-					}
-				);
+				const TRefCountPtr<FHairGroupInstance> HairGroupInstance = GetRefCountedHairGroupInstance(ProxyData.HairGroupInstSource.Get(), ProxyData.HairGroupIndex);
+				
+				FNiagaraDataBuffer& ParticlesBuffer = DataSet_RT->GetCurrentDataChecked();
+				if(!GEnableProxyInstanceTransform)
+				{
+					const FMatrix44f WorldToLocal = ComputeWorldTransform(&ProxyData, HairGroupInstance).Inverse();
+					InterpolateGroomGuides(GraphBuilder, ParticlesBuffer, NodePositionComponent, RestPositionComponent,
+						ProxyData.HairStrandsBuffer, HairGroupInstance, ProxyData.StrandsSize, WorldToLocal);
+				}
+				else
+				{
+					const FMatrix44f WorldToLocal = ComputeWorldTransform(&ProxyData, ProxyData.HairGroupInstance).Inverse();
+					InterpolateGroomGuides(GraphBuilder, ParticlesBuffer, NodePositionComponent, RestPositionComponent,
+						ProxyData.HairStrandsBuffer, ProxyData.HairGroupInstance, ProxyData.StrandsSize, WorldToLocal);
+				}
+				GraphBuilder.Execute();
 			}
-		}
+		);
 	}
 }
 
