Default evaluate implementation to avoid ensure in the graph evaluation for deprecated nodes

diff --git a/Engine/Plugins/Experimental/Dataflow/Source/DataflowNodes/Private/Dataflow/DataflowCollectionEditSkinWeightsNode.cpp b/Engine/Plugins/Experimental/Dataflow/Source/DataflowNodes/Private/Dataflow/DataflowCollectionEditSkinWeightsNode.cpp
--- a/Engine/Plugins/Experimental/Dataflow/Source/DataflowNodes/Private/Dataflow/DataflowCollectionEditSkinWeightsNode.cpp
+++ b/Engine/Plugins/Experimental/Dataflow/Source/DataflowNodes/Private/Dataflow/DataflowCollectionEditSkinWeightsNode.cpp
@@ -189,7 +189,7 @@ TArray<UE::Dataflow::FRenderingParameter> FDataflowCollectionEditSkinWeightsNode
 void FDataflowCollectionEditSkinWeightsNode::AddPrimitiveComponents(UE::Dataflow::FContext& Context, const TSharedPtr<const FManagedArrayCollection> RenderCollection, TObjectPtr<UObject> NodeOwner,
 	TObjectPtr<AActor> RootActor, TArray<TObjectPtr<UPrimitiveComponent>>& PrimitiveComponents)  
 {
-	if(RootActor)
+	if(RootActor && RenderCollection)
 	{
 		GeometryCollection::Facades::FRenderingFacade RenderingFacade(*RenderCollection);
 		const int32 NumGeometry = RenderingFacade.IsValid() ? RenderingFacade.NumGeometry() : 0;
diff --git a/Engine/Plugins/Runtime/HairStrands/Source/HairStrandsDataflow/Private/AttachGuidesRootsNode.cpp b/Engine/Plugins/Runtime/HairStrands/Source/HairStrandsDataflow/Private/AttachGuidesRootsNode.cpp
--- a/Engine/Plugins/Runtime/HairStrands/Source/HairStrandsDataflow/Private/AttachGuidesRootsNode.cpp
+++ b/Engine/Plugins/Runtime/HairStrands/Source/HairStrandsDataflow/Private/AttachGuidesRootsNode.cpp
@@ -8,7 +8,29 @@
 #include UE_INLINE_GENERATED_CPP_BY_NAME(AttachGuidesRootsNode)
 
 void FAttachGuidesRootsDataflowNode::Evaluate(UE::Dataflow::FContext& Context, const FDataflowOutput* Out) const
-{}
+{
+	if (Out->IsA<FManagedArrayCollection>(&Collection))
+	{
+		SafeForwardInput(Context, &Collection, &Collection);
+	}
+	else if (Out->IsA<FCollectionAttributeKey>(&KinematicWeightsKey))
+	{
+		FCollectionAttributeKey Key;
+		Key.Group =  FGeometryCollection::VerticesGroup.ToString();
+		Key.Attribute = GeometryCollection::Facades::FVertexBoneWeightsFacade::KinematicWeightAttributeName.ToString();
+		
+		SetValue(Context, Key, &KinematicWeightsKey);
+	}
+}
+
+FAttachCurveRootsDataflowNode::FAttachCurveRootsDataflowNode(const UE::Dataflow::FNodeParameters& InParam, FGuid InGuid)
+: FDataflowNode(InParam, InGuid)
+{
+	RegisterInputConnection(&Collection);
+	RegisterInputConnection(&CurveSelection);
+	RegisterOutputConnection(&Collection, &Collection);
+	RegisterOutputConnection(&KinematicWeightsKey);
+}
 
 void FAttachCurveRootsDataflowNode::Evaluate(UE::Dataflow::FContext& Context, const FDataflowOutput* Out) const
 {
diff --git a/Engine/Plugins/Runtime/HairStrands/Source/HairStrandsDataflow/Private/BuildGroomSkinningNodes.cpp b/Engine/Plugins/Runtime/HairStrands/Source/HairStrandsDataflow/Private/BuildGroomSkinningNodes.cpp
--- a/Engine/Plugins/Runtime/HairStrands/Source/HairStrandsDataflow/Private/BuildGroomSkinningNodes.cpp
+++ b/Engine/Plugins/Runtime/HairStrands/Source/HairStrandsDataflow/Private/BuildGroomSkinningNodes.cpp
@@ -108,6 +108,23 @@ namespace UE::Groom::Private
 			}
 		}
 	}
+
+	FCollectionAttributeKey GetBoneIndicesKey()
+	{
+		FCollectionAttributeKey Key;
+		Key.Group = FGeometryCollection::VerticesGroup.ToString();
+		Key.Attribute = GeometryCollection::Facades::FVertexBoneWeightsFacade::BoneIndicesAttributeName.ToString();
+	
+		return Key;
+	}
+
+	FCollectionAttributeKey GetBoneWeightsKey()
+	{
+		FCollectionAttributeKey Key;
+		Key.Group = FGeometryCollection::VerticesGroup.ToString();
+		Key.Attribute = GeometryCollection::Facades::FVertexBoneWeightsFacade::BoneWeightsAttributeName.ToString();
+		return Key;
+	}
 }
 
 FTransferSkinWeightsGroomNode::FTransferSkinWeightsGroomNode(const UE::Dataflow::FNodeParameters& InParam, FGuid InGuid)
@@ -122,7 +139,24 @@ FTransferSkinWeightsGroomNode::FTransferSkinWeightsGroomNode(const UE::Dataflow:
 }
 
 void FTransferSkinWeightsGroomNode::Evaluate(UE::Dataflow::FContext& Context, const FDataflowOutput* Out) const
-{}
+{
+	if (Out->IsA(&Collection))
+	{
+		SafeForwardInput(Context, &Collection, &Collection);
+	}
+	else if (Out->IsA<FCollectionAttributeKey>(&BoneIndicesKey))
+	{
+		SetValue(Context, UE::Groom::Private::GetBoneIndicesKey(), &BoneIndicesKey);
+	}
+	else if (Out->IsA<FCollectionAttributeKey>(&BoneWeightsKey))
+	{
+		SetValue(Context, UE::Groom::Private::GetBoneWeightsKey(), &BoneWeightsKey);
+	}
+	else if (Out->IsA<TObjectPtr<USkeletalMesh>>(&SkeletalMesh))
+	{
+		SafeForwardInput(Context, &SkeletalMesh, &SkeletalMesh);
+	}
+}
 
 FTransferGeometrySkinWeightsDataflowNode::FTransferGeometrySkinWeightsDataflowNode(const UE::Dataflow::FNodeParameters& InParam, FGuid InGuid)
 	: FDataflowNode(InParam, InGuid)
@@ -159,30 +193,15 @@ void FTransferGeometrySkinWeightsDataflowNode::Evaluate(UE::Dataflow::FContext&
 	}
 	else if (Out->IsA<FCollectionAttributeKey>(&BoneIndicesKey))
 	{
-		SetValue(Context, GetBoneIndicesKey(), &BoneIndicesKey);
+		SetValue(Context, UE::Groom::Private::GetBoneIndicesKey(), &BoneIndicesKey);
 	}
 	else if (Out->IsA<FCollectionAttributeKey>(&BoneWeightsKey))
 	{
-		SetValue(Context, GetBoneWeightsKey(), &BoneWeightsKey);
+		SetValue(Context, UE::Groom::Private::GetBoneWeightsKey(), &BoneWeightsKey);
 	}
 }
 
-FCollectionAttributeKey FTransferGeometrySkinWeightsDataflowNode::GetBoneIndicesKey()
-{
-	FCollectionAttributeKey Key;
-	Key.Group = FGeometryCollection::VerticesGroup.ToString();
-	Key.Attribute = GeometryCollection::Facades::FVertexBoneWeightsFacade::BoneIndicesAttributeName.ToString();
-	
-	return Key;
-}
 
-FCollectionAttributeKey FTransferGeometrySkinWeightsDataflowNode::GetBoneWeightsKey()
-{
-	FCollectionAttributeKey Key;
-	Key.Group = FGeometryCollection::VerticesGroup.ToString();
-	Key.Attribute = GeometryCollection::Facades::FVertexBoneWeightsFacade::BoneWeightsAttributeName.ToString();
-	return Key;
-}
 
 
 
diff --git a/Engine/Plugins/Runtime/HairStrands/Source/HairStrandsDataflow/Private/BuildGroomSplineSkinningNode.cpp b/Engine/Plugins/Runtime/HairStrands/Source/HairStrandsDataflow/Private/BuildGroomSplineSkinningNode.cpp
--- a/Engine/Plugins/Runtime/HairStrands/Source/HairStrandsDataflow/Private/BuildGroomSplineSkinningNode.cpp
+++ b/Engine/Plugins/Runtime/HairStrands/Source/HairStrandsDataflow/Private/BuildGroomSplineSkinningNode.cpp
@@ -1,6 +1,7 @@
 // Copyright Epic Games, Inc. All Rights Reserved.
 
 #include "BuildGroomSplineSkinningNode.h"
+#include "BuildGroomSkinningNodes.h"
 #include "Algo/MaxElement.h"
 #include "Algo/Count.h"
 #include "Dataflow/DataflowSelection.h"
@@ -15,24 +16,7 @@ namespace UE::Groom::Private
 	const FName FCollectionSplineAttributes::VertexSplineParamAttribute("SplineParam");
 	const FName FCollectionSplineAttributes::VertexSplineBonesAttribute("SplineBones");
 
-	static FCollectionAttributeKey GetBoneIndicesKey()
-	{
-		FCollectionAttributeKey Key;
-		Key.Group = FGeometryCollection::VerticesGroup.ToString();
-		Key.Attribute = GeometryCollection::Facades::FVertexBoneWeightsFacade::BoneIndicesAttributeName.ToString();
-	
-		return Key;
-	}
-
-	static FCollectionAttributeKey GetBoneWeightsKey()
-	{
-		FCollectionAttributeKey Key;
-		Key.Group = FGeometryCollection::VerticesGroup.ToString();
-		Key.Attribute = GeometryCollection::Facades::FVertexBoneWeightsFacade::BoneWeightsAttributeName.ToString();
-		return Key;
-	}
-
-	static FCollectionAttributeKey GetSplineParamKey()
+	FCollectionAttributeKey GetSplineParamKey()
 	{
 		FCollectionAttributeKey Key;
 		Key.Group = FGeometryCollection::VerticesGroup.ToString();
@@ -41,7 +25,7 @@ namespace UE::Groom::Private
 		return Key;
 	}
 
-	static FCollectionAttributeKey GetSplineBonesKey()
+	FCollectionAttributeKey GetSplineBonesKey()
 	{
 		FCollectionAttributeKey Key;
 		Key.Group = FGeometryCollection::VerticesGroup.ToString();
@@ -497,6 +481,18 @@ FBuildGroomSplineSkinWeightsNode::FBuildGroomSplineSkinWeightsNode(const UE::Dat
 
 void FBuildGroomSplineSkinWeightsNode::Evaluate(UE::Dataflow::FContext& Context, const FDataflowOutput* Out) const
 {
+	if (Out->IsA(&Collection))
+	{
+		SafeForwardInput(Context, &Collection, &Collection);
+	}
+	else if (Out->IsA<FCollectionAttributeKey>(&SplineParamKey))
+	{
+		SetValue(Context, UE::Groom::Private::GetSplineParamKey(), &SplineParamKey);
+	}
+	else if (Out->IsA<FCollectionAttributeKey>(&SplineBoneKey))
+	{
+		SetValue(Context, UE::Groom::Private::GetSplineBonesKey(), &SplineBoneKey);
+	}
 }
 
 FConvertSplineToLinearSkinWeightsNode::FConvertSplineToLinearSkinWeightsNode(const UE::Dataflow::FNodeParameters& InParam, FGuid InGuid)
@@ -512,6 +508,18 @@ FConvertSplineToLinearSkinWeightsNode::FConvertSplineToLinearSkinWeightsNode(con
 
 void FConvertSplineToLinearSkinWeightsNode::Evaluate(UE::Dataflow::FContext& Context, const FDataflowOutput* Out) const
 {
+	if (Out->IsA(&Collection))
+	{
+		SafeForwardInput(Context, &Collection, &Collection);
+	}
+	else if (Out->IsA<FCollectionAttributeKey>(&BoneIndicesKey))
+	{
+		SetValue(Context, UE::Groom::Private::GetBoneIndicesKey(), &BoneIndicesKey);
+	}
+	else if (Out->IsA<FCollectionAttributeKey>(&BoneWeightsKey))
+	{
+		SetValue(Context, UE::Groom::Private::GetBoneWeightsKey(), &BoneWeightsKey);
+	}
 }
 
 FConvertLinearToSplineSkinWeightsNode::FConvertLinearToSplineSkinWeightsNode(const UE::Dataflow::FNodeParameters& InParam, FGuid InGuid)
@@ -525,6 +533,18 @@ FConvertLinearToSplineSkinWeightsNode::FConvertLinearToSplineSkinWeightsNode(con
 
 void FConvertLinearToSplineSkinWeightsNode::Evaluate(UE::Dataflow::FContext& Context, const FDataflowOutput* Out) const
 {
+	if (Out->IsA<FManagedArrayCollection>(&Collection))
+	{
+		SafeForwardInput(Context, &Collection, &Collection);
+	}
+	else if (Out->IsA<FCollectionAttributeKey>(&SplineParamKey))
+	{
+		SetValue(Context, UE::Groom::Private::GetSplineParamKey(), &SplineParamKey);
+	}
+	else if (Out->IsA<FCollectionAttributeKey>(&SplineBoneKey))
+	{
+		SetValue(Context, UE::Groom::Private::GetSplineBonesKey(), &SplineBoneKey);
+	}
 }
 
 FBuildSplineSkinWeightsDataflowNode::FBuildSplineSkinWeightsDataflowNode(const UE::Dataflow::FNodeParameters& InParam, FGuid InGuid)
@@ -568,7 +588,7 @@ void FBuildSplineSkinWeightsDataflowNode::Evaluate(UE::Dataflow::FContext& Conte
 					GetValue(Context, &RootBones), GetValue(Context, &SamplesPerSegment), Splines);
 
 				UE::Groom::Private::BuildSplineSkinningWeights(GeometryCollection, MeshFacade, CurveFacade,
-						InputSelection, Splines, GetSplineParamKey(Context), GetSplineBonesKey(Context));
+						InputSelection, Splines, UE::Groom::Private::GetSplineParamKey(), UE::Groom::Private::GetSplineBonesKey());
 				
 				for(int32 GeometryIndex = 0; GeometryIndex < SkinningFacade.NumGeometry(); ++GeometryIndex)
 				{
@@ -593,30 +613,14 @@ void FBuildSplineSkinWeightsDataflowNode::Evaluate(UE::Dataflow::FContext& Conte
 	}
 	else if (Out->IsA<FCollectionAttributeKey>(&SplineParamKey))
 	{
-		SetValue(Context, GetSplineParamKey(Context), &SplineParamKey);
+		SetValue(Context, UE::Groom::Private::GetSplineParamKey(), &SplineParamKey);
 	}
 	else if (Out->IsA<FCollectionAttributeKey>(&SplineBonesKey))
 	{
-		SetValue(Context, GetSplineBonesKey(Context), &SplineBonesKey);
+		SetValue(Context, UE::Groom::Private::GetSplineBonesKey(), &SplineBonesKey);
 	}
 }
 
-FCollectionAttributeKey FBuildSplineSkinWeightsDataflowNode::GetSplineParamKey(UE::Dataflow::FContext& Context) const
-{
-	FCollectionAttributeKey Key;
-	Key.Group = FGeometryCollection::VerticesGroup.ToString();
-	Key.Attribute = UE::Groom::Private::FCollectionSplineAttributes::VertexSplineParamAttribute.ToString();
-	return Key;
-}
-
-FCollectionAttributeKey FBuildSplineSkinWeightsDataflowNode::GetSplineBonesKey(UE::Dataflow::FContext& Context) const
-{
-	FCollectionAttributeKey Key;
-	Key.Group = FGeometryCollection::VerticesGroup.ToString();
-	Key.Attribute = UE::Groom::Private::FCollectionSplineAttributes::VertexSplineBonesAttribute.ToString();
-	return Key;
-}
-
 FSplineToLinearSkinWeightsDataflowNode::FSplineToLinearSkinWeightsDataflowNode(const UE::Dataflow::FNodeParameters& InParam, FGuid InGuid)
 	: FDataflowNode(InParam, InGuid)
 {
@@ -666,8 +670,7 @@ FCollectionAttributeKey FSplineToLinearSkinWeightsDataflowNode::GetSplineParamKe
 	
 	if (Key.Attribute.IsEmpty() && Key.Group.IsEmpty())
 	{
-		Key.Group = FGeometryCollection::VerticesGroup.ToString();
-		Key.Attribute = UE::Groom::Private::FCollectionSplineAttributes::VertexSplineParamAttribute.ToString();
+		return UE::Groom::Private::GetSplineParamKey();
 	}
 
 	return Key;
@@ -679,8 +682,7 @@ FCollectionAttributeKey FSplineToLinearSkinWeightsDataflowNode::GetSplineBonesKe
 
 	if (Key.Attribute.IsEmpty() && Key.Group.IsEmpty())
 	{
-		Key.Group = FGeometryCollection::VerticesGroup.ToString();
-		Key.Attribute = UE::Groom::Private::FCollectionSplineAttributes::VertexSplineBonesAttribute.ToString();
+		return UE::Groom::Private::GetSplineBonesKey();
 	}
 
 	return Key;
diff --git a/Engine/Plugins/Runtime/HairStrands/Source/HairStrandsDataflow/Private/BuildGuidesLODsNode.cpp b/Engine/Plugins/Runtime/HairStrands/Source/HairStrandsDataflow/Private/BuildGuidesLODsNode.cpp
--- a/Engine/Plugins/Runtime/HairStrands/Source/HairStrandsDataflow/Private/BuildGuidesLODsNode.cpp
+++ b/Engine/Plugins/Runtime/HairStrands/Source/HairStrandsDataflow/Private/BuildGuidesLODsNode.cpp
@@ -124,6 +124,18 @@ namespace UE::Groom::Private
 
 void FBuildGuidesLODsDataflowNode::Evaluate(UE::Dataflow::FContext& Context, const FDataflowOutput* Out) const
 {
+	if (Out->IsA(&Collection))
+	{
+		SafeForwardInput(Context, &Collection, &Collection);
+	}
+	else if (Out->IsA<FCollectionAttributeKey>(&CurveParentsKey))
+	{
+		SetValue(Context, UE::Groom::Private::GetCurveParentsKey(), &CurveParentsKey);
+	}
+	else if (Out->IsA<FCollectionAttributeKey>(&CurveLodsKey))
+	{
+		SetValue(Context, UE::Groom::Private::GetCurveLodsKey(), &CurveLodsKey);
+	}
 }
 
 FBuildCurveLODsDataflowNode::FBuildCurveLODsDataflowNode(const UE::Dataflow::FNodeParameters& InParam, FGuid InGuid)
diff --git a/Engine/Plugins/Runtime/HairStrands/Source/HairStrandsDataflow/Private/GenerateGuidesCurvesNode.cpp b/Engine/Plugins/Runtime/HairStrands/Source/HairStrandsDataflow/Private/GenerateGuidesCurvesNode.cpp
--- a/Engine/Plugins/Runtime/HairStrands/Source/HairStrandsDataflow/Private/GenerateGuidesCurvesNode.cpp
+++ b/Engine/Plugins/Runtime/HairStrands/Source/HairStrandsDataflow/Private/GenerateGuidesCurvesNode.cpp
@@ -107,7 +107,12 @@ namespace UE::Groom::Private
 }
 
 void FGenerateGuidesCurvesDataflowNode::Evaluate(UE::Dataflow::FContext& Context, const FDataflowOutput* Out) const
-{}
+{
+	if (Out->IsA(&Collection))
+	{
+		SafeForwardInput(Context, &Collection, &Collection);
+	}
+}
 
 FGenerateCurveGeometryDataflowNode::FGenerateCurveGeometryDataflowNode(const UE::Dataflow::FNodeParameters& InParam, FGuid InGuid)
 	: FDataflowNode(InParam, InGuid)
diff --git a/Engine/Plugins/Runtime/HairStrands/Source/HairStrandsDataflow/Private/GetGroomAssetNode.cpp b/Engine/Plugins/Runtime/HairStrands/Source/HairStrandsDataflow/Private/GetGroomAssetNode.cpp
--- a/Engine/Plugins/Runtime/HairStrands/Source/HairStrandsDataflow/Private/GetGroomAssetNode.cpp
+++ b/Engine/Plugins/Runtime/HairStrands/Source/HairStrandsDataflow/Private/GetGroomAssetNode.cpp
@@ -63,6 +63,11 @@ namespace UE::Groom::Private
 
 void FGetGroomAssetDataflowNode::Evaluate(UE::Dataflow::FContext& Context, const FDataflowOutput* Out) const
 {
+	if (Out->IsA<FManagedArrayCollection>(&Collection))
+	{
+		FManagedArrayCollection GroomCollection;
+		SetValue(Context, MoveTemp(GroomCollection), &Collection);
+	}
 }
 
 void FGetGroomAssetDataflowNode_v2::Evaluate(UE::Dataflow::FContext& Context, const FDataflowOutput* Out) const
diff --git a/Engine/Plugins/Runtime/HairStrands/Source/HairStrandsDataflow/Private/GetGroomAttributesNodes.cpp b/Engine/Plugins/Runtime/HairStrands/Source/HairStrandsDataflow/Private/GetGroomAttributesNodes.cpp
--- a/Engine/Plugins/Runtime/HairStrands/Source/HairStrandsDataflow/Private/GetGroomAttributesNodes.cpp
+++ b/Engine/Plugins/Runtime/HairStrands/Source/HairStrandsDataflow/Private/GetGroomAttributesNodes.cpp
@@ -10,6 +10,11 @@
 
 void FGetGroomAttributesDataflowNode::Evaluate(UE::Dataflow::FContext& Context, const FDataflowOutput* Out) const
 {
+	if (Out->IsA<FCollectionAttributeKey>(&AttributeKey))
+	{
+		FCollectionAttributeKey Key;
+		SetValue(Context, MoveTemp(Key), &AttributeKey);
+	}
 }
 
 FGetCurveAttributesDataflowNode::FGetCurveAttributesDataflowNode(const UE::Dataflow::FNodeParameters& InParam, FGuid InGuid)
diff --git a/Engine/Plugins/Runtime/HairStrands/Source/HairStrandsDataflow/Private/GroomAssetTerminalNode.cpp b/Engine/Plugins/Runtime/HairStrands/Source/HairStrandsDataflow/Private/GroomAssetTerminalNode.cpp
--- a/Engine/Plugins/Runtime/HairStrands/Source/HairStrandsDataflow/Private/GroomAssetTerminalNode.cpp
+++ b/Engine/Plugins/Runtime/HairStrands/Source/HairStrandsDataflow/Private/GroomAssetTerminalNode.cpp
@@ -275,6 +275,10 @@ void FGroomAssetTerminalDataflowNode::SetAssetValue(TObjectPtr<UObject> Asset, U
 
 void FGroomAssetTerminalDataflowNode::Evaluate(UE::Dataflow::FContext& Context, const FDataflowOutput* Out) const
 {
+	if (Out->IsA<FManagedArrayCollection>(&Collection))
+	{
+		SafeForwardInput(Context, &Collection, &Collection);
+	}
 }
 
 FGroomAssetTerminalDataflowNode_v2::FGroomAssetTerminalDataflowNode_v2(const UE::Dataflow::FNodeParameters& InParam, FGuid InGuid)
diff --git a/Engine/Plugins/Runtime/HairStrands/Source/HairStrandsDataflow/Private/ResampleGuidesPointsNode.cpp b/Engine/Plugins/Runtime/HairStrands/Source/HairStrandsDataflow/Private/ResampleGuidesPointsNode.cpp
--- a/Engine/Plugins/Runtime/HairStrands/Source/HairStrandsDataflow/Private/ResampleGuidesPointsNode.cpp
+++ b/Engine/Plugins/Runtime/HairStrands/Source/HairStrandsDataflow/Private/ResampleGuidesPointsNode.cpp
@@ -135,7 +135,12 @@ namespace UE::Groom::Private
 }
 
 void FResampleGuidesPointsDataflowNode::Evaluate(UE::Dataflow::FContext& Context, const FDataflowOutput* Out) const
-{}
+{
+	if (Out->IsA<FManagedArrayCollection>(&Collection))
+	{
+		SafeForwardInput(Context, &Collection, &Collection);
+	}
+}
 
 FResampleCurvePointsDataflowNode::FResampleCurvePointsDataflowNode(const UE::Dataflow::FNodeParameters& InParam, FGuid InGuid)
 : FDataflowNode(InParam, InGuid)
diff --git a/Engine/Plugins/Runtime/HairStrands/Source/HairStrandsDataflow/Private/SmoothGuidesCurvesNode.cpp b/Engine/Plugins/Runtime/HairStrands/Source/HairStrandsDataflow/Private/SmoothGuidesCurvesNode.cpp
--- a/Engine/Plugins/Runtime/HairStrands/Source/HairStrandsDataflow/Private/SmoothGuidesCurvesNode.cpp
+++ b/Engine/Plugins/Runtime/HairStrands/Source/HairStrandsDataflow/Private/SmoothGuidesCurvesNode.cpp
@@ -79,7 +79,12 @@ namespace UE::Groom::Private
 }
 
 void FSmoothGuidesCurvesDataflowNode::Evaluate(UE::Dataflow::FContext& Context, const FDataflowOutput* Out) const
-{}
+{
+	if (Out->IsA<FManagedArrayCollection>(&Collection))
+	{
+		SafeForwardInput(Context, &Collection, &Collection);
+	}
+}
 
 FSmoothCurvePointsDataflowNode::FSmoothCurvePointsDataflowNode(const UE::Dataflow::FNodeParameters& InParam, FGuid InGuid)
 	: FDataflowNode(InParam, InGuid)
diff --git a/Engine/Plugins/Runtime/HairStrands/Source/HairStrandsDataflow/Public/AttachGuidesRootsNode.h b/Engine/Plugins/Runtime/HairStrands/Source/HairStrandsDataflow/Public/AttachGuidesRootsNode.h
--- a/Engine/Plugins/Runtime/HairStrands/Source/HairStrandsDataflow/Public/AttachGuidesRootsNode.h
+++ b/Engine/Plugins/Runtime/HairStrands/Source/HairStrandsDataflow/Public/AttachGuidesRootsNode.h
@@ -56,15 +56,9 @@ struct FAttachCurveRootsDataflowNode : public FDataflowNode
 
 public:
 	
-	FAttachCurveRootsDataflowNode(const UE::Dataflow::FNodeParameters& InParam, FGuid InGuid = FGuid::NewGuid())
-	: FDataflowNode(InParam, InGuid)
-	{
-		RegisterInputConnection(&Collection);
-		RegisterInputConnection(&CurveSelection);
-		RegisterOutputConnection(&Collection, &Collection);
-		RegisterOutputConnection(&KinematicWeightsKey);
-	}
+	FAttachCurveRootsDataflowNode(const UE::Dataflow::FNodeParameters& InParam, FGuid InGuid = FGuid::NewGuid());
 
+private:
 	//~ Begin FDataflowNode interface
 	virtual void Evaluate(UE::Dataflow::FContext& Context, const FDataflowOutput* Out) const override;
 	//~ End FDataflowNode interface
diff --git a/Engine/Plugins/Runtime/HairStrands/Source/HairStrandsDataflow/Public/BuildGroomSkinningNodes.h b/Engine/Plugins/Runtime/HairStrands/Source/HairStrandsDataflow/Public/BuildGroomSkinningNodes.h
--- a/Engine/Plugins/Runtime/HairStrands/Source/HairStrandsDataflow/Public/BuildGroomSkinningNodes.h
+++ b/Engine/Plugins/Runtime/HairStrands/Source/HairStrandsDataflow/Public/BuildGroomSkinningNodes.h
@@ -12,6 +12,15 @@
 
 #include "BuildGroomSkinningNodes.generated.h"
 
+namespace UE::Groom::Private
+{
+	/** Retrieve the bone indices key */
+	FCollectionAttributeKey GetBoneIndicesKey();
+
+	/** Retrieve the bone weights key */
+	FCollectionAttributeKey GetBoneWeightsKey();
+}
+
 /** Build the curves skinning by transferring the indices weights from a skelmesh */
 USTRUCT(meta = (Experimental, DataflowGroom, Deprecated="5.7"))
 struct UE_DEPRECATED(5.7, "Use the newer version of this node instead.") FTransferSkinWeightsGroomNode : public FDataflowNode
@@ -23,6 +32,8 @@ struct UE_DEPRECATED(5.7, "Use the newer version of this node instead.") FTransf
 public:
 	
 	FTransferSkinWeightsGroomNode(const UE::Dataflow::FNodeParameters& InParam, FGuid InGuid = FGuid::NewGuid());
+	
+private:
 
 	//~ Begin FDataflowNode interface
 	virtual void Evaluate(UE::Dataflow::FContext& Context, const FDataflowOutput* Out) const override;
@@ -59,13 +70,6 @@ public:
 	/** Bone weights key to be used in other nodes if necessary */
 	UPROPERTY(meta = (DisplayName = "Bone Weights", DataflowOutput))
 	FCollectionAttributeKey BoneWeightsKey;
-
-private :
-	/** Get the bone indices key */
-	FCollectionAttributeKey GetBoneIndicesKey() const;
-
-	/** Get the bone weights key */
-	FCollectionAttributeKey GetBoneWeightsKey() const;
 };
 
 /** Build the curves skinning by transferring the indices weights from a skelmesh geometry */
@@ -78,9 +82,9 @@ struct FTransferGeometrySkinWeightsDataflowNode : public FDataflowNode
 
 public:
 	FTransferGeometrySkinWeightsDataflowNode(const UE::Dataflow::FNodeParameters& InParam, FGuid InGuid = FGuid::NewGuid());
-	
 
 private:
+	
 	//~ Begin FDataflowNode interface
 	virtual void Evaluate(UE::Dataflow::FContext& Context, const FDataflowOutput* Out) const override;
 	//~ End FDataflowNode interface
@@ -112,10 +116,4 @@ private:
 	/** Bone weights key to be used in other nodes if necessary */
 	UPROPERTY(meta = (DisplayName = "Bone Weights", DataflowOutput))
 	FCollectionAttributeKey BoneWeightsKey;
-	
-	/** Get the bone indices key */
-	static FCollectionAttributeKey GetBoneIndicesKey();
-
-	/** Get the bone weights key */
-	static FCollectionAttributeKey GetBoneWeightsKey();
 };
diff --git a/Engine/Plugins/Runtime/HairStrands/Source/HairStrandsDataflow/Public/BuildGroomSplineSkinningNode.h b/Engine/Plugins/Runtime/HairStrands/Source/HairStrandsDataflow/Public/BuildGroomSplineSkinningNode.h
--- a/Engine/Plugins/Runtime/HairStrands/Source/HairStrandsDataflow/Public/BuildGroomSplineSkinningNode.h
+++ b/Engine/Plugins/Runtime/HairStrands/Source/HairStrandsDataflow/Public/BuildGroomSplineSkinningNode.h
@@ -10,6 +10,15 @@
 
 #include "BuildGroomSplineSkinningNode.generated.h"
 
+namespace UE::Groom::Private
+{
+	/** Retrieve the spline param key */
+	FCollectionAttributeKey GetSplineParamKey();
+
+	/** Retrieve the spline bones key */
+	FCollectionAttributeKey GetSplineBonesKey();
+}
+
 /** Build spline skinning data from skeletal mesh */
 USTRUCT(meta = (Experimental, DataflowGroom, Deprecated="5.7"))
 struct UE_DEPRECATED(5.7, "Use the newer version of this node instead.") FBuildGroomSplineSkinWeightsNode : public FDataflowNode
@@ -224,12 +233,6 @@ private:
 	/** Spline Bones key containing root and end spline bone. To be used in other nodes if necessary */
 	UPROPERTY(meta = (DisplayName = "Spline Bones", DataflowOutput))
 	FCollectionAttributeKey SplineBonesKey;
-
-	/** Get the spline param key */
-	FCollectionAttributeKey GetSplineParamKey(UE::Dataflow::FContext& Context) const;
-
-	/** Get the spline bones key */
-	FCollectionAttributeKey GetSplineBonesKey(UE::Dataflow::FContext& Context) const;
 };
 
 /** Convert spline skinning data to linear data */
