Cooker: GroomAssets: UnexpectedLoad: Mark the load of the assets required for EGroomNiagaraSolvers as used-in-game dependencies, to suppress the cooker's UnexpectedLoad message for them.
#rnx

diff --git a/Engine/Plugins/Runtime/HairStrands/Source/HairStrandsCore/Private/GroomAsset.cpp b/Engine/Plugins/Runtime/HairStrands/Source/HairStrandsCore/Private/GroomAsset.cpp
--- a/Engine/Plugins/Runtime/HairStrands/Source/HairStrandsCore/Private/GroomAsset.cpp
+++ b/Engine/Plugins/Runtime/HairStrands/Source/HairStrandsCore/Private/GroomAsset.cpp
@@ -41,6 +41,7 @@
 #include "GroomDeformerBuilder.h"
 #include "HairCardsVertexFactory.h"
 #include "PSOPrecache.h"
+#include "UObject/ICookInfo.h"
 #include "UObject/UObjectIterator.h"
 #include "RenderGraphBuilder.h"
 #include "UObject/DevObjectVersion.h"
@@ -1811,6 +1812,7 @@ void UGroomAsset::PostInitProperties()
 
 void UGroomAsset::UpdatePhysicsSystems()
 {
+	FCookLoadScope UsedInGameScope(ECookLoadType::UsedInGame);
 	for (auto& Group : GetHairGroupsPhysics())
 	{
 		if (Group.SolverSettings.NiagaraSolver == EGroomNiagaraSolvers::AngularSprings)
