Fix source indexing for the groom attributes

#virtualized

diff --git a/Engine/Plugins/Runtime/HairStrands/Shaders/Private/Deformers/DeformerGroomAttributeRead.ush b/Engine/Plugins/Runtime/HairStrands/Shaders/Private/Deformers/DeformerGroomAttributeRead.ush
--- a/Engine/Plugins/Runtime/HairStrands/Shaders/Private/Deformers/DeformerGroomAttributeRead.ush
+++ b/Engine/Plugins/Runtime/HairStrands/Shaders/Private/Deformers/DeformerGroomAttributeRead.ush
@@ -1,13 +1,26 @@
 // Copyright Epic Games, Inc. All Rights Reserved.
 
 #include "/Plugin/ComputeFramework/Private/ComputeKernelCommon.ush"
+#include "/Engine/Private/HairStrands/HairStrandsPack.ush"
 
 uint {DataInterfaceName}_NumValues; 
 uint {DataInterfaceName}_BaseValueIndex; 
 uint {DataInterfaceName}_TotalValueCount;
+uint {DataInterfaceName}_bHasIndexMapping;
 
 Buffer<float> {DataInterfaceName}_ValueTypedData;
 
+ByteAddressBuffer {DataInterfaceName}_IndexMappingBuffer;
+
+// Utilities function
+
+uint GetAttributeIndex_{DataInterfaceName}(uint GlobalIndex)
+{
+	const uint LocalIndex = GlobalIndex - {DataInterfaceName}_BaseValueIndex;
+	return {DataInterfaceName}_bHasIndexMapping ?
+		ReadHairSourceAttributeIndex({DataInterfaceName}_IndexMappingBuffer, LocalIndex, {DataInterfaceName}_bHasIndexMapping) : LocalIndex;
+}
+
 // Interface functions
 
 uint ReadNumValues_{DataInterfaceName}()
@@ -17,39 +30,39 @@ uint ReadNumValues_{DataInterfaceName}()
 
 uint ReadValueTypedUint_{DataInterfaceName}(uint ValueIndex)
 {
-	ValueIndex -= {DataInterfaceName}_BaseValueIndex;
+	ValueIndex = GetAttributeIndex_{DataInterfaceName}(ValueIndex);
 	return (ValueIndex < {DataInterfaceName}_NumValues) ? (uint){DataInterfaceName}_ValueTypedData[ValueIndex] : 0;
 }
 
 bool ReadValueTypedBool_{DataInterfaceName}(uint ValueIndex)
 {
-	ValueIndex -= {DataInterfaceName}_BaseValueIndex;
+	ValueIndex = GetAttributeIndex_{DataInterfaceName}(ValueIndex);
 	return (ValueIndex < {DataInterfaceName}_NumValues) ? (({DataInterfaceName}_ValueTypedData[ValueIndex] == 1.0f) ? true : false) : false;
 }
 
 int ReadValueTypedInt_{DataInterfaceName}(uint ValueIndex)
 {
-	ValueIndex -= {DataInterfaceName}_BaseValueIndex;
+	ValueIndex = GetAttributeIndex_{DataInterfaceName}(ValueIndex);
 	return (ValueIndex < {DataInterfaceName}_NumValues) ? (int){DataInterfaceName}_ValueTypedData[ValueIndex] : 0;
 }
 
 int2 ReadValueTypedInt2_{DataInterfaceName}(uint ValueIndex)
 {
-	ValueIndex -= {DataInterfaceName}_BaseValueIndex;
+	ValueIndex = GetAttributeIndex_{DataInterfaceName}(ValueIndex);
 	const int ValueOffset = 2 * ValueIndex;
 	return (ValueIndex < {DataInterfaceName}_NumValues) ? int2({DataInterfaceName}_ValueTypedData[ValueOffset], {DataInterfaceName}_ValueTypedData[ValueOffset+1]) : int2(0,0);
 }
 
 int3 ReadValueTypedInt3_{DataInterfaceName}(uint ValueIndex)
 {
-	ValueIndex -= {DataInterfaceName}_BaseValueIndex;
+	ValueIndex = GetAttributeIndex_{DataInterfaceName}(ValueIndex);
 	const int ValueOffset = 3 * ValueIndex;
 	return (ValueIndex < {DataInterfaceName}_NumValues) ? int3({DataInterfaceName}_ValueTypedData[ValueOffset], {DataInterfaceName}_ValueTypedData[ValueOffset+1], {DataInterfaceName}_ValueTypedData[ValueOffset+2]) : int3(0,0,0);
 }
 
 int4 ReadValueTypedInt4_{DataInterfaceName}(uint ValueIndex)
 {
-	ValueIndex -= {DataInterfaceName}_BaseValueIndex;
+	ValueIndex = GetAttributeIndex_{DataInterfaceName}(ValueIndex);
 	const int ValueOffset = 4 * ValueIndex;
 	return (ValueIndex < {DataInterfaceName}_NumValues) ? int4({DataInterfaceName}_ValueTypedData[ValueOffset], {DataInterfaceName}_ValueTypedData[ValueOffset+1],
 					{DataInterfaceName}_ValueTypedData[ValueOffset+2], {DataInterfaceName}_ValueTypedData[ValueOffset+3]) : int4(0,0,0,0);
@@ -57,13 +70,13 @@ int4 ReadValueTypedInt4_{DataInterfaceName}(uint ValueIndex)
 
 float ReadValueTypedFloat_{DataInterfaceName}(uint ValueIndex)
 {
-	ValueIndex -= {DataInterfaceName}_BaseValueIndex;
+	ValueIndex = GetAttributeIndex_{DataInterfaceName}(ValueIndex);
 	return (ValueIndex < {DataInterfaceName}_NumValues) ? {DataInterfaceName}_ValueTypedData[ValueIndex] : 0.0f;
 }
 
 float2 ReadValueTypedFloat2_{DataInterfaceName}(uint ValueIndex)
 {
-	ValueIndex -= {DataInterfaceName}_BaseValueIndex;
+	ValueIndex = GetAttributeIndex_{DataInterfaceName}(ValueIndex);;
 	const int ValueOffset = 2 * ValueIndex;
 	return (ValueIndex < {DataInterfaceName}_NumValues) ? float2({DataInterfaceName}_ValueTypedData[ValueOffset], {DataInterfaceName}_ValueTypedData[ValueOffset+1]) : float2(0.0f,0.0f);
 
@@ -71,14 +84,14 @@ float2 ReadValueTypedFloat2_{DataInterfaceName}(uint ValueIndex)
 
 float3 ReadValueTypedFloat3_{DataInterfaceName}(uint ValueIndex)
 {
-	ValueIndex -= {DataInterfaceName}_BaseValueIndex;
+	ValueIndex = GetAttributeIndex_{DataInterfaceName}(ValueIndex);
 	const int ValueOffset = 3 * ValueIndex;
 	return (ValueIndex < {DataInterfaceName}_NumValues) ? float3({DataInterfaceName}_ValueTypedData[ValueOffset], {DataInterfaceName}_ValueTypedData[ValueOffset+1], {DataInterfaceName}_ValueTypedData[ValueOffset+2]) : float3(0.0f,0.0f,0.0f);
 }
 
 float4 ReadValueTypedFloat4_{DataInterfaceName}(uint ValueIndex)
 {
-	ValueIndex -= {DataInterfaceName}_BaseValueIndex;
+	ValueIndex = GetAttributeIndex_{DataInterfaceName}(ValueIndex);
 	const int ValueOffset = 4 * ValueIndex;
 	return (ValueIndex < {DataInterfaceName}_NumValues) ? float4({DataInterfaceName}_ValueTypedData[ValueOffset], {DataInterfaceName}_ValueTypedData[ValueOffset+1],
 					{DataInterfaceName}_ValueTypedData[ValueOffset+2], {DataInterfaceName}_ValueTypedData[ValueOffset+3]) : float4(0.0f,0.0f,0.0f,0.0f);
@@ -86,7 +99,7 @@ float4 ReadValueTypedFloat4_{DataInterfaceName}(uint ValueIndex)
 
 float4x4 ReadValueTypedFloat4x4_{DataInterfaceName}(uint ValueIndex)
 {
-	ValueIndex -= {DataInterfaceName}_BaseValueIndex;
+	ValueIndex = GetAttributeIndex_{DataInterfaceName}(ValueIndex);
 	const int ValueOffset = 16 * ValueIndex;
 	return (ValueIndex < {DataInterfaceName}_NumValues) ? float4x4({DataInterfaceName}_ValueTypedData[ValueOffset], {DataInterfaceName}_ValueTypedData[ValueOffset+1],
 						{DataInterfaceName}_ValueTypedData[ValueOffset+2], {DataInterfaceName}_ValueTypedData[ValueOffset+3],
@@ -101,7 +114,7 @@ float4x4 ReadValueTypedFloat4x4_{DataInterfaceName}(uint ValueIndex)
 
 float3x4 ReadValueTypedFloat3x4_{DataInterfaceName}(uint ValueIndex)
 {
-	ValueIndex -= {DataInterfaceName}_BaseValueIndex;
+	ValueIndex = GetAttributeIndex_{DataInterfaceName}(ValueIndex);
 	const int ValueOffset = 12 * ValueIndex;
 	return (ValueIndex < {DataInterfaceName}_NumValues) ? float3x4({DataInterfaceName}_ValueTypedData[ValueOffset], {DataInterfaceName}_ValueTypedData[ValueOffset+1],
 						{DataInterfaceName}_ValueTypedData[ValueOffset+2], {DataInterfaceName}_ValueTypedData[ValueOffset+3],
diff --git a/Engine/Plugins/Runtime/HairStrands/Source/HairStrandsDeformer/Private/DeformerGroomAttributeRead.cpp b/Engine/Plugins/Runtime/HairStrands/Source/HairStrandsDeformer/Private/DeformerGroomAttributeRead.cpp
--- a/Engine/Plugins/Runtime/HairStrands/Source/HairStrandsDeformer/Private/DeformerGroomAttributeRead.cpp
+++ b/Engine/Plugins/Runtime/HairStrands/Source/HairStrandsDeformer/Private/DeformerGroomAttributeRead.cpp
@@ -221,7 +221,9 @@ BEGIN_SHADER_PARAMETER_STRUCT(FOptimusGroomAttributeReadParameters, )
 	SHADER_PARAMETER(uint32, NumValues)
 	SHADER_PARAMETER(uint32, BaseValueIndex)
 	SHADER_PARAMETER(uint32, TotalValueCount)
+	SHADER_PARAMETER(uint32, bHasIndexMapping)
 	SHADER_PARAMETER_RDG_BUFFER_SRV(Buffer<float>, ValueTypedData)
+	SHADER_PARAMETER_RDG_BUFFER_SRV(ByteAddressBuffer, IndexMappingBuffer)
 END_SHADER_PARAMETER_STRUCT()
 
 void UOptimusGroomAttributeReadDataInterface::GetShaderParameters(TCHAR const* UID, FShaderParametersMetadataBuilder& InOutBuilder, FShaderParametersMetadataAllocations& InOutAllocations) const
@@ -281,7 +283,23 @@ FOptimusGroomAttributeReadProviderProxy::FOptimusGroomAttributeReadProviderProxy
 	TArray<const UGroomComponent*> GroomComponents;
 	UE::Groom::Private::GatherGroomComponents(MeshComponent, GroomComponents);
 	UE::Groom::Private::GroomComponentsToInstances(GroomComponents, GroupInstances);
-	UE::Groom::Private::GetGroomInvocationElementGroups(GroomComponents,UOptimusGroomExecDataInterface::GetExecutionDomainName(GroomAttributeGroup), GroupElements, true);
+	UE::Groom::Private::GetGroomInvocationElementGroups(GroomComponents,
+		UOptimusGroomExecDataInterface::GetExecutionDomainName(GroomAttributeGroup), GroupElements, true);
+
+	// Retrieve the rendering elements (not the source ones)
+	TArray<TPair<UGroomAsset*, UE::Groom::Private::FGroupElements>> InvocationElements;
+	UE::Groom::Private::GetGroomInvocationElementGroups(GroomComponents,
+		UOptimusGroomExecDataInterface::GetExecutionDomainName(GroomAttributeGroup), InvocationElements, false);
+
+	// Compute the number of elements for the base indices
+	NumElements.Reset();
+	for (TPair<UGroomAsset*, UE::Groom::Private::FGroupElements>& Elements : InvocationElements)
+	{
+		for (int32 GroupIndex = 0, NumGroups = Elements.Value.GroupOffsets.Num()-1; GroupIndex < NumGroups; ++GroupIndex)
+		{
+			NumElements.Add(Elements.Value.GroupOffsets[GroupIndex+1]-Elements.Value.GroupOffsets[GroupIndex]);
+		}
+	}
 }
 
 bool FOptimusGroomAttributeReadProviderProxy::IsValid(FValidationData const& InValidationData) const
@@ -505,7 +523,8 @@ FORCEINLINE int32 NumAttributeValues<FVector2f>()
 
 template<typename DataType>
 static void CreateInternalBuffers(const UGroomAsset* GroomAsset, const UE::Groom::Private::FGroupElements& GroupElements, const FName& AttributeName,
-	const EOptimusGroomExecDomain& AttributeGroup, FRDGBuilder& GraphBuilder, TArray<FRDGBufferSRVRef>& AttributeResources)
+	const EOptimusGroomExecDomain& AttributeGroup, FRDGBuilder& GraphBuilder, TArray<FRDGBufferSRVRef>& AttributeResources,
+		TArray<FRDGBufferSRVRef>& IndexMappings, TArray<uint32>& bHasIndexMapping)
 {
 	const FName GroupName = UOptimusGroomExecDataInterface::GetExecutionDomainName(AttributeGroup);
 	
@@ -530,6 +549,13 @@ static void CreateInternalBuffers(const UGroomAsset* GroomAsset, const UE::Groom
      			
      			if(FHairStrandsRestResource* RestResource = bIsGuide ? GroupInstance->Guides.RestResource : bIsStrand ? GroupInstance->Strands.RestResource : nullptr)
      			{
+     				// For now we only use index mapping for strands since there is no randomization for guides
+     				const bool bCurveMapping = ((AttributeGroup == EOptimusGroomExecDomain::Curve) && ((RestResource->BulkData.Header.Flags & FHairStrandsBulkData::DataFlags_HasCurveMapping) != 0));
+     				const bool bPointMapping = ((AttributeGroup == EOptimusGroomExecDomain::ControlPoint) && ((RestResource->BulkData.Header.Flags & FHairStrandsBulkData::DataFlags_HasPointMapping) != 0));
+     				
+     				IndexMappings.Add( RegisterAsSRV(GraphBuilder, bCurveMapping ? RestResource->CurveMappingBuffer : RestResource->PointMappingBuffer));
+     				bHasIndexMapping.Add(bCurveMapping || bPointMapping);
+     				
      				FRDGExternalBuffer* ExternalBuffer = RestResource->ExternalBuffers.Find(BufferName);
      				if(!ExternalBuffer)
      				{
@@ -561,67 +587,68 @@ static void CreateInternalBuffers(const UGroomAsset* GroomAsset, const UE::Groom
 }
 
 static void AddGroupResources(UGroomAsset* GroomAsset, const UE::Groom::Private::FGroupElements& GroupElements, const EOptimusGroomAttributeTypes GroomAttributeType,
-	const FName& GroomAttributeName,  const EOptimusGroomExecDomain& GroomAttributeGroup, FRDGBuilder& GraphBuilder, TArray<FRDGBufferSRVRef>& GroomAttributeResources)
+const FName& GroomAttributeName,  const EOptimusGroomExecDomain& GroomAttributeGroup, FRDGBuilder& GraphBuilder, TArray<FRDGBufferSRVRef>& GroomAttributeResources,
+	TArray<FRDGBufferSRVRef>& IndexMappingResources, TArray<uint32>& bHasIndexMapping)
 {
 	if(GroomAttributeType == EOptimusGroomAttributeTypes::Bool)
 	{
 		CreateInternalBuffers<bool>(GroomAsset, GroupElements, GroomAttributeName,
-			GroomAttributeGroup, GraphBuilder, GroomAttributeResources);
+			GroomAttributeGroup, GraphBuilder, GroomAttributeResources, IndexMappingResources,  bHasIndexMapping);
 	}
 	else if(GroomAttributeType == EOptimusGroomAttributeTypes::Int || GroomAttributeType == EOptimusGroomAttributeTypes::Uint)
 	{
 		CreateInternalBuffers<int32>(GroomAsset, GroupElements, GroomAttributeName,
-			GroomAttributeGroup, GraphBuilder, GroomAttributeResources);
+			GroomAttributeGroup, GraphBuilder, GroomAttributeResources, IndexMappingResources,  bHasIndexMapping);
 	}
 	else if(GroomAttributeType == EOptimusGroomAttributeTypes::IntVector2)
 	{
 		CreateInternalBuffers<FIntVector2>(GroomAsset, GroupElements, GroomAttributeName,
-			GroomAttributeGroup, GraphBuilder, GroomAttributeResources);
+			GroomAttributeGroup, GraphBuilder, GroomAttributeResources, IndexMappingResources,  bHasIndexMapping);
 	}
 	else if(GroomAttributeType == EOptimusGroomAttributeTypes::IntVector3)
 	{
 		CreateInternalBuffers<FIntVector3>(GroomAsset, GroupElements, GroomAttributeName,
-			GroomAttributeGroup, GraphBuilder, GroomAttributeResources);
+			GroomAttributeGroup, GraphBuilder, GroomAttributeResources, IndexMappingResources,  bHasIndexMapping);
 	}
 	else if(GroomAttributeType == EOptimusGroomAttributeTypes::IntVector4)
 	{
 		CreateInternalBuffers<FIntVector4>(GroomAsset, GroupElements, GroomAttributeName,
-			GroomAttributeGroup, GraphBuilder, GroomAttributeResources);
+			GroomAttributeGroup, GraphBuilder, GroomAttributeResources, IndexMappingResources,  bHasIndexMapping);
 	}
 	else if(GroomAttributeType == EOptimusGroomAttributeTypes::Float)
 	{
 		CreateInternalBuffers<float>(GroomAsset, GroupElements, GroomAttributeName,
-			GroomAttributeGroup, GraphBuilder, GroomAttributeResources);
+			GroomAttributeGroup, GraphBuilder, GroomAttributeResources, IndexMappingResources,  bHasIndexMapping);
 	}
 	else if(GroomAttributeType == EOptimusGroomAttributeTypes::Vector2)
 	{
 		CreateInternalBuffers<FVector2f>(GroomAsset, GroupElements,  GroomAttributeName,
-			GroomAttributeGroup, GraphBuilder, GroomAttributeResources);
+			GroomAttributeGroup, GraphBuilder, GroomAttributeResources, IndexMappingResources,  bHasIndexMapping);
 	}
 	else if(GroomAttributeType == EOptimusGroomAttributeTypes::Vector3 || GroomAttributeType == EOptimusGroomAttributeTypes::Rotator)
 	{
 		CreateInternalBuffers<FVector3f>(GroomAsset, GroupElements, GroomAttributeName,
-			GroomAttributeGroup, GraphBuilder, GroomAttributeResources);
+			GroomAttributeGroup, GraphBuilder, GroomAttributeResources, IndexMappingResources,  bHasIndexMapping);
 	}
 	else if(GroomAttributeType == EOptimusGroomAttributeTypes::Vector4)
 	{
 		CreateInternalBuffers<FVector4f>(GroomAsset, GroupElements, GroomAttributeName,
-			GroomAttributeGroup, GraphBuilder, GroomAttributeResources);
+			GroomAttributeGroup, GraphBuilder, GroomAttributeResources, IndexMappingResources,  bHasIndexMapping);
 	}
 	else if(GroomAttributeType == EOptimusGroomAttributeTypes::LinearColor)
 	{
 		CreateInternalBuffers<FLinearColor>(GroomAsset, GroupElements, GroomAttributeName,
-			GroomAttributeGroup, GraphBuilder, GroomAttributeResources);
+			GroomAttributeGroup, GraphBuilder, GroomAttributeResources, IndexMappingResources,  bHasIndexMapping);
 	}
 	else if(GroomAttributeType == EOptimusGroomAttributeTypes::Quat)
 	{
 		CreateInternalBuffers<FQuat4f>(GroomAsset, GroupElements, GroomAttributeName,
-			GroomAttributeGroup, GraphBuilder, GroomAttributeResources);
+			GroomAttributeGroup, GraphBuilder, GroomAttributeResources, IndexMappingResources,  bHasIndexMapping);
 	}
 	else if(GroomAttributeType == EOptimusGroomAttributeTypes::Transform || GroomAttributeType == EOptimusGroomAttributeTypes::Matrix3x4)
 	{
 		CreateInternalBuffers<FTransform3f>(GroomAsset, GroupElements, GroomAttributeName,
-			GroomAttributeGroup, GraphBuilder, GroomAttributeResources);
+			GroomAttributeGroup, GraphBuilder, GroomAttributeResources, IndexMappingResources,  bHasIndexMapping);
 	}
 }
 
@@ -630,7 +657,8 @@ void FOptimusGroomAttributeReadProviderProxy::CreateInternalBuffers(FRDGBuilder&
 	for (const TPair<UGroomAsset*,UE::Groom::Private::FGroupElements>& AssetPair : GroupElements)
 	{
 		AddGroupResources(AssetPair.Key, AssetPair.Value, GroomAttributeType, GroomAttributeName,
-				GroomAttributeGroup,GraphBuilder,AttributeValuesResources);
+				GroomAttributeGroup, GraphBuilder, AttributeValuesResources,
+				IndexMappingResources, bHasIndexMapping);
 	}
 }
 
@@ -639,6 +667,9 @@ void FOptimusGroomAttributeReadProviderProxy::AllocateResources(FRDGBuilder& Gra
 	if (!FallbackStructuredSRV) { FallbackStructuredSRV  = GraphBuilder.CreateSRV(GSystemTextures.GetDefaultStructuredBuffer(GraphBuilder, 16u)); }
 
 	AttributeValuesResources.Reset();
+	IndexMappingResources.Reset();
+	bHasIndexMapping.Reset();
+	
 	CreateInternalBuffers(GraphBuilder);
 }
 
@@ -658,10 +689,13 @@ void FOptimusGroomAttributeReadProviderProxy::GatherDispatchData(FDispatchData c
 			FParameters& Parameters = ParameterArray[InvocationIndex];
 			if (const FHairGroupInstance* GroupInstance = AssetPair.Value.GroupInstances[InstanceIndex])
 			{
-				Parameters.NumValues = AssetPair.Value.GroupOffsets[GroupIndex+1]-AssetPair.Value.GroupOffsets[GroupIndex];
+				Parameters.NumValues = NumElements.IsValidIndex(InvocationIndex) ? NumElements[InvocationIndex] : 0;
 				Parameters.BaseValueIndex = BaseValueIndex;
 				Parameters.ValueTypedData = AttributeValuesResources[InvocationIndex];
-
+				
+				Parameters.bHasIndexMapping = bHasIndexMapping[InvocationIndex];
+				Parameters.IndexMappingBuffer = IndexMappingResources[InvocationIndex];
+				
 				BaseValueIndex += Parameters.NumValues;
 			}
 			else
@@ -669,6 +703,9 @@ void FOptimusGroomAttributeReadProviderProxy::GatherDispatchData(FDispatchData c
 				Parameters.NumValues = 0;
 				Parameters.BaseValueIndex = BaseValueIndex;
 				Parameters.ValueTypedData = FallbackStructuredSRV;
+
+				Parameters.bHasIndexMapping = 0;
+				Parameters.IndexMappingBuffer = FallbackStructuredSRV;
 			}
 			++InstanceIndex;
 			++InvocationIndex;
diff --git a/Engine/Plugins/Runtime/HairStrands/Source/HairStrandsDeformer/Private/DeformerGroomAttributeRead.h b/Engine/Plugins/Runtime/HairStrands/Source/HairStrandsDeformer/Private/DeformerGroomAttributeRead.h
--- a/Engine/Plugins/Runtime/HairStrands/Source/HairStrandsDeformer/Private/DeformerGroomAttributeRead.h
+++ b/Engine/Plugins/Runtime/HairStrands/Source/HairStrandsDeformer/Private/DeformerGroomAttributeRead.h
@@ -145,10 +145,19 @@ private:
 
 	/** Groom assets group indices the proxy instance is linked to */
 	TArray<TPair<UGroomAsset*, UE::Groom::Private::FGroupElements>> GroupElements;
+
+	/** Num elements that will be used in deformer for each invocation (usually larger than the number of source elements) */
+	TArray<int32> NumElements;
 	
 	/** Attribute Resources used to dispatch CS on GPU */
 	TArray<FRDGBufferSRVRef> AttributeValuesResources;
 
+	/** Index mappings resources */
+	TArray<FRDGBufferSRVRef> IndexMappingResources;
+
+	/** Check if we need to use the index mappings */
+	TArray<uint32> bHasIndexMapping;
+
 	/** Attribute name */
 	FName GroomAttributeName = NAME_None;
 
diff --git a/Engine/Shaders/Private/HairStrands/HairStrandsPack.ush b/Engine/Shaders/Private/HairStrands/HairStrandsPack.ush
--- a/Engine/Shaders/Private/HairStrands/HairStrandsPack.ush
+++ b/Engine/Shaders/Private/HairStrands/HairStrandsPack.ush
@@ -523,6 +523,15 @@ uint ReadHairSourceCurveIndex(ByteAddressBuffer InBuffer, uint InIndex, bool bIs
 }
 
 uint ReadHairSourcePointIndex(ByteAddressBuffer InBuffer, uint InIndex, bool bIsBufferValid)
+{
+	uint Out = ~0;
+	if (bIsBufferValid)
+	{
+		Out = InBuffer.Load(InIndex * 4u);
+	}
+	return Out;
+}
+uint ReadHairSourceAttributeIndex(ByteAddressBuffer InBuffer, uint InIndex, bool bIsBufferValid)
 {
 	uint Out = ~0;
 	if (bIsBufferValid)
