Fix groom attribute deformer DI numvalues

diff --git a/Engine/Plugins/Runtime/HairStrands/Source/HairStrandsDeformer/Private/DeformerGroomAttributeRead.cpp b/Engine/Plugins/Runtime/HairStrands/Source/HairStrandsDeformer/Private/DeformerGroomAttributeRead.cpp
--- a/Engine/Plugins/Runtime/HairStrands/Source/HairStrandsDeformer/Private/DeformerGroomAttributeRead.cpp
+++ b/Engine/Plugins/Runtime/HairStrands/Source/HairStrandsDeformer/Private/DeformerGroomAttributeRead.cpp
@@ -689,14 +689,16 @@ void FOptimusGroomAttributeReadProviderProxy::GatherDispatchData(FDispatchData c
 			FParameters& Parameters = ParameterArray[InvocationIndex];
 			if (const FHairGroupInstance* GroupInstance = AssetPair.Value.GroupInstances[InstanceIndex])
 			{
-				Parameters.NumValues = NumElements.IsValidIndex(InvocationIndex) ? NumElements[InvocationIndex] : 0;
+				const int32 ElementCount = (NumElements.IsValidIndex(InvocationIndex) ? NumElements[InvocationIndex] : 0);
+
+				Parameters.NumValues = AssetPair.Value.GroupOffsets[InstanceIndex+1] - AssetPair.Value.GroupOffsets[InstanceIndex];
 				Parameters.BaseValueIndex = BaseValueIndex;
 				Parameters.ValueTypedData = AttributeValuesResources[InvocationIndex];
-				
+
 				Parameters.bHasIndexMapping = bHasIndexMapping[InvocationIndex];
 				Parameters.IndexMappingBuffer = IndexMappingResources[InvocationIndex];
-				
-				BaseValueIndex += Parameters.NumValues;
+
+				BaseValueIndex += ElementCount;
 			}
 			else
 			{
