Adapted hair spline skinning deformer to new GroomAttribute data interface
#virtualized

diff --git a/Engine/Plugins/Runtime/HairStrands/Source/HairStrandsDataflow/Private/BuildGroomSplineSkinningNode.cpp b/Engine/Plugins/Runtime/HairStrands/Source/HairStrandsDataflow/Private/BuildGroomSplineSkinningNode.cpp
--- a/Engine/Plugins/Runtime/HairStrands/Source/HairStrandsDataflow/Private/BuildGroomSplineSkinningNode.cpp
+++ b/Engine/Plugins/Runtime/HairStrands/Source/HairStrandsDataflow/Private/BuildGroomSplineSkinningNode.cpp
@@ -289,24 +289,27 @@ namespace UE::Groom::Private
 					const TArray<int32> VertexIndices = MeshFacade.GetVertexIndices(GeometryIndex);
 					for (const int32& VertexIndex : VertexIndices)
 					{
-						if(!bValidSelection || (bValidSelection && (VertexSelection.IsSelected(VertexIndex))))
+						if ((*SplineBones)[VertexIndex][0] != INDEX_NONE && (*SplineBones)[VertexIndex][1] != INDEX_NONE)
 						{
-							const float SplineParam = (*SplineParams)[VertexIndex];
-							const int32 SegmentIndex = trunc(SplineParam);
-							const float SegmentParam = SplineParam - SegmentIndex;
-							
-							if (SegmentIndex < (*SplineBones)[VertexIndex][1])
-							{
-								BoneIndices = {SegmentIndex, SegmentIndex + 1};
-								BoneWeights = {1.0f - SegmentParam, SegmentParam};
-							}
-							else
+							if (!bValidSelection || (bValidSelection && (VertexSelection.IsSelected(VertexIndex))))
 							{
-								BoneIndices = {SegmentIndex};
-								BoneWeights = {1.0f};
-							}
+								const float SplineParam = (*SplineParams)[VertexIndex];
+								const int32 SegmentIndex = trunc(SplineParam);
+								const float SegmentParam = SplineParam - SegmentIndex;
+
+								if (SegmentIndex < (*SplineBones)[VertexIndex][1])
+								{
+									BoneIndices = {SegmentIndex, SegmentIndex + 1};
+									BoneWeights = {1.0f - SegmentParam, SegmentParam};
+								}
+								else
+								{
+									BoneIndices = {SegmentIndex};
+									BoneWeights = {1.0f};
+								}
 
-							SkinningFacade.ModifyBoneWeight(VertexIndex, BoneIndices, BoneWeights);
+								SkinningFacade.ModifyBoneWeight(VertexIndex, BoneIndices, BoneWeights);
+							}
 						}
 					}
 				}
