Fix crash when cooking and running a project containing Groom asset but having r.HairStrands.Enable set to 0.

diff --git a/Engine/Plugins/Runtime/HairStrands/Source/HairStrandsCore/Private/GroomAsset.cpp b/Engine/Plugins/Runtime/HairStrands/Source/HairStrandsCore/Private/GroomAsset.cpp
--- a/Engine/Plugins/Runtime/HairStrands/Source/HairStrandsCore/Private/GroomAsset.cpp
+++ b/Engine/Plugins/Runtime/HairStrands/Source/HairStrandsCore/Private/GroomAsset.cpp
@@ -1319,6 +1319,13 @@ void UGroomAsset::PostLoad()
 
 	Super::PostLoad();
 
+	// Early out if hair strands are disabled
+	const bool bIsHairEnabled = IsGroomEnabled();
+	if (!bIsHairEnabled)
+	{
+		return;
+	}
+
 	// Compute a hash of the Groom asset fullname for finding unique groom during LOD selection/streaming
 	AssetNameHash = GetAssetNameHash(GetFullName());
 
diff --git a/Engine/Plugins/Runtime/HairStrands/Source/HairStrandsCore/Private/GroomBindingAsset.cpp b/Engine/Plugins/Runtime/HairStrands/Source/HairStrandsCore/Private/GroomBindingAsset.cpp
--- a/Engine/Plugins/Runtime/HairStrands/Source/HairStrandsCore/Private/GroomBindingAsset.cpp
+++ b/Engine/Plugins/Runtime/HairStrands/Source/HairStrandsCore/Private/GroomBindingAsset.cpp
@@ -270,6 +270,14 @@ void UGroomBindingAsset::Serialize(FArchive& Ar)
 
 	Super::Serialize(Ar);
 	Ar.UsingCustomVersion(FAnimObjectVersion::GUID);
+
+	// Early out if hair strands are disabled
+	const bool bIsHairEnabled = IsGroomEnabled();
+	if (!bIsHairEnabled)
+	{
+		return;
+	}
+
 #if WITH_EDITOR
 	// When using editor:
 	// * The header are loaded in CacheDerivedData(), and the data are streamed from DDC
@@ -547,6 +555,13 @@ void UGroomBindingAsset::PostLoad()
 
 	Super::PostLoad();
 
+	// Early out if hair strands are disabled
+	const bool bIsHairEnabled = IsGroomEnabled();
+	if (!bIsHairEnabled)
+	{
+		return;
+	}
+
 	// Compute a hash of the Groom asset fullname for finding unique groom during LOD selection/streaming
 	AssetNameHash = GetAssetNameHash(GetFullName());
 
@@ -1557,6 +1572,13 @@ static UGroomBindingAsset::FCachedCookedPlatformData* FindCachedCookedPlatformDa
 
 void UGroomBindingAsset::BeginCacheForCookedPlatformData(const ITargetPlatform* TargetPlatform)
 {
+	// Early out if hair strands are disabled
+	const bool bIsHairEnabled = IsGroomEnabled();
+	if (!bIsHairEnabled)
+	{
+		return;
+	}
+
 	// Finish any async compilation of the editor target before building other platforms
 	if (IsCompiling())
 	{
