Fix non-determinism in hair builder.

#preffight skip

diff --git a/Engine/Plugins/Runtime/HairStrands/Source/HairStrandsCore/Private/GroomBuilder.cpp b/Engine/Plugins/Runtime/HairStrands/Source/HairStrandsCore/Private/GroomBuilder.cpp
--- a/Engine/Plugins/Runtime/HairStrands/Source/HairStrandsCore/Private/GroomBuilder.cpp
+++ b/Engine/Plugins/Runtime/HairStrands/Source/HairStrandsCore/Private/GroomBuilder.cpp
@@ -3987,6 +3987,8 @@ static void BuildClusterBulkData(
 			Packed.Screen0 = Packed.Screen0 | (Quantize8Bits(Info.ScreenSize[1], InvScreenScale, ScreenOffset) << 8);
 			Packed.Screen0 = Packed.Screen0 | (Quantize8Bits(Info.ScreenSize[2], InvScreenScale, ScreenOffset) << 16);
 			Packed.Screen0 = Packed.Screen0 | (Quantize8Bits(Info.ScreenSize[3], InvScreenScale, ScreenOffset) << 24);
+
+			Packed.Screen1 = 0;
 			Packed.Screen1 = Packed.Screen1 | (Quantize8Bits(Info.ScreenSize[4], InvScreenScale, ScreenOffset) << 0);
 			Packed.Screen1 = Packed.Screen1 | (Quantize8Bits(Info.ScreenSize[5], InvScreenScale, ScreenOffset) << 8);
 			Packed.Screen1 = Packed.Screen1 | (Quantize8Bits(Info.ScreenSize[6], InvScreenScale, ScreenOffset) << 16);
@@ -3996,6 +3998,8 @@ static void BuildClusterBulkData(
 			Packed.Radius0 = Packed.Radius0 | (Quantize8Bits(Info.RadiusScale[1], InvRadiusScale, RadiusOffset) << 8);
 			Packed.Radius0 = Packed.Radius0 | (Quantize8Bits(Info.RadiusScale[2], InvRadiusScale, RadiusOffset) << 16);
 			Packed.Radius0 = Packed.Radius0 | (Quantize8Bits(Info.RadiusScale[3], InvRadiusScale, RadiusOffset) << 24);
+
+			Packed.Radius1 = 0;
 			Packed.Radius1 = Packed.Radius1 | (Quantize8Bits(Info.RadiusScale[4], InvRadiusScale, RadiusOffset) << 0);
 			Packed.Radius1 = Packed.Radius1 | (Quantize8Bits(Info.RadiusScale[5], InvRadiusScale, RadiusOffset) << 8);
 			Packed.Radius1 = Packed.Radius1 | (Quantize8Bits(Info.RadiusScale[6], InvRadiusScale, RadiusOffset) << 16);
