Remove the template picker while importing groom asset

diff --git a/Engine/Plugins/Runtime/HairStrands/Source/HairStrandsEditor/Private/HairStrandsFactory.cpp b/Engine/Plugins/Runtime/HairStrands/Source/HairStrandsEditor/Private/HairStrandsFactory.cpp
--- a/Engine/Plugins/Runtime/HairStrands/Source/HairStrandsEditor/Private/HairStrandsFactory.cpp
+++ b/Engine/Plugins/Runtime/HairStrands/Source/HairStrandsEditor/Private/HairStrandsFactory.cpp
@@ -457,11 +457,6 @@ UObject* UHairStrandsFactory::FactoryCreateFile(UClass* InClass, UObject* InPare
 			CurrentAsset = GroomCaches[0];
 		}
 	} 
-	
-	if(UGroomAsset* GroomAsset = Cast<UGroomAsset>(CurrentAsset))
-	{
-		UE::Groom::BuildGroomDataflowAsset(GroomAsset);
-	}
 
 	return CurrentAsset;
 }
