Enable groom RT geometry with r.hairstrands.raytracing is enabled.

diff --git a/Engine/Plugins/Runtime/HairStrands/Source/HairStrandsCore/Private/GroomComponent.cpp b/Engine/Plugins/Runtime/HairStrands/Source/HairStrandsCore/Private/GroomComponent.cpp
--- a/Engine/Plugins/Runtime/HairStrands/Source/HairStrandsCore/Private/GroomComponent.cpp
+++ b/Engine/Plugins/Runtime/HairStrands/Source/HairStrandsCore/Private/GroomComponent.cpp
@@ -3234,7 +3234,7 @@ void UGroomComponent::InitResources(bool bIsBindingReloading)
 				}
 
 				#if RHI_RAYTRACING
-				if (IsRayTracingAllowed() && bVisibleInRayTracing)
+				if (IsHairRayTracingEnabled() && IsRayTracingAllowed() && bVisibleInRayTracing)
 				{
 					if (bNeedDeformedPositions)
 					{
@@ -3306,7 +3306,7 @@ void UGroomComponent::InitResources(bool bIsBindingReloading)
 				}
 
 				#if RHI_RAYTRACING
-				if (IsRayTracingAllowed() && bVisibleInRayTracing)
+				if (IsHairRayTracingEnabled() && IsRayTracingAllowed() && bVisibleInRayTracing)
 				{
 					if (bNeedDeformedPositions)
 					{
diff --git a/Engine/Plugins/Runtime/HairStrands/Source/HairStrandsCore/Private/GroomManager.cpp b/Engine/Plugins/Runtime/HairStrands/Source/HairStrandsCore/Private/GroomManager.cpp
--- a/Engine/Plugins/Runtime/HairStrands/Source/HairStrandsCore/Private/GroomManager.cpp
+++ b/Engine/Plugins/Runtime/HairStrands/Source/HairStrandsCore/Private/GroomManager.cpp
@@ -2178,7 +2178,7 @@ static void RunHairStrandsInterpolation_Cards(
 
 	// Cards only - Build/update RT geometry
 	#if RHI_RAYTRACING
-	if (IsRayTracingEnabled() && !bIsReflectionCapture)
+	if (IsHairRayTracingEnabled() && IsRayTracingEnabled() && !bIsReflectionCapture)
 	{
 		for (uint32 InstanceIndex : CardInstances)
 		{
@@ -2216,7 +2216,7 @@ static void RunHairStrandsInterpolation_Cards(
 
 	// Meshes only - Build/update RT geometry
 	#if RHI_RAYTRACING
-	if (IsRayTracingEnabled() && !bIsReflectionCapture)
+	if (IsHairRayTracingEnabled() && IsRayTracingEnabled() && !bIsReflectionCapture)
 	{
 		for (uint32 InstanceIndex : MeshInstances)
 		{
