Groom : set Groom asset to collection node to use Guides by default

diff --git a/Engine/Plugins/Runtime/HairStrands/Source/HairStrandsDataflow/Public/GetGroomAssetNode.h b/Engine/Plugins/Runtime/HairStrands/Source/HairStrandsDataflow/Public/GetGroomAssetNode.h
--- a/Engine/Plugins/Runtime/HairStrands/Source/HairStrandsDataflow/Public/GetGroomAssetNode.h
+++ b/Engine/Plugins/Runtime/HairStrands/Source/HairStrandsDataflow/Public/GetGroomAssetNode.h
@@ -101,7 +101,7 @@ private:
 
 	/** Type of curves to use to fill the groom collection (guides/strands) */
 	UPROPERTY(EditAnywhere, Category = "Groom", meta = (DisplayName = "Curves Type"), meta = (DataflowInput))
-	EGroomCollectionType CurvesType = EGroomCollectionType::Strands;
+	EGroomCollectionType CurvesType = EGroomCollectionType::Guides;
 
 	/** Curves thickness for geometry generation */
 	UPROPERTY(EditAnywhere, Category = "Geometry", meta = (DataflowInput, DisplayName = "Curves Thickness"))
