Groom editor : fix crash when changing template within the dataflow editor
- Done by preventing the Dataflow editor from being closed and re-opened when the dataflow asset is changed
- Simply use the refresh code path the editor already supports

diff --git a/Engine/Plugins/Experimental/Dataflow/Source/DataflowEditor/Private/Dataflow/DataflowEditorToolkit.cpp b/Engine/Plugins/Experimental/Dataflow/Source/DataflowEditor/Private/Dataflow/DataflowEditorToolkit.cpp
--- a/Engine/Plugins/Experimental/Dataflow/Source/DataflowEditor/Private/Dataflow/DataflowEditorToolkit.cpp
+++ b/Engine/Plugins/Experimental/Dataflow/Source/DataflowEditor/Private/Dataflow/DataflowEditorToolkit.cpp
@@ -1427,6 +1427,12 @@ void FDataflowEditorToolkit::OnPropertyValueChanged(const FPropertyChangedEvent&
 	}
 }
 
+void FDataflowEditorToolkit::OnDataflowAssetChanged()
+{
+	const FPropertyChangedEvent PropertyChangedEvent(FDataflowInstance::StaticStruct()->FindPropertyByName(FDataflowInstance::GetDataflowAssetPropertyName()));
+	OnAssetPropertyValueChanged(PropertyChangedEvent);
+}
+
 void FDataflowEditorToolkit::OnAssetPropertyValueChanged(const FPropertyChangedEvent& PropertyChangedEvent)
 {
 	if (const TObjectPtr<UDataflowBaseContent>& EditorContent = GetEditorContent())
diff --git a/Engine/Plugins/Experimental/Dataflow/Source/DataflowEditor/Public/Dataflow/DataflowEditorToolkit.h b/Engine/Plugins/Experimental/Dataflow/Source/DataflowEditor/Public/Dataflow/DataflowEditorToolkit.h
--- a/Engine/Plugins/Experimental/Dataflow/Source/DataflowEditor/Public/Dataflow/DataflowEditorToolkit.h
+++ b/Engine/Plugins/Experimental/Dataflow/Source/DataflowEditor/Public/Dataflow/DataflowEditorToolkit.h
@@ -109,6 +109,8 @@ public:
 
 	UE_API const FString& GetDebugDrawOverlayString() const;
 	
+	UE_API void OnDataflowAssetChanged();
+
 	/** Get the toolkit evaluation mode */
 	const EDataflowEditorEvaluationMode& GetEvaluationMode() const { return EvaluationMode; }
 
diff --git a/Engine/Plugins/Runtime/HairStrands/Source/HairStrandsEditor/Private/GroomActions.cpp b/Engine/Plugins/Runtime/HairStrands/Source/HairStrandsEditor/Private/GroomActions.cpp
--- a/Engine/Plugins/Runtime/HairStrands/Source/HairStrandsEditor/Private/GroomActions.cpp
+++ b/Engine/Plugins/Runtime/HairStrands/Source/HairStrandsEditor/Private/GroomActions.cpp
@@ -40,6 +40,7 @@
 #include "HairStrandsFactory.h"
 #include "ToolMenus.h"
 #include "Dataflow/DataflowEditor.h"
+#include "Dataflow/DataflowEditorToolkit.h"
 #include "Dialog/SMessageDialog.h"
 
 #include "ThumbnailRendering/SceneThumbnailInfo.h"
@@ -70,14 +71,39 @@ EAssetCommandResult UAssetDefinition_GroomAsset::OpenAssets(const FAssetOpenArgs
 				GroomAsset->GetOnCreateGroomDataflow().AddLambda(
 					[](UGroomAsset* GroomAsset)
 					{
-						UE::Groom::BuildGroomDataflowAsset(GroomAsset);
-						if (UAssetEditorSubsystem* AssetEditorSubsystem = GEditor->GetEditorSubsystem<UAssetEditorSubsystem>())
+						if (UE::Groom::BuildGroomDataflowAsset(GroomAsset))
 						{
-							// Close all the opened editors
-							AssetEditorSubsystem->CloseAllEditorsForAsset(GroomAsset);
-							
-							// Open a new editor to be able to edit the groom asset with dataflow
-							AssetEditorSubsystem->OpenEditorForAsset(GroomAsset);
+							if (UAssetEditorSubsystem* AssetEditorSubsystem = GEditor->GetEditorSubsystem<UAssetEditorSubsystem>())
+							{
+								// note : we cannot use CloseAllEditorsForAsset for the Dataflow Editor , the persona based editor seem to have an issue with calling this
+								// and will crash if because of an order of destruction issue with the toolkit. the path to solve it is unclear and may be risky
+								// As an alternative let's not reload the dataflow editor and leverage the in place dataflow asset change feature that already exists
+								bool bDataflowEditorAlreadyOpened = false;
+
+								TArray<IAssetEditorInstance*> EditorInstances = AssetEditorSubsystem->FindEditorsForAssetAndSubObjects(GroomAsset);
+								for (IAssetEditorInstance* EditorInstance : EditorInstances)
+								{
+									if (EditorInstance->GetEditorName() != "DataflowEditor")
+									{
+										EditorInstance->CloseWindow(EAssetEditorCloseReason::CloseAllEditorsForAsset);
+									}
+									else if (!bDataflowEditorAlreadyOpened && EditorInstance->GetEditorName() == "DataflowEditor")
+									{
+										if (FDataflowEditorToolkit* DataflowEditorToolkit = static_cast<FDataflowEditorToolkit*>(EditorInstance))
+										{
+											DataflowEditorToolkit->OnDataflowAssetChanged();
+										}
+										bDataflowEditorAlreadyOpened = true;
+									}
+									AssetEditorSubsystem->OnAssetEditorRequestClose().Broadcast(GroomAsset, EAssetEditorCloseReason::CloseAllEditorsForAsset);
+								}
+
+								if (!bDataflowEditorAlreadyOpened)
+								{
+									// Open a new editor to be able to edit the groom asset with dataflow
+									AssetEditorSubsystem->OpenEditorForAsset(GroomAsset);
+								}
+							}
 						}
 					} );
 			}
diff --git a/Engine/Plugins/Runtime/HairStrands/Source/HairStrandsEditor/Private/HairStrandsFactory.cpp b/Engine/Plugins/Runtime/HairStrands/Source/HairStrandsEditor/Private/HairStrandsFactory.cpp
--- a/Engine/Plugins/Runtime/HairStrands/Source/HairStrandsEditor/Private/HairStrandsFactory.cpp
+++ b/Engine/Plugins/Runtime/HairStrands/Source/HairStrandsEditor/Private/HairStrandsFactory.cpp
@@ -167,7 +167,7 @@ namespace UE::Groom
 		FGroomDataflowTemplateManager::GetInstance().RemoveDataflowTemplatePath(DataflowTemplatePath);
 	}
 
-	void BuildGroomDataflowAsset(UGroomAsset* GroomAsset)
+	bool BuildGroomDataflowAsset(UGroomAsset* GroomAsset)
 	{
 		// Create a new Dataflow asset
 		const FString DataflowPath = FPackageName::GetLongPackagePath(GroomAsset->GetOutermost()->GetName());
@@ -210,7 +210,10 @@ namespace UE::Groom
 
 			// Set the Dataflow to the groom asset
 			GroomAsset->GetDataflowSettings().SetDataflowAsset(Dataflow);
+
+			return true;
 		}
+		return false;
 	}
 }
 
diff --git a/Engine/Plugins/Runtime/HairStrands/Source/HairStrandsEditor/Public/HairStrandsFactory.h b/Engine/Plugins/Runtime/HairStrands/Source/HairStrandsEditor/Public/HairStrandsFactory.h
--- a/Engine/Plugins/Runtime/HairStrands/Source/HairStrandsEditor/Public/HairStrandsFactory.h
+++ b/Engine/Plugins/Runtime/HairStrands/Source/HairStrandsEditor/Public/HairStrandsFactory.h
@@ -47,7 +47,7 @@ namespace UE::Groom
 	HAIRSTRANDSEDITOR_API void UnregisterGroomDataflowTemplatePath(const FString& DataflowTemplatePath);
 
 	/** Build groom dataflow asset from templates */
-	HAIRSTRANDSEDITOR_API void BuildGroomDataflowAsset(UGroomAsset* GroomAsset);
+	HAIRSTRANDSEDITOR_API bool BuildGroomDataflowAsset(UGroomAsset* GroomAsset);
 }
 
 /**
