Dataflow : add geometry collection node category  by default for groom assets
- gives access to all the geometry collection compatible nodes and also to selection nodes

diff --git a/Engine/Plugins/Runtime/HairStrands/Source/HairStrandsDataflow/Private/HairStrandsDataflowModule.cpp b/Engine/Plugins/Runtime/HairStrands/Source/HairStrandsDataflow/Private/HairStrandsDataflowModule.cpp
--- a/Engine/Plugins/Runtime/HairStrands/Source/HairStrandsDataflow/Private/HairStrandsDataflowModule.cpp
+++ b/Engine/Plugins/Runtime/HairStrands/Source/HairStrandsDataflow/Private/HairStrandsDataflowModule.cpp
@@ -57,6 +57,7 @@ PRAGMA_ENABLE_DEPRECATION_WARNINGS
 
 	// register node category for groom asset
 	UE_DATAFLOW_REGISTER_CATEGORY_FORASSET_TYPE("Groom", UGroomAsset);
+	UE_DATAFLOW_REGISTER_CATEGORY_FORASSET_TYPE("GeometryCollection", UGroomAsset);
 }
 
 void FHairStrandsDataflowModule::ShutdownModule()
