Use of the groom description source/edit

diff --git a/Engine/Plugins/Runtime/HairStrands/Source/HairStrandsCore/Private/GroomAsset.cpp b/Engine/Plugins/Runtime/HairStrands/Source/HairStrandsCore/Private/GroomAsset.cpp
--- a/Engine/Plugins/Runtime/HairStrands/Source/HairStrandsCore/Private/GroomAsset.cpp
+++ b/Engine/Plugins/Runtime/HairStrands/Source/HairStrandsCore/Private/GroomAsset.cpp
@@ -860,18 +860,30 @@ void UGroomAsset::Serialize(FArchive& Ar)
 #if WITH_EDITORONLY_DATA
 		else
 		{
-			// When serializing data for editor, serialize the HairDescription as bulk data
-			// The computed groom data is fetched from the Derived Data Cache
-			if (!HairDescriptionBulkData[EHairDescriptionType::Source])
+			auto SerializeHairDescription = [this, &Ar](const EHairDescriptionType DescriptionType)
 			{
-				// When loading, bulk data can be null so instantiate a new one to serialize into
-				HairDescriptionBulkData[EHairDescriptionType::Source] = MakeUnique<FHairDescriptionBulkData>();
-			}
+				const uint8 DescriptionIndex = FMath::Clamp(static_cast<uint8>(DescriptionType), 0, 1);
+					
+				// When serializing data for editor, serialize the HairDescription as bulk data
+				// The computed groom data is fetched from the Derived Data Cache
+				if (!HairDescriptionBulkData[DescriptionIndex])
+				{
+					// When loading, bulk data can be null so instantiate a new one to serialize into
+					HairDescriptionBulkData[DescriptionIndex] = MakeUnique<FHairDescriptionBulkData>();
+				}
 
-			// We don' serialize the hair description if it is a transactional object (i.e., for undo/redo)
-			if (!Ar.IsTransacting())
+				// We don' serialize the hair description if it is a transactional object (i.e., for undo/redo)
+				if (!Ar.IsTransacting())
+				{
+					HairDescriptionBulkData[DescriptionIndex]->Serialize(Ar, this);
+				}
+			};
+
+			SerializeHairDescription(EHairDescriptionType::Source);
+			
+			if (GetHairDescriptionType() == EHairDescriptionType::Edit)
 			{
-				HairDescriptionBulkData[EHairDescriptionType::Source]->Serialize(Ar, this);
+				SerializeHairDescription(EHairDescriptionType::Edit);
 			}
 		}
 #endif // WITH_EDITORONLY_DATA
@@ -1207,7 +1219,7 @@ void UGroomAsset::UpdateHairGroupsInfo()
 	#if WITH_EDITORONLY_DATA
 		// Transfer group name from the hair description group to the UGroomAsset.
 		// This pass is used right after an asset is imported.
-		if (const FHairDescriptionGroups* LocalHairDescriptionGroups = CachedHairDescriptionGroups[EHairDescriptionType::Source].Get())
+		if (const FHairDescriptionGroups* LocalHairDescriptionGroups = CachedHairDescriptionGroups[GetHairDescriptionIndex()].Get())
 		{
 			if (LocalHairDescriptionGroups->HairGroups.IsValidIndex(GroupIndex))
 			{
@@ -1334,7 +1346,7 @@ void UGroomAsset::PostLoad()
 	const bool bIsLegacyAsset = GetHairGroupsInterpolation().Num() == 0;
 	if (bIsLegacyAsset)
 	{
-		if (HairDescriptionBulkData[HairDescriptionType])
+		if (HairDescriptionBulkData[GetHairDescriptionIndex()])
 		{
 			const FHairDescriptionGroups& LocalHairDescriptionGroups = GetHairDescriptionGroups();
 
@@ -2059,7 +2071,7 @@ const TArray<UAssetUserData*>* UGroomAsset::GetAssetUserDataArray() const
 bool UGroomAsset::CanRebuildFromDescription() const
 {
 #if WITH_EDITORONLY_DATA
-	return HairDescriptionBulkData[HairDescriptionType].IsValid() && !HairDescriptionBulkData[HairDescriptionType]->IsEmpty();
+	return HairDescriptionBulkData[GetHairDescriptionIndex()].IsValid() && !HairDescriptionBulkData[GetHairDescriptionIndex()]->IsEmpty();
 #else
 	return false;
 #endif
@@ -2467,7 +2479,7 @@ FString UGroomAsset::GetDerivedDataKeyForStrands(uint32 GroupIndex)
 
 	// If simulation or global interpolation is enabled, then interpolation data are required. Otherwise they can be skipped
 	const bool bNeedInterpolationData = NeedsInterpolationData(GroupIndex);
-	const FString KeySuffix = GroomDerivedDataCacheUtils::BuildStrandsDerivedDataKeySuffix(GroupIndex, InterpolationSettings, LODSettings, PhysicsSettings, bNeedInterpolationData, HairDescriptionBulkData[HairDescriptionType].Get());
+	const FString KeySuffix = GroomDerivedDataCacheUtils::BuildStrandsDerivedDataKeySuffix(GroupIndex, InterpolationSettings, LODSettings, PhysicsSettings, bNeedInterpolationData, HairDescriptionBulkData[GetHairDescriptionIndex()].Get());
 	const FString DerivedDataKey = GroomDerivedDataCacheUtils::BuildGroomDerivedDataKey(KeySuffix);
 
 	return DerivedDataKey;
@@ -2477,69 +2489,76 @@ void UGroomAsset::CommitHairDescription(FHairDescription&& InHairDescription, EH
 {
 	TRACE_CPUPROFILER_EVENT_SCOPE(UGroomAsset::CommitHairDescription);
 
-	HairDescriptionType = Type;
+	SetHairDescriptionType(Type);
 
 	UE::TUniqueLock LockScope(InternalLock);
 
-	CachedHairDescription[HairDescriptionType] = MakeUnique<FHairDescription>(InHairDescription);
+	CachedHairDescription[GetHairDescriptionIndex()] = MakeUnique<FHairDescription>(InHairDescription);
 
-	if (!HairDescriptionBulkData[HairDescriptionType])
+	if (!HairDescriptionBulkData[GetHairDescriptionIndex()])
 	{
-		HairDescriptionBulkData[HairDescriptionType] = MakeUnique<FHairDescriptionBulkData>();
+		HairDescriptionBulkData[GetHairDescriptionIndex()] = MakeUnique<FHairDescriptionBulkData>();
 	}
 
 	// Update the cached hair description groups with the new hair description data
-	CachedHairDescriptionGroups[HairDescriptionType] = MakeUnique<FHairDescriptionGroups>();
-	FGroomBuilder::BuildHairDescriptionGroups(*CachedHairDescription[HairDescriptionType], *CachedHairDescriptionGroups[HairDescriptionType]);
+	CachedHairDescriptionGroups[GetHairDescriptionIndex()] = MakeUnique<FHairDescriptionGroups>();
+	FGroomBuilder::BuildHairDescriptionGroups(*CachedHairDescription[GetHairDescriptionIndex()], *CachedHairDescriptionGroups[GetHairDescriptionIndex()]);
 
-	HairDescriptionBulkData[HairDescriptionType]->SaveHairDescription(*CachedHairDescription[HairDescriptionType]);
+	HairDescriptionBulkData[GetHairDescriptionIndex()]->SaveHairDescription(*CachedHairDescription[GetHairDescriptionIndex()]);
 }
 
 const FHairDescriptionGroups& UGroomAsset::GetHairDescriptionGroups() const
 {
 	TRACE_CPUPROFILER_EVENT_SCOPE(UGroomAsset::GetHairDescriptionGroups);
 
-	check(HairDescriptionBulkData[HairDescriptionType]);
+	check(HairDescriptionBulkData[GetHairDescriptionIndex()]);
 	
 	{
 		UE::TUniqueLock LockScope(InternalLock);
 
-		if (!CachedHairDescription[HairDescriptionType])
+		if (!CachedHairDescription[GetHairDescriptionIndex()])
 		{
 			TUniquePtr<FHairDescription> HairDescription = MakeUnique<FHairDescription>();
-			HairDescriptionBulkData[HairDescriptionType]->LoadHairDescription(*HairDescription);
+			HairDescriptionBulkData[GetHairDescriptionIndex()]->LoadHairDescription(*HairDescription);
 
 			// Since we don't take read-lock, we need to make sure all memory is visible inside newly created object before publishing it.
 			FPlatformMisc::MemoryBarrier();
 
-			CachedHairDescription[HairDescriptionType] = MoveTemp(HairDescription);
+			CachedHairDescription[GetHairDescriptionIndex()] = MoveTemp(HairDescription);
 		}
 
-		if (!CachedHairDescriptionGroups[HairDescriptionType])
+		if (!CachedHairDescriptionGroups[GetHairDescriptionIndex()])
 		{
 			TUniquePtr<FHairDescriptionGroups> HairDescriptionGroups = MakeUnique<FHairDescriptionGroups>();
-			FGroomBuilder::BuildHairDescriptionGroups(*CachedHairDescription[HairDescriptionType], *HairDescriptionGroups, true /*bAllowAddEndControlPoint*/);
+			FGroomBuilder::BuildHairDescriptionGroups(*CachedHairDescription[GetHairDescriptionIndex()], *HairDescriptionGroups, true /*bAllowAddEndControlPoint*/);
 			
 			// Since we don't take read-lock, we need to make sure all memory is visible inside newly created object before publishing it.
 			FPlatformMisc::MemoryBarrier();
 
-			CachedHairDescriptionGroups[HairDescriptionType] = MoveTemp(HairDescriptionGroups);
+			CachedHairDescriptionGroups[GetHairDescriptionIndex()] = MoveTemp(HairDescriptionGroups);
 		}
 	}
 
-	return *CachedHairDescriptionGroups[HairDescriptionType];
+	return *CachedHairDescriptionGroups[GetHairDescriptionIndex()];
 }
 
-FHairDescription UGroomAsset::GetHairDescription() const
+FHairDescription UGroomAsset::LoadHairDescription(const EHairDescriptionType DescriptionType) const
 {
+	const uint8 DescriptionIndex = FMath::Clamp(static_cast<uint8>(DescriptionType), 0, 1);
+	
 	FHairDescription OutHairDescription;
-	if (HairDescriptionBulkData[HairDescriptionType])
+	if (HairDescriptionBulkData[DescriptionIndex])
 	{
-		HairDescriptionBulkData[HairDescriptionType]->LoadHairDescription(OutHairDescription);
+		HairDescriptionBulkData[DescriptionIndex]->LoadHairDescription(OutHairDescription);
 	}
 	return MoveTemp(OutHairDescription);
 }
 
+FHairDescription UGroomAsset::GetHairDescription() const
+{
+	return LoadHairDescription(GetHairDescriptionType());
+}
+
 bool UGroomAsset::CacheDerivedDatas()
 {
 	bRetryLoadFromGameThread = false;
@@ -2614,7 +2633,7 @@ bool UGroomAsset::CacheDerivedData(uint32 GroupIndex)
 
 bool UGroomAsset::CacheStrandsData(uint32 GroupIndex, FString& OutDerivedDataKey)
 {
-	if (!HairDescriptionBulkData[HairDescriptionType] || HairDescriptionBulkData[HairDescriptionType]->IsEmpty())
+	if (!HairDescriptionBulkData[GetHairDescriptionIndex()] || HairDescriptionBulkData[GetHairDescriptionIndex()]->IsEmpty())
 	{
 		// Groom having only cards/meshes (i.e., not containing imported strands) don't have hair description
 		bool bHasStrands = false;
diff --git a/Engine/Plugins/Runtime/HairStrands/Source/HairStrandsCore/Private/GroomEdit.cpp b/Engine/Plugins/Runtime/HairStrands/Source/HairStrandsCore/Private/GroomEdit.cpp
--- a/Engine/Plugins/Runtime/HairStrands/Source/HairStrandsCore/Private/GroomEdit.cpp
+++ b/Engine/Plugins/Runtime/HairStrands/Source/HairStrandsCore/Private/GroomEdit.cpp
@@ -215,14 +215,14 @@ void ConvertFromGroomAsset(UGroomAsset* InGroom, FEditableGroom* OutGroom, const
 	FEditableGroom& Out = *OutGroom;
 
 	FHairDescriptionGroups HairDescriptionGroups;
-	FGroomBuilder::BuildHairDescriptionGroups(In.GetHairDescription(), HairDescriptionGroups, bAllowAddEndControlPoint);
+	FGroomBuilder::BuildHairDescriptionGroups(In.LoadHairDescription(EHairDescriptionType::Source), HairDescriptionGroups, bAllowAddEndControlPoint);
 
 	const uint32 GroupCount = In.GetNumHairGroups();
-	Out.Groups.SetNum(GroupCount);
+	if(GroupCount == HairDescriptionGroups.HairGroups.Num() && GroupCount == In.GetHairGroupsInterpolation().Num())
+	{ 
+		Out.Groups.SetNum(GroupCount);
 
-	for (uint32 GroupIndex = 0; GroupIndex < GroupCount; ++GroupIndex)
-	{
-		if (HairDescriptionGroups.HairGroups.IsValidIndex(GroupIndex))
+		for (uint32 GroupIndex = 0; GroupIndex < GroupCount; ++GroupIndex)
 		{
 			FHairGroupInfo GroupInfo;
 			FHairStrandsDatas StrandsData;
@@ -237,6 +237,10 @@ void ConvertFromGroomAsset(UGroomAsset* InGroom, FEditableGroom* OutGroom, const
 			HairStrandsDataToEditableHairStrands(StrandsData, Out.Groups[GroupIndex].Strands);
 		}
 	}
+	else
+	{
+		Out.Groups.Reset();
+	}
 #endif
 }
 
@@ -289,7 +293,7 @@ void ConvertToGroomAsset(UGroomAsset* OutGroom, const FEditableGroom* InGroom, u
 	// Commit & build data for the groom asset
 	// Internally, this will derived the build data, recreate resources, and 
 	// recreate all components using this groom assets
-	OutGroom->CommitHairDescription(MoveTemp(HairDescription), UGroomAsset::EHairDescriptionType::Source);
+	OutGroom->CommitHairDescription(MoveTemp(HairDescription), EHairDescriptionType::Edit);
 	OutGroom->CacheDerivedDatas();
 	OutGroom->OnGroomAssetChanged.Broadcast();
 #endif
diff --git a/Engine/Plugins/Runtime/HairStrands/Source/HairStrandsCore/Private/GroomRBFDeformer.cpp b/Engine/Plugins/Runtime/HairStrands/Source/HairStrandsCore/Private/GroomRBFDeformer.cpp
--- a/Engine/Plugins/Runtime/HairStrands/Source/HairStrandsCore/Private/GroomRBFDeformer.cpp
+++ b/Engine/Plugins/Runtime/HairStrands/Source/HairStrandsCore/Private/GroomRBFDeformer.cpp
@@ -500,7 +500,7 @@ static void ApplyDeformationToGroom(const TArray<FRBFDeformedPositions>& Deforme
 		}
 	}
 
-	GroomAsset->CommitHairDescription(MoveTemp(HairDescription), UGroomAsset::EHairDescriptionType::Source);
+	GroomAsset->CommitHairDescription(MoveTemp(HairDescription), EHairDescriptionType::Edit);
 	GroomAsset->UpdateHairGroupsInfo();
 
 	// Update/reimport the cards/meshes geometry which have been deformed prior to call this function
diff --git a/Engine/Plugins/Runtime/HairStrands/Source/HairStrandsCore/Public/GroomAsset.h b/Engine/Plugins/Runtime/HairStrands/Source/HairStrandsCore/Public/GroomAsset.h
--- a/Engine/Plugins/Runtime/HairStrands/Source/HairStrandsCore/Public/GroomAsset.h
+++ b/Engine/Plugins/Runtime/HairStrands/Source/HairStrandsCore/Public/GroomAsset.h
@@ -451,6 +451,14 @@ struct FHairVertexFactoryTypesPerMaterialData
 	FPSOPrecacheVertexFactoryDataList VertexFactoryDataList;
 };
 
+UENUM()
+enum class EHairDescriptionType : uint8
+{
+	Source,
+	Edit,
+	Count
+};
+
 /**
  * Implements an asset that can be used to store hair strands
  */
@@ -868,17 +876,16 @@ private:
 	bool HasValidData_Cards(uint32 GroupIndex) const;
 	bool HasValidData_Meshes(uint32 GroupIndex) const;
 public:
-	enum EHairDescriptionType
-	{
-		Source,
-		Edit,
-		Count
-	};
-
+	
 	/** Commits a HairDescription to buffer for serialization */
 	UE_API void CommitHairDescription(FHairDescription&& HairDescription, EHairDescriptionType Type);
+
+	/** Get the current hair description */
 	UE_API FHairDescription GetHairDescription() const;
 
+	// Loads a specific hair description stored on the asset
+	UE_API FHairDescription LoadHairDescription(const EHairDescriptionType DescriptionType) const;
+
 	/** Get/Build render & guides data based on the hair description and interpolation settings */
 	UE_API bool GetHairStrandsDatas(const int32 GroupIndex, FHairStrandsDatas& OutStrandsData, FHairStrandsDatas& OutGuidesData) const;
 	UE_API bool GetHairCardsGuidesDatas(const int32 GroupIndex, const int32 LODIndex, FHairStrandsDatas& OutCardsGuidesData);
@@ -899,17 +906,28 @@ public:
 
 	UE_API const FHairDescriptionGroups& GetHairDescriptionGroups() const;
 private:
+
+	/** Set the hair description type */
+	void SetHairDescriptionType(const EHairDescriptionType DescriptionType) {HairDescriptionType = DescriptionType;};
+
+	/** Get the hair description type */
+	EHairDescriptionType GetHairDescriptionType() const {return HairDescriptionType;}
+
+	/** Get the hair description index */
+	uint8 GetHairDescriptionIndex() const {return FMath::Clamp(static_cast<uint8>(HairDescriptionType),0, 1);}
 	
 	bool IsFullyCached();
-	TUniquePtr<FHairDescriptionBulkData> HairDescriptionBulkData[EHairDescriptionType::Count];
+	TUniquePtr<FHairDescriptionBulkData> HairDescriptionBulkData[static_cast<uint8>(EHairDescriptionType::Count)];
+
+	UPROPERTY()
 	EHairDescriptionType HairDescriptionType = EHairDescriptionType::Source;
 
 	// Transient HairDescription & HairDescriptionGroups, which are built from HairDescriptionBulkData.
 	// All these data (bulk/desc/groups) needs to be in sync. I.e., when the HairDescription is updated, 
 	// HairDescriptionGroups needs to also be updated
 	mutable UE::FRecursiveMutex InternalLock;
-	mutable TUniquePtr<FHairDescription> CachedHairDescription[EHairDescriptionType::Count];
-	mutable TUniquePtr<FHairDescriptionGroups> CachedHairDescriptionGroups[EHairDescriptionType::Count];
+	mutable TUniquePtr<FHairDescription> CachedHairDescription[static_cast<uint8>(EHairDescriptionType::Count)];
+	mutable TUniquePtr<FHairDescriptionGroups> CachedHairDescriptionGroups[static_cast<uint8>(EHairDescriptionType::Count)];
 
 	TArray<FString> StrandsDerivedDataKey;
 	TArray<FString> CardsDerivedDataKey;
diff --git a/Engine/Plugins/Runtime/HairStrands/Source/HairStrandsDataflow/Private/GenerateGuidesCurvesNode.cpp b/Engine/Plugins/Runtime/HairStrands/Source/HairStrandsDataflow/Private/GenerateGuidesCurvesNode.cpp
--- a/Engine/Plugins/Runtime/HairStrands/Source/HairStrandsDataflow/Private/GenerateGuidesCurvesNode.cpp
+++ b/Engine/Plugins/Runtime/HairStrands/Source/HairStrandsDataflow/Private/GenerateGuidesCurvesNode.cpp
@@ -10,7 +10,7 @@
 
 namespace UE::Groom::Private
 {
-	FORCEINLINE void GenerateStrandsCurves(const GeometryCollection::Facades::FCollectionCurveGeometryFacade& CurvesFacade, const FDataflowCurveSelection& CurveSelection, const uint32 GuidesCount, TArray<uint32>& SampleIndices)
+	FORCEINLINE void GenerateGeometryCurves(const GeometryCollection::Facades::FCollectionCurveGeometryFacade& CurvesFacade, const FDataflowCurveSelection& CurveSelection, const uint32 GuidesCount, TArray<uint32>& SampleIndices)
 	{
 		TArray<FVector3f> RootPositions;
 		RootPositions.Reserve(CurvesFacade.GetNumCurves());
@@ -33,49 +33,49 @@ namespace UE::Groom::Private
 		SampleIndices = PointsSampler.SampleIndices;
 	}
 
-	FORCEINLINE void GenerateStrandsPoints(const GeometryCollection::Facades::FCollectionCurveGeometryFacade& StrandsFacade, const TArray<uint32>& SampleIndices,
-		const bool bMergeGuides, GeometryCollection::Facades::FCollectionCurveGeometryFacade& GuidesFacade)
+	FORCEINLINE void GenerateGeometryPoints(const GeometryCollection::Facades::FCollectionCurveGeometryFacade& SourceFacade, const TArray<uint32>& SampleIndices,
+		const bool bMergeCurves, GeometryCollection::Facades::FCollectionCurveGeometryFacade& TargetFacade)
 	{
 		TArray<TArray<FVector3f>> GeometryPointPositions;
 		TArray<TArray<int32>> GeometryCurvePoints;
 		TArray<TArray<int32>> GeometrySourceIndices;
 			
-		GeometryPointPositions.SetNum(StrandsFacade.GetNumGeometry());
-		GeometryCurvePoints.SetNum(StrandsFacade.GetNumGeometry());
-		GeometrySourceIndices.SetNum(StrandsFacade.GetNumGeometry());
+		GeometryPointPositions.SetNum(SourceFacade.GetNumGeometry());
+		GeometryCurvePoints.SetNum(SourceFacade.GetNumGeometry());
+		GeometrySourceIndices.SetNum(SourceFacade.GetNumGeometry());
 		
-		if(bMergeGuides && GuidesFacade.GetNumGeometry() == StrandsFacade.GetNumGeometry())
+		if(bMergeCurves && TargetFacade.GetNumGeometry() == SourceFacade.GetNumGeometry())
 		{
 			int32 PointOffset = 0;
 			int32 CurveOffset = 0;
-			for (int32 GeometryIndex = 0; GeometryIndex < GuidesFacade.GetNumGeometry(); ++GeometryIndex)
+			for (int32 GeometryIndex = 0; GeometryIndex < TargetFacade.GetNumGeometry(); ++GeometryIndex)
 			{
-				for (int32 CurveIndex = CurveOffset; CurveIndex < GuidesFacade.GetGeometryCurveOffsets()[GeometryIndex]; ++CurveIndex)
+				for (int32 CurveIndex = CurveOffset; CurveIndex < TargetFacade.GetGeometryCurveOffsets()[GeometryIndex]; ++CurveIndex)
 				{
-					for (int32 PointIndex = PointOffset; PointIndex < GuidesFacade.GetCurvePointOffsets()[CurveIndex]; ++PointIndex)
+					for (int32 PointIndex = PointOffset; PointIndex < TargetFacade.GetCurvePointOffsets()[CurveIndex]; ++PointIndex)
 					{
-						GeometryPointPositions[GeometryIndex].Add(GuidesFacade.GetPointRestPositions()[PointIndex]);
+						GeometryPointPositions[GeometryIndex].Add(TargetFacade.GetPointRestPositions()[PointIndex]);
 					}
-					GeometryCurvePoints[GeometryIndex].Add(GuidesFacade.GetCurvePointOffsets()[CurveIndex]-PointOffset);
-					GeometrySourceIndices[GeometryIndex].Add(GuidesFacade.GetCurveSourceIndices()[CurveIndex]);
-					PointOffset = GuidesFacade.GetCurvePointOffsets()[CurveIndex];
+					GeometryCurvePoints[GeometryIndex].Add(TargetFacade.GetCurvePointOffsets()[CurveIndex]-PointOffset);
+					GeometrySourceIndices[GeometryIndex].Add(TargetFacade.GetCurveSourceIndices()[CurveIndex]);
+					PointOffset = TargetFacade.GetCurvePointOffsets()[CurveIndex];
 				}
-				CurveOffset = GuidesFacade.GetGeometryCurveOffsets()[GeometryIndex];
+				CurveOffset = TargetFacade.GetGeometryCurveOffsets()[GeometryIndex];
 			}
 		}
 			
 		for(const uint32& SampleIndex : SampleIndices)
 		{
-			const int32 PrevPoint = (SampleIndex == 0) ? 0 : StrandsFacade.GetCurvePointOffsets()[SampleIndex-1];
-			const int32 NextPoint = StrandsFacade.GetCurvePointOffsets()[SampleIndex];
+			const int32 PrevPoint = (SampleIndex == 0) ? 0 : SourceFacade.GetCurvePointOffsets()[SampleIndex-1];
+			const int32 NextPoint = SourceFacade.GetCurvePointOffsets()[SampleIndex];
 
-			const int32 GeometryIndex = StrandsFacade.GetCurveGeometryIndices()[SampleIndex];
+			const int32 GeometryIndex = SourceFacade.GetCurveGeometryIndices()[SampleIndex];
 			for(int32 PointIndex = PrevPoint; PointIndex < NextPoint; ++PointIndex)
 			{
-				GeometryPointPositions[GeometryIndex].Add(StrandsFacade.GetPointRestPositions()[PointIndex]);
+				GeometryPointPositions[GeometryIndex].Add(SourceFacade.GetPointRestPositions()[PointIndex]);
 			}
 			GeometryCurvePoints[GeometryIndex].Add(NextPoint-PrevPoint);
-			GeometrySourceIndices[GeometryIndex].Add(SampleIndex);
+			GeometrySourceIndices[GeometryIndex].Add(SourceFacade.GetCurveSourceIndices()[SampleIndex]);
 		}
 
 		TArray<FVector3f> PointRestPositions;
@@ -84,7 +84,7 @@ namespace UE::Groom::Private
 		
 		int32 CurveOffset = 0;
 		int32 PointOffset = 0;
-		for (int32 GeometryIndex = 0; GeometryIndex < StrandsFacade.GetNumGeometry(); ++GeometryIndex)
+		for (int32 GeometryIndex = 0; GeometryIndex < SourceFacade.GetNumGeometry(); ++GeometryIndex)
 		{
 			PointRestPositions.Append(GeometryPointPositions[GeometryIndex]);
 			for (int32 CurveIndex = 0; CurveIndex < GeometryCurvePoints[GeometryIndex].Num(); ++CurveIndex)
@@ -96,11 +96,11 @@ namespace UE::Groom::Private
 			CurveOffset += GeometryCurvePoints[GeometryIndex].Num();
 			GeometryCurveOffsets.Add(CurveOffset);
 				
-			const FString GroupName = StrandsFacade.GetGeometryGroupNames()[GeometryIndex].Replace(TEXT("Strands"), TEXT("Guides"));
+			const FString GroupName = SourceFacade.GetGeometryGroupNames()[GeometryIndex];
 			GeometryGroupNames.Add(GroupName);
 		}
-		GuidesFacade.InitCurvesCollection(PointRestPositions, CurvePointOffsets, 
-		GeometryCurveOffsets, GeometryGroupNames,  StrandsFacade.GetGeometryCurveThickness(), CurveSourceIndices);
+		TargetFacade.InitCurvesCollection(PointRestPositions, CurvePointOffsets, 
+		GeometryCurveOffsets, GeometryGroupNames,  SourceFacade.GetGeometryCurveThickness(), CurveSourceIndices);
 	}
 
 	
@@ -138,10 +138,10 @@ void FGenerateCurveGeometryDataflowNode::Evaluate(UE::Dataflow::FContext& Contex
 		if(SourceFacade.IsValid() && CountValue > 0)
 		{
 			TArray<uint32> SampleIndices;
-			UE::Groom::Private::GenerateStrandsCurves(SourceFacade, DataflowSelection, CountValue, SampleIndices);
+			UE::Groom::Private::GenerateGeometryCurves(SourceFacade, DataflowSelection, CountValue, SampleIndices);
 
 			GeometryCollection::Facades::FCollectionCurveGeometryFacade TargetFacade(TargetCollection);
-			UE::Groom::Private::GenerateStrandsPoints(SourceFacade, SampleIndices, bMergeCurves, TargetFacade);
+			UE::Groom::Private::GenerateGeometryPoints(SourceFacade, SampleIndices, bMergeCurves, TargetFacade);
 		}
 		
 		SetValue(Context, MoveTemp(TargetCollection), &Collection);
diff --git a/Engine/Plugins/Runtime/HairStrands/Source/HairStrandsDataflow/Private/GetGroomAssetNode.cpp b/Engine/Plugins/Runtime/HairStrands/Source/HairStrandsDataflow/Private/GetGroomAssetNode.cpp
--- a/Engine/Plugins/Runtime/HairStrands/Source/HairStrandsDataflow/Private/GetGroomAssetNode.cpp
+++ b/Engine/Plugins/Runtime/HairStrands/Source/HairStrandsDataflow/Private/GetGroomAssetNode.cpp
@@ -20,12 +20,14 @@ namespace UE::Groom::Private
 		TArray<FString> GeometryGroupNames;
 		TArray<float> GeometryCurveThickness;
 		TArray<int32> CurveSourceIndices;
-
+		
 		uint32 GroupIndex = 0;
 		for(FEditableGroomGroup& GroomGroup : EditGroom.Groups)
 		{
 			if(CurvesType == EGroomCollectionType::Guides)
 			{
+				int32 CurveIndex = 0;
+				int32 SourceType = static_cast<uint8>(EGroomCollectionType::Guides);
 				for (const FEditableHairGuide& EditType : GroomGroup.Guides)
 				{
 					for(auto& GuidePoint : EditType.ControlPoints)
@@ -33,11 +35,14 @@ namespace UE::Groom::Private
 						PointRestPositions.Add(GuidePoint.Position);
 					}
 					CurvePointOffsets.Add(PointRestPositions.Num());
-					CurveSourceIndices.Add(0);
+					CurveSourceIndices.Add((CurveIndex << 1) | SourceType);
+					++CurveIndex;
 				}
 			}
 			else
 			{
+				int32 CurveIndex = 0;
+				int32 SourceType = static_cast<uint8>(EGroomCollectionType::Strands);
 				for (const FEditableHairStrand& EditType : GroomGroup.Strands)
 				{
 					for(auto& GuidePoint : EditType.ControlPoints)
@@ -45,7 +50,8 @@ namespace UE::Groom::Private
 						PointRestPositions.Add(GuidePoint.Position);
 					}
 					CurvePointOffsets.Add(PointRestPositions.Num());
-					CurveSourceIndices.Add(0);
+					CurveSourceIndices.Add((CurveIndex << 1) | SourceType);
+					++CurveIndex;
 				}
 			}
 			GeometryCurveOffsets.Add(CurvePointOffsets.Num());
diff --git a/Engine/Plugins/Runtime/HairStrands/Source/HairStrandsDataflow/Private/GroomAssetTerminalNode.cpp b/Engine/Plugins/Runtime/HairStrands/Source/HairStrandsDataflow/Private/GroomAssetTerminalNode.cpp
--- a/Engine/Plugins/Runtime/HairStrands/Source/HairStrandsDataflow/Private/GroomAssetTerminalNode.cpp
+++ b/Engine/Plugins/Runtime/HairStrands/Source/HairStrandsDataflow/Private/GroomAssetTerminalNode.cpp
@@ -3,7 +3,9 @@
 #include "GroomAssetTerminalNode.h"
 
 #include "AssetCompilingManager.h"
+#include "GetGroomAssetNode.h"
 #include "GroomEdit.h"
+#include "HLSLTypeAliases.h"
 #include "Dataflow/DataflowConnectionTypes.h"
 #include "Dataflow/DataflowObjectInterface.h"
 #include "GeometryCollection/Facades/CollectionCurveFacade.h"
@@ -15,45 +17,162 @@ namespace UE::Groom::Private
 {
 	static const FName AttributeKeyTypeName(TEXT("FCollectionAttributeKey"));
 
-	FORCEINLINE void BuildEditableGuides(const TArray<int32>& ObjectCurveOffsets,
-		const TArray<int32>& CurvePointOffsets, const TArray<FVector3f>& PointRestPositions, const TArray<int32>& CurveStrandIndices, FEditableGroom& EditGroom)
+	static void BuildEditableGuides(const TArray<int32>& ObjectCurveOffsets,
+		const TArray<int32>& CurvePointOffsets, const TArray<FVector3f>& PointRestPositions, const TArray<int32>& CurveStrandIndices,
+		const FEditableGroom& SourceGroom, FEditableGroom& EditGroom)
 	{
-		int32 ObjectIndex = 0, CurveIndex = 0;
-		int32 PrevCurve = 0, PrevPoint = 0;;
-		for (FEditableGroomGroup& Group : EditGroom.Groups)
+		if (SourceGroom.Groups.Num() == EditGroom.Groups.Num())
 		{
-			const int32 NextCurve = ObjectCurveOffsets[ObjectIndex];
-			Group.Guides.SetNum(NextCurve-PrevCurve);
-			
-			// Only rebuild the guides since the strands are the same
-			for (FEditableHairGuide& Guide : Group.Guides)
+			int32 ObjectIndex = 0, CurveIndex = 0;
+			int32 PrevCurve = 0, PrevPoint = 0;;
+			for (FEditableGroomGroup& Group : EditGroom.Groups)
 			{
-				const int32 NextPoint = CurvePointOffsets[CurveIndex];
-				const int32 StrandIndex = CurveStrandIndices[CurveIndex];
+				const int32 NextCurve = ObjectCurveOffsets[ObjectIndex];
+				Group.Guides.SetNum(NextCurve-PrevCurve);
+				
+				for (FEditableHairGuide& Guide : Group.Guides)
+				{
+					const int32 NextPoint = CurvePointOffsets[CurveIndex];
+					const int32 SourceType = CurveStrandIndices[CurveIndex] & 1;
+					const int32 SourceCurve = CurveStrandIndices[CurveIndex] >> 1;
 						
-				Guide.ControlPoints.Reset();
+					Guide.ControlPoints.Reset();
 
-				if((NextPoint - PrevPoint - 1) > 0)
-				{ 
-					for (int32 PointIndex = PrevPoint; PointIndex < NextPoint; ++PointIndex)
-					{
-						const float PointLength =
+					if((NextPoint - PrevPoint - 1) > 0)
+					{ 
+						for (int32 PointIndex = PrevPoint; PointIndex < NextPoint; ++PointIndex)
+						{
 							Guide.ControlPoints.Add({ PointRestPositions[PointIndex],
-								FMath::Clamp(static_cast<float>(PointIndex - PrevPoint) / (NextPoint - PrevPoint - 1), 0.f, 1.f) });
+									FMath::Clamp(static_cast<float>(PointIndex - PrevPoint) / (NextPoint - PrevPoint - 1), 0.f, 1.f) });
+						}
 					}
+
+					const FEditableGroomGroup& SourceGroup = SourceGroom.Groups[ObjectIndex];
+
+					Guide.bHasGuideID = false;
+					Guide.bHasRootUV = false;
+					
+					if (SourceType == static_cast<uint8>(EGroomCollectionType::Guides))
+					{
+						if(SourceGroup.Guides.IsValidIndex(SourceCurve))
+						{
+							Guide.bHasGuideID = SourceGroup.Guides[SourceCurve].bHasGuideID;
+							Guide.GuideID = SourceGroup.Guides[SourceCurve].GuideID;
+
+							Guide.bHasRootUV = SourceGroup.Guides[SourceCurve].bHasRootUV;
+							Guide.RootUV = SourceGroup.Guides[SourceCurve].RootUV;	
+						}
+					}
+				
+					PrevPoint = NextPoint;
+					++CurveIndex;
 				}
-						
-				if((StrandIndex != INDEX_NONE) && (StrandIndex < Group.Strands.Num()))
+				PrevCurve = NextCurve;
+				++ObjectIndex;
+			}
+		}
+	}
+
+	static void BuildEditableStrands(const TArray<int32>& ObjectCurveOffsets,
+			const TArray<int32>& CurvePointOffsets, const TArray<FVector3f>& PointRestPositions, const TArray<int32>& CurveStrandIndices,
+			const FEditableGroom& SourceGroom, FEditableGroom& EditGroom)
+	{
+		if (SourceGroom.Groups.Num() == EditGroom.Groups.Num())
+		{
+			int32 ObjectIndex = 0, CurveIndex = 0;
+			int32 PrevCurve = 0, PrevPoint = 0;
+			for (FEditableGroomGroup& Group : EditGroom.Groups)
+			{
+				const int32 NextCurve = ObjectCurveOffsets[ObjectIndex];
+				Group.Strands.SetNum(NextCurve-PrevCurve);
+
+				const FEditableGroomGroup& SourceGroup = SourceGroom.Groups[ObjectIndex];
+				
+				for (FEditableHairStrand& Strand : Group.Strands)
 				{
-					Guide.GuideID = CurveIndex;
-					Guide.RootUV = Group.Strands[StrandIndex].RootUV;
-				}
+					const int32 NextPoint = CurvePointOffsets[CurveIndex];
+					const int32 SourceType = CurveStrandIndices[CurveIndex] & 1;
+					const int32 SourceCurve = CurveStrandIndices[CurveIndex] >> 1;
+						
+					Strand.ControlPoints.Reset();
+
+					if((NextPoint - PrevPoint - 1) > 0)
+					{ 
+						for (int32 PointIndex = PrevPoint; PointIndex < NextPoint; ++PointIndex)
+						{
+							FEditableHairStrandControlPoint ControlPoint;
+							ControlPoint.Position = PointRestPositions[PointIndex];
+							ControlPoint.U = static_cast<float>(PointIndex - PrevPoint) / (NextPoint - PrevPoint - 1);
+
+							ControlPoint.bHasAO = false;
+							ControlPoint.bHasRoughness = false;
+							ControlPoint.bHasColor = false;
+							
+							if (SourceType == static_cast<uint8>(EGroomCollectionType::Strands))
+							{
+								if(SourceGroup.Strands.IsValidIndex(SourceCurve) && SourceGroup.Strands[SourceCurve].ControlPoints.Num() >= 2)
+								{
+									const float SourcePointCoord = (SourceGroup.Strands[SourceCurve].ControlPoints.Num()-1) * ControlPoint.U;
+									const int32 SourcePointIndex = FMath::Floor(SourcePointCoord);
+									const float SourceLerpValue = SourcePointCoord - SourcePointIndex;
+									
+									const uint32 SourcePrevIndex = FMath::Clamp(SourcePointIndex,
+													0, SourceGroup.Strands[SourceCurve].ControlPoints.Num()-2);
+									const uint32 SourceNextIndex = SourcePrevIndex + 1;
+									
+									const FEditableHairStrandControlPoint& SourcePrevPoint =
+										SourceGroup.Strands[SourceCurve].ControlPoints[SourcePrevIndex];
+									const FEditableHairStrandControlPoint& SourceNextPoint =
+										SourceGroup.Strands[SourceCurve].ControlPoints[SourceNextIndex];
+
+									ControlPoint.Radius = SourcePrevPoint.Radius * (1.0-SourceLerpValue) + SourceNextPoint.Radius * SourceLerpValue;
+						
+									ControlPoint.bHasAO = SourcePrevPoint.bHasAO && SourceNextPoint.bHasAO;
+									ControlPoint.bHasRoughness = SourcePrevPoint.bHasRoughness && SourceNextPoint.bHasRoughness;
+									ControlPoint.bHasColor = SourcePrevPoint.bHasColor && SourceNextPoint.bHasColor;
+
+									ControlPoint.AO = SourcePrevPoint.AO * (1.0-SourceLerpValue) + SourceNextPoint.AO * SourceLerpValue;
+									ControlPoint.Roughness = SourcePrevPoint.Roughness * (1.0-SourceLerpValue) + SourceNextPoint.Roughness * SourceLerpValue;
+									ControlPoint.BaseColor = FLinearColor::LerpUsingHSV(SourcePrevPoint.BaseColor, SourceNextPoint.BaseColor, SourceLerpValue);
+								}
+							}
+							Strand.ControlPoints.Add(ControlPoint);
+						}
+					}
+
+					Strand.bHasStrandID = false;
+					Strand.bHasRootUV = false;
+					Strand.bHasClosestGuide = false;
+					Strand.bHasClumpID = false;
+					
+					if (SourceType == static_cast<uint8>(EGroomCollectionType::Strands))
+					{
+						if(SourceGroup.Strands.IsValidIndex(SourceCurve))
+						{
+							Strand.bHasStrandID = SourceGroup.Strands[SourceCurve].bHasStrandID;
+							Strand.StrandID = SourceGroup.Strands[SourceCurve].StrandID;
+
+							Strand.bHasRootUV = SourceGroup.Strands[SourceCurve].bHasRootUV;
+							Strand.RootUV = SourceGroup.Strands[SourceCurve].RootUV;	
+							
+							Strand.bHasClumpID = SourceGroup.Strands[SourceCurve].bHasClumpID;
+                            Strand.ClumpID = SourceGroup.Strands[SourceCurve].ClumpID;
+
+							Strand.bHasClosestGuide = SourceGroup.Strands[SourceCurve].bHasClosestGuide;
+							for (uint32 GuideIndex = 0; GuideIndex < 3; ++GuideIndex)
+							{
+								Strand.GuideIDs[GuideIndex] = SourceGroup.Strands[SourceCurve].GuideIDs[GuideIndex];
+								Strand.GuideWeights[GuideIndex] = SourceGroup.Strands[SourceCurve].GuideWeights[GuideIndex];
+							}
+						}
+					}
 				
-				PrevPoint = NextPoint;
-				++CurveIndex;
+					PrevPoint = NextPoint;
+					++CurveIndex;
+				}
+				PrevCurve = NextCurve;
+				++ObjectIndex;
 			}
-			PrevCurve = NextCurve;
-			++ObjectIndex;
 		}
 	}
 
@@ -320,39 +439,47 @@ void FGroomAssetTerminalDataflowNode_v2::SetAssetValue(TObjectPtr<UObject> Asset
 			}
 		
 			TArray<FCollectionAttributeKey> InternalAttributes;
-			InternalAttributes.Append(SkinningAttributes);
-			InternalAttributes.Append(HierarchyAttributes);
+			//InternalAttributes.Append(SkinningAttributes);
+			//InternalAttributes.Append(HierarchyAttributes);
+
+			FEditableGroom EditGroom;
+			ConvertFromGroomAsset(const_cast<UGroomAsset*>(GroomAsset), &EditGroom, false, false, false);
+			FEditableGroom SourceGroom = EditGroom;
 		
 			if(StrandsFacade.IsValid())
 			{
 				UE::Groom::Private::TransferCurvesAttributes(GroomStrands, OutCollection, ExternalAttributes, InternalAttributes, FString("Strands"));
+
+				if(StrandsFacade.GetNumGeometry() == EditGroom.Groups.Num())
+				{
+					// Build the editable strands
+					UE::Groom::Private::BuildEditableStrands(StrandsFacade.GetGeometryCurveOffsets(), StrandsFacade.GetCurvePointOffsets(),
+						StrandsFacade.GetPointRestPositions(), StrandsFacade.GetCurveSourceIndices(), SourceGroom, EditGroom);
+				}
 			}
 			if(GuidesFacade.IsValid())
 			{
 				UE::Groom::Private::TransferCurvesAttributes(GroomGuides, OutCollection, ExternalAttributes, InternalAttributes, FString("Guides"));
-
-				// Build an editable groom asset for the guides
-				FEditableGroom EditGroom;
-				ConvertFromGroomAsset(const_cast<UGroomAsset*>(GroomAsset), &EditGroom, false, false, false);
-
+				
 				if(GuidesFacade.GetNumGeometry() == EditGroom.Groups.Num())
 				{
 					// Build the editable guides
 					UE::Groom::Private::BuildEditableGuides(GuidesFacade.GetGeometryCurveOffsets(), GuidesFacade.GetCurvePointOffsets(),
-						GuidesFacade.GetPointRestPositions(), GuidesFacade.GetCurveSourceIndices(),EditGroom);
+						GuidesFacade.GetPointRestPositions(), GuidesFacade.GetCurveSourceIndices(), SourceGroom, EditGroom);
+				}
+			}
 
-					// Ensure compilation dependent assets is done
-					FAssetCompilingManager::Get().FinishCompilationForObjects({ GroomAsset });
+			// Ensure compilation dependent assets is done
+			FAssetCompilingManager::Get().FinishCompilationForObjects({ GroomAsset });
 
-					// Convert to groom asset
-					ConvertToGroomAsset(const_cast<UGroomAsset*>(GroomAsset), &EditGroom, EEditableGroomOperations::ControlPoints_Modified);
-				}
-				// To prevent future reconstruction in the BuildData we set the type to be imported
-				for(FHairGroupsInterpolation& GroupInterpolation : GroomAsset->GetHairGroupsInterpolation())
-				{
-					GroupInterpolation.InterpolationSettings.GuideType = EGroomGuideType::Imported;
-				}
-			}  
+			// Convert to groom asset
+			ConvertToGroomAsset(const_cast<UGroomAsset*>(GroomAsset), &EditGroom, EEditableGroomOperations::ControlPoints_Modified);
+	
+			// To prevent future reconstruction in the BuildData we set the type to be imported
+			for(FHairGroupsInterpolation& GroupInterpolation : GroomAsset->GetHairGroupsInterpolation())
+			{
+				GroupInterpolation.InterpolationSettings.GuideType = EGroomGuideType::Imported;
+			}
 
 			GroomAsset->GetDataflowSettings().SetRestCollection(OutCollection);
 			GroomAsset->GetDataflowSettings().InitSkeletalMeshes(GuidesFacade.GetNumGeometry());
diff --git a/Engine/Plugins/Runtime/HairStrands/Source/HairStrandsEditor/Private/HairStrandsImporter.cpp b/Engine/Plugins/Runtime/HairStrands/Source/HairStrandsEditor/Private/HairStrandsImporter.cpp
--- a/Engine/Plugins/Runtime/HairStrands/Source/HairStrandsEditor/Private/HairStrandsImporter.cpp
+++ b/Engine/Plugins/Runtime/HairStrands/Source/HairStrandsEditor/Private/HairStrandsImporter.cpp
@@ -219,7 +219,7 @@ UGroomAsset* FHairStrandsImporter::ImportHair(const FHairImportContext& ImportCo
 	check(OutHairAsset->AreGroupsValid());
 	check(uint32(OutHairAsset->GetNumHairGroups()) == GroupCount);
 
-	OutHairAsset->CommitHairDescription(MoveTemp(NewHairDescription), UGroomAsset::EHairDescriptionType::Source);
+	OutHairAsset->CommitHairDescription(MoveTemp(NewHairDescription), EHairDescriptionType::Source);
 
 	const bool bSucceeded = OutHairAsset->CacheDerivedDatas();
 	if (!bSucceeded)
