Optional simulation visualisation for strands / Fix bad indexing for custom attributes / update simulation view when asset is changing / Relative offsets for guides LODs

#FailedValidations Preflight

diff --git a/Engine/Plugins/Experimental/Dataflow/Source/DataflowEditor/Public/Dataflow/DataflowEditorToolkit.h b/Engine/Plugins/Experimental/Dataflow/Source/DataflowEditor/Public/Dataflow/DataflowEditorToolkit.h
--- a/Engine/Plugins/Experimental/Dataflow/Source/DataflowEditor/Public/Dataflow/DataflowEditorToolkit.h
+++ b/Engine/Plugins/Experimental/Dataflow/Source/DataflowEditor/Public/Dataflow/DataflowEditorToolkit.h
@@ -129,22 +129,22 @@ public:
 	const EDataflowEditorEvaluationMode& GetEvaluationMode() const { return EvaluationMode; }
 
 	/** Reset dataflow simulation */
-	void ResetDataflowSimulation() const;
+	UE_API void ResetDataflowSimulation() const;
 	
 	/** Toggle dataflow simulation */
-	void ToggleDataflowSimulation() const;
+	UE_API void ToggleDataflowSimulation() const;
 	
 	/** Start dataflow simulation */
-	void StartDataflowSimulation() const;
+	UE_API void StartDataflowSimulation() const;
 
 	/** Stop dataflow simulation */
-	void StopDataflowSimulation() const;
+	UE_API void StopDataflowSimulation() const;
 	
 	/** Step dataflow simulation */
-	void StepDataflowSimulation() const;
+	UE_API void StepDataflowSimulation() const;
 
 	/** Pause dataflow simulation */
-	void PauseDataflowSimulation() const;
+	UE_API void PauseDataflowSimulation() const;
 
 	/** Check if the simulation can be toggled or not */
 	bool HasSimulationManager() const;
diff --git a/Engine/Plugins/Experimental/Dataflow/Source/DataflowNodes/Private/Dataflow/DataflowCollectionEditSkinWeightsNode.cpp b/Engine/Plugins/Experimental/Dataflow/Source/DataflowNodes/Private/Dataflow/DataflowCollectionEditSkinWeightsNode.cpp
--- a/Engine/Plugins/Experimental/Dataflow/Source/DataflowNodes/Private/Dataflow/DataflowCollectionEditSkinWeightsNode.cpp
+++ b/Engine/Plugins/Experimental/Dataflow/Source/DataflowNodes/Private/Dataflow/DataflowCollectionEditSkinWeightsNode.cpp
@@ -152,8 +152,8 @@ void CorrectSkinWeights(TArray<TArray<int32>>& BoneIndices, TArray<TArray<float>
 	{
 		if(BoneWeights[VertexIndex].Num() == BoneIndices[VertexIndex].Num())
 		{
-			ValidIndices.Reset(BoneIndices.Num());
-			ValidWeights.Reset(BoneWeights.Num());
+			ValidIndices.Reset(BoneIndices[VertexIndex].Num());
+			ValidWeights.Reset(BoneWeights[VertexIndex].Num());
 			
 			for(int32 WeightIndex = 0, NumWeights = BoneWeights[VertexIndex].Num(); WeightIndex < NumWeights; ++WeightIndex)
 			{
diff --git a/Engine/Plugins/Runtime/HairStrands/Source/HairStrandsDataflow/Private/BuildGuidesLODsNode.cpp b/Engine/Plugins/Runtime/HairStrands/Source/HairStrandsDataflow/Private/BuildGuidesLODsNode.cpp
--- a/Engine/Plugins/Runtime/HairStrands/Source/HairStrandsDataflow/Private/BuildGuidesLODsNode.cpp
+++ b/Engine/Plugins/Runtime/HairStrands/Source/HairStrandsDataflow/Private/BuildGuidesLODsNode.cpp
@@ -77,7 +77,7 @@ namespace UE::Groom::Private
 							MinParent = ParentIndex;
 						}
 					}
-					CurveParentIndices[GuideIndex] = MinParent;
+					CurveParentIndices[GuideIndex] = (MinParent != INDEX_NONE) ? GuideIndex-MinParent : INDEX_NONE;
 					CurveLodIndices[GuideIndex] = GuideLod;
 				}
 
diff --git a/Engine/Plugins/Runtime/HairStrands/Source/HairStrandsDataflow/Private/GroomDataflowVisualization.cpp b/Engine/Plugins/Runtime/HairStrands/Source/HairStrandsDataflow/Private/GroomDataflowVisualization.cpp
--- a/Engine/Plugins/Runtime/HairStrands/Source/HairStrandsDataflow/Private/GroomDataflowVisualization.cpp
+++ b/Engine/Plugins/Runtime/HairStrands/Source/HairStrandsDataflow/Private/GroomDataflowVisualization.cpp
@@ -31,7 +31,7 @@ namespace UE::Groom
 			{
 				if (const TSharedPtr<FDataflowSimulationScene>& SimulationScene = Toolkit->GetSimulationScene())
 				{
-					if (GetGroomComponent(SimulationScene.Get()))
+					if (const UGroomComponent* GroomComponent = GetGroomComponent(SimulationScene.Get()))
 					{
 						constexpr uint8 NumFlags = static_cast<uint8>(EGroomViewMode::Count);
 						auto AddSimulationVisualisation = [this, &ViewportClient , &MenuBuilder, NumFlags](const TAttribute<FText>& FlagLabel, const EGroomViewMode ViewMode)
@@ -71,10 +71,26 @@ namespace UE::Groom
 								MenuBuilder.AddMenuEntry(FlagLabel, FText(), FSlateIcon(), Action, NAME_None, EUserInterfaceActionType::ToggleButton);
 							};
 
+						bool bEnableGuidesVisualisation = false;
+						if(TObjectPtr<UGroomAsset> GroomAsset = GroomComponent->GroomAsset)
+						{ 
+							for (int32 GroupIndex = 0; GroupIndex < GroomAsset->GetNumHairGroups(); ++GroupIndex)
+							{
+								if (GroomAsset->IsSimulationEnable(GroupIndex, INDEX_NONE) || GroomAsset->IsDeformationEnable(GroupIndex))
+								{
+									bEnableGuidesVisualisation = true;
+									break;
+								}
+							}
+						}
+
 						MenuBuilder.BeginSection(TEXT("GroomSimulation_Visualizations"), LOCTEXT("GroomSimulationVisualization", "Groom Simulation Visualization"));
 
-						AddSimulationVisualisation(LOCTEXT("SimulatedGuides", "Simulated guides"), EGroomViewMode::SimHairStrands);
-						AddSimulationVisualisation(LOCTEXT("RenderedStrands", "Rendered strands"), EGroomViewMode::RenderHairStrands);
+						if(bEnableGuidesVisualisation)
+						{ 
+							AddSimulationVisualisation(LOCTEXT("DeformedGuides", "Deformed guides"), EGroomViewMode::SimHairStrands);
+							AddSimulationVisualisation(LOCTEXT("GuidesInfluence", "Guides influence"), EGroomViewMode::RenderHairStrands);
+						}
 						AddSimulationVisualisation(LOCTEXT("StrandsGroups", "Strands groups"), EGroomViewMode::Group);
 						AddSimulationVisualisation(LOCTEXT("StrandsClumps", "Strands clumps"), EGroomViewMode::ClumpID);
 						AddSimulationVisualisation(LOCTEXT("StrandsClusters", "Strands clusters"), EGroomViewMode::Cluster);
diff --git a/Engine/Plugins/Runtime/HairStrands/Source/HairStrandsDataflow/Public/BuildGroomSkinningNodes.h b/Engine/Plugins/Runtime/HairStrands/Source/HairStrandsDataflow/Public/BuildGroomSkinningNodes.h
--- a/Engine/Plugins/Runtime/HairStrands/Source/HairStrandsDataflow/Public/BuildGroomSkinningNodes.h
+++ b/Engine/Plugins/Runtime/HairStrands/Source/HairStrandsDataflow/Public/BuildGroomSkinningNodes.h
@@ -78,7 +78,7 @@ struct FTransferGeometrySkinWeightsDataflowNode : public FDataflowNode
 {
 	GENERATED_BODY()
 	
-	DATAFLOW_NODE_DEFINE_INTERNAL(FTransferGeometrySkinWeightsDataflowNode, "TransferGeometrySkinWeights", "Groom", "")
+	DATAFLOW_NODE_DEFINE_INTERNAL(FTransferGeometrySkinWeightsDataflowNode, "TransferLinearSkinWeights", "Groom", "")
 
 public:
 	FTransferGeometrySkinWeightsDataflowNode(const UE::Dataflow::FNodeParameters& InParam, FGuid InGuid = FGuid::NewGuid());
diff --git a/Engine/Plugins/Runtime/HairStrands/Source/HairStrandsDeformer/Private/DeformerGroomAttributeRead.cpp b/Engine/Plugins/Runtime/HairStrands/Source/HairStrandsDeformer/Private/DeformerGroomAttributeRead.cpp
--- a/Engine/Plugins/Runtime/HairStrands/Source/HairStrandsDeformer/Private/DeformerGroomAttributeRead.cpp
+++ b/Engine/Plugins/Runtime/HairStrands/Source/HairStrandsDeformer/Private/DeformerGroomAttributeRead.cpp
@@ -292,13 +292,16 @@ FOptimusGroomAttributeReadProviderProxy::FOptimusGroomAttributeReadProviderProxy
 		UOptimusGroomExecDataInterface::GetExecutionDomainName(GroomAttributeGroup), InvocationElements, false);
 
 	// Compute the number of elements for the base indices
-	NumElements.Reset();
+	NumElements.SetNumZeroed(InvocationElements.Num());
+	int32 AssetIndex = 0;
 	for (TPair<UGroomAsset*, UE::Groom::Private::FGroupElements>& Elements : InvocationElements)
 	{
+		NumElements[AssetIndex].Reset();
 		for (int32 GroupIndex = 0, NumGroups = Elements.Value.GroupOffsets.Num()-1; GroupIndex < NumGroups; ++GroupIndex)
 		{
-			NumElements.Add(Elements.Value.GroupOffsets[GroupIndex+1]-Elements.Value.GroupOffsets[GroupIndex]);
+			NumElements[AssetIndex].Add(Elements.Value.GroupOffsets[GroupIndex+1]-Elements.Value.GroupOffsets[GroupIndex]);
 		}
+		++AssetIndex;
 	}
 }
 
@@ -681,17 +684,19 @@ void FOptimusGroomAttributeReadProviderProxy::GatherDispatchData(FDispatchData c
 
 	int32 InvocationIndex = 0;
 	int32 BaseValueIndex = 0;
-	for (const TPair<UGroomAsset*,UE::Groom::Private::FGroupElements>& AssetPair : GroupElements)
+	for (int32 AssetIndex = 0, NumAssets = GroupElements.Num(); AssetIndex < NumAssets; ++AssetIndex)
 	{
-		int32 InstanceIndex = 0;
-		for(int32 GroupIndex : AssetPair.Value.GroupIndices)
+		const UE::Groom::Private::FGroupElements& AssetElements = GroupElements[AssetIndex].Value;
+		for (int32 InstanceIndex = 0, NumInstances = GroupElements[AssetIndex].Value.GroupIndices.Num(); InstanceIndex < NumInstances; ++InstanceIndex)
 		{
+			const int32 GroupIndex = GroupElements[AssetIndex].Value.GroupIndices[InstanceIndex];
+			
 			FParameters& Parameters = ParameterArray[InvocationIndex];
-			if (const FHairGroupInstance* GroupInstance = AssetPair.Value.GroupInstances[InstanceIndex])
+			if (const FHairGroupInstance* GroupInstance = AssetElements.GroupInstances[InstanceIndex])
 			{
-				const int32 ElementCount = (NumElements.IsValidIndex(InvocationIndex) ? NumElements[InvocationIndex] : 0);
+				const int32 ElementCount = NumElements[AssetIndex][GroupIndex];
 
-				Parameters.NumValues = AssetPair.Value.GroupOffsets[InstanceIndex+1] - AssetPair.Value.GroupOffsets[InstanceIndex];
+				Parameters.NumValues = AssetElements.GroupOffsets[GroupIndex+1] - AssetElements.GroupOffsets[GroupIndex];
 				Parameters.BaseValueIndex = BaseValueIndex;
 				Parameters.ValueTypedData = AttributeValuesResources[InvocationIndex];
 
@@ -709,7 +714,6 @@ void FOptimusGroomAttributeReadProviderProxy::GatherDispatchData(FDispatchData c
 				Parameters.bHasIndexMapping = 0;
 				Parameters.IndexMappingBuffer = FallbackStructuredSRV;
 			}
-			++InstanceIndex;
 			++InvocationIndex;
 		}
 	}
diff --git a/Engine/Plugins/Runtime/HairStrands/Source/HairStrandsDeformer/Private/DeformerGroomAttributeRead.h b/Engine/Plugins/Runtime/HairStrands/Source/HairStrandsDeformer/Private/DeformerGroomAttributeRead.h
--- a/Engine/Plugins/Runtime/HairStrands/Source/HairStrandsDeformer/Private/DeformerGroomAttributeRead.h
+++ b/Engine/Plugins/Runtime/HairStrands/Source/HairStrandsDeformer/Private/DeformerGroomAttributeRead.h
@@ -146,8 +146,8 @@ private:
 	/** Groom assets group indices the proxy instance is linked to */
 	TArray<TPair<UGroomAsset*, UE::Groom::Private::FGroupElements>> GroupElements;
 
-	/** Num elements that will be used in deformer for each invocation (usually larger than the number of source elements) */
-	TArray<int32> NumElements;
+	/** Num elements per assets that will be used in deformer for each invocation (usually larger than the number of source elements) */
+	TArray<TArray<int32>> NumElements;
 	
 	/** Attribute Resources used to dispatch CS on GPU */
 	TArray<FRDGBufferSRVRef> AttributeValuesResources;
diff --git a/Engine/Plugins/Runtime/HairStrands/Source/HairStrandsDeformer/Private/DeformerGroomDomainsSource.cpp b/Engine/Plugins/Runtime/HairStrands/Source/HairStrandsDeformer/Private/DeformerGroomDomainsSource.cpp
--- a/Engine/Plugins/Runtime/HairStrands/Source/HairStrandsDeformer/Private/DeformerGroomDomainsSource.cpp
+++ b/Engine/Plugins/Runtime/HairStrands/Source/HairStrandsDeformer/Private/DeformerGroomDomainsSource.cpp
@@ -174,5 +174,79 @@ bool UOptimusGroomSolverComponentSource::GetComponentElementCountsForExecutionDo
 	return false;
 }
 
+FName UOptimusGroomCollisionComponentSource::FCollisionExecutionDomains::Vertices("CollisionVertices");
+FName UOptimusGroomCollisionComponentSource::FCollisionExecutionDomains::Triangles("CollisionTriangles");
+
+FText UOptimusGroomCollisionComponentSource::GetDisplayName() const
+{
+	return LOCTEXT("GroomCollisionComponent", "Groom Collision Component");
+}
+
+TSubclassOf<UActorComponent> UOptimusGroomCollisionComponentSource::GetComponentClass() const
+{
+	return UGroomSolverComponent::StaticClass();
+}
+
+bool UOptimusGroomCollisionComponentSource::IsUsableAsPrimarySource() const
+{
+	return GetComponentClass()->IsChildOf<UGroomSolverComponent>();
+}
+
+TArray<FName> UOptimusGroomCollisionComponentSource::GetExecutionDomains() const
+{
+	return {FCollisionExecutionDomains::Vertices, FCollisionExecutionDomains::Triangles};
+}
+
+int32 UOptimusGroomCollisionComponentSource::GetLodIndex(const UActorComponent* InComponent) const
+{
+	int32 LODIndex = 0;
+	if(const UGroomSolverComponent* GroomSolver = Cast<UGroomSolverComponent>(InComponent))
+	{
+		int32 LODCount = 1;
+		for(const TPair<TObjectPtr<UMeshComponent>, int32>& CollisionComponent : GroomSolver->GetCollisionComponents())
+		{
+			if (USkinnedMeshComponent* SkinnedMeshComponent = Cast<USkinnedMeshComponent>(CollisionComponent.Key))
+			{
+				LODIndex += SkinnedMeshComponent->GetPredictedLODLevel() * LODCount;
+				LODCount *= SkinnedMeshComponent->GetNumLODs();
+			}
+		}
+	}
+	
+	return LODIndex;
+}
+
+uint32 UOptimusGroomCollisionComponentSource::GetDefaultNumInvocations(const UActorComponent* ActorComponent, int32 LodIndex) const
+{
+	if(const UGroomSolverComponent* GroomSolver = Cast<UGroomSolverComponent>(ActorComponent))
+	{
+		TArray<TPair<TObjectPtr<UMeshComponent>, int32>> CollisionComponents;
+		for(const TPair<TObjectPtr<UMeshComponent>, int32>& CollisionComponent : GroomSolver->GetCollisionComponents())
+		{
+			CollisionComponents.Add(CollisionComponent);
+		}
+		return UE::Groom::Private::GetCollisionNumInvocations(CollisionComponents);
+	}
+	
+	return 0;
+}
+
+bool UOptimusGroomCollisionComponentSource::GetComponentElementCountsForExecutionDomain(
+	FName DomainName,const UActorComponent* ActorComponent, int32 LodIndex, TArray<int32>& InvocationElementCounts) const
+{
+	if(const UGroomSolverComponent* GroomSolver = Cast<UGroomSolverComponent>(ActorComponent))
+	{
+		TArray<TPair<TObjectPtr<UMeshComponent>, int32>> CollisionComponents;
+		for(const TPair<TObjectPtr<UMeshComponent>, int32>& CollisionComponent : GroomSolver->GetCollisionComponents())
+		{
+			CollisionComponents.Add(CollisionComponent);
+		}
+		InvocationElementCounts.Reset();
+		const uint32 TotalCount = UE::Groom::Private::GetCollisionInvocationElementCounts(CollisionComponents, DomainName, InvocationElementCounts);
+		return true;
+	}
+	return false;
+}
+
 
 #undef LOCTEXT_NAMESPACE
diff --git a/Engine/Plugins/Runtime/HairStrands/Source/HairStrandsDeformer/Private/DeformerGroomDomainsSource.h b/Engine/Plugins/Runtime/HairStrands/Source/HairStrandsDeformer/Private/DeformerGroomDomainsSource.h
--- a/Engine/Plugins/Runtime/HairStrands/Source/HairStrandsDeformer/Private/DeformerGroomDomainsSource.h
+++ b/Engine/Plugins/Runtime/HairStrands/Source/HairStrandsDeformer/Private/DeformerGroomDomainsSource.h
@@ -9,7 +9,7 @@ class UActorComponent;
 class UGroomComponent;
 class UGroomSolverComponent;
 
-/** Execution domain that will be used by the groom component optimus data interfaces */
+/** Execution domain that will be used by the groom asset optimus data interfaces */
 UCLASS()
 class UOptimusGroomAssetComponentSource : public UOptimusComponentSource
 {
@@ -91,4 +91,35 @@ public:
 	//~ End UOptimusComputeDataInterface Interface
 };
 
+/** Execution domain that will be used by the groom collision optimus data interfaces */
+UCLASS()
+class UOptimusGroomCollisionComponentSource : public UOptimusComponentSource
+{
+	GENERATED_BODY()
+public:
+
+	static constexpr int32 GroupSize = 64;
+	
+	struct FCollisionExecutionDomains
+	{
+		static FName Vertices;
+		static FName Triangles;
+		static FName Implicits;
+	};
+	
+	//~ Begin  UOptimusComponentSource implementations
+	virtual FText GetDisplayName() const override;
+	virtual FName GetBindingName() const override { return FName("Groom Collision"); }
+	virtual TSubclassOf<UActorComponent> GetComponentClass() const override;
+	virtual TArray<FName> GetExecutionDomains() const override;
+	virtual int32 GetLodIndex(const UActorComponent* InComponent) const override;
+	virtual uint32 GetDefaultNumInvocations(const UActorComponent* InComponent, int32 InLod) const override;
+	virtual bool GetComponentElementCountsForExecutionDomain(FName InDomainName, const UActorComponent* InComponent, int32 InLodIndex, TArray<int32>& OutInvocationElementCounts) const override;
+	virtual bool IsUsableAsPrimarySource() const override;
+	//~ End UOptimusComputeDataInterface Interface
+};
+
+
+
+
 
diff --git a/Engine/Plugins/Runtime/HairStrands/Source/HairStrandsDeformer/Private/DeformerGroomInterfaceUtils.h b/Engine/Plugins/Runtime/HairStrands/Source/HairStrandsDeformer/Private/DeformerGroomInterfaceUtils.h
--- a/Engine/Plugins/Runtime/HairStrands/Source/HairStrandsDeformer/Private/DeformerGroomInterfaceUtils.h
+++ b/Engine/Plugins/Runtime/HairStrands/Source/HairStrandsDeformer/Private/DeformerGroomInterfaceUtils.h
@@ -6,6 +6,8 @@
 #include "RenderGraphBuilder.h"
 #include "SkeletalMeshDeformerHelpers.h"
 #include "Components/SkeletalMeshComponent.h"
+#include "Components/StaticMeshComponent.h"
+#include "Engine/StaticMesh.h"
 #include "SkeletalRenderPublic.h"
 #include "SkinnedMeshSceneProxyDesc.h"
 #include "Rendering/SkeletalMeshRenderData.h"
@@ -165,7 +167,7 @@ namespace UE::Groom::Private
 	}
 
 	template<typename DataType>
-	FORCEINLINE bool HaveValidInstanceResources(const DataType& InstanceData)
+	static bool HaveValidInstanceResources(const DataType& InstanceData)
 	{
 		if(InstanceData.RestResource && InstanceData.DeformedResource)
 		{
@@ -178,7 +180,7 @@ namespace UE::Groom::Private
 		return false;
 	}
 
-	FORCEINLINE bool HaveGuidesInstanceResources(const TArray<const FHairGroupInstance*>& GroupInstances)
+	static bool HaveGuidesInstanceResources(const TArray<const FHairGroupInstance*>& GroupInstances)
 	{
 		for(const FHairGroupInstance* GroupInstance : GroupInstances)
 		{
@@ -203,12 +205,12 @@ namespace UE::Groom::Private
 	}
 
 	template<typename DataType>
-	FORCEINLINE bool HaveValidSkinnedResources(const DataType& SkinnedData, const int32 LODIndex)
+	static bool HaveValidSkinnedResources(const DataType& SkinnedData, const int32 LODIndex)
 	{
 		if(SkinnedData.HasValidRootData() && (SkinnedData.RestRootResource->GetRootCount() > 0) &&
 		  (SkinnedData.RestRootResource->LODs.Num() == SkinnedData.DeformedRootResource->LODs.Num()))
 		{
-			if(LODIndex < SkinnedData.RestRootResource->LODs.Num())
+			if(SkinnedData.RestRootResource->LODs.IsValidIndex(LODIndex))
 			{
 				FHairStrandsLODRestRootResource* RestLODDatas = SkinnedData.RestRootResource->LODs[LODIndex];
 				FHairStrandsLODDeformedRootResource* DeformedLODDatas = SkinnedData.DeformedRootResource->LODs[LODIndex];
@@ -219,7 +221,7 @@ namespace UE::Groom::Private
 		return false;
 	}
 
-	FORCEINLINE bool HaveGuidesSkinnedResources(const TArray<const FHairGroupInstance*>& GroupInstances)
+	static bool HaveGuidesSkinnedResources(const TArray<const FHairGroupInstance*>& GroupInstances)
 	{
 		for(const FHairGroupInstance* GroupInstance : GroupInstances)
 		{
@@ -295,39 +297,36 @@ namespace UE::Groom::Private
 						}
 					}
 				}
-				
-				if (const FHairGroupInstance* Instance = GroomComponent->GetGroupInstance(GroupIndex))
+				if (const TObjectPtr<UGroomAsset> GroomAsset = GroomComponent->GroomAsset)
 				{
-					if(Instance->Strands.IsValid())
+					if (GroomAsset->GetHairGroupsPlatformData().IsValidIndex(GroupIndex))
 					{
+						const FHairGroupPlatformData& GroupData = GroomAsset->GetHairGroupsPlatformData()[GroupIndex];
 						if(DomainName == UOptimusGroomAssetComponentSource::FStrandsExecutionDomains::Curves)
 						{
-							NumElements += bSourceElements ? Instance->Strands.GetData().GetNumSourceCurves() : Instance->Strands.GetData().GetNumCurves();
+							NumElements += bSourceElements ? GroupData.Strands.BulkData.GetNumSourceCurves() : GroupData.Strands.BulkData.GetNumCurves();
 						}
 						else if (DomainName == UOptimusGroomAssetComponentSource::FStrandsExecutionDomains::Edges)
 						{
-							NumElements +=  bSourceElements ? Instance->Strands.GetData().GetNumSourcePoints() - Instance->Strands.GetData().GetNumSourceCurves() :
-									Instance->Strands.GetData().GetNumPoints() - Instance->Strands.GetData().GetNumCurves();
+							NumElements += bSourceElements ? GroupData.Strands.BulkData.GetNumSourcePoints() - GroupData.Strands.BulkData.GetNumSourceCurves() :
+									GroupData.Strands.BulkData.GetNumPoints() - GroupData.Strands.BulkData.GetNumCurves();
 						}
 						else if (DomainName == UOptimusGroomAssetComponentSource::FStrandsExecutionDomains::Points)
 						{
-							NumElements += bSourceElements ? Instance->Strands.GetData().GetNumSourcePoints() : Instance->Strands.GetData().GetNumPoints();
+							NumElements += bSourceElements ? GroupData.Strands.BulkData.GetNumSourcePoints() : GroupData.Strands.BulkData.GetNumPoints();
 						}
-					}
-					if(Instance->Guides.IsValid())
-					{
-						if(DomainName == UOptimusGroomAssetComponentSource::FGuidesExecutionDomains::Curves)
+						else if(DomainName == UOptimusGroomAssetComponentSource::FGuidesExecutionDomains::Curves)
 						{
-							NumElements += bSourceElements ? Instance->Guides.GetData().GetNumSourceCurves() : Instance->Guides.GetData().GetNumCurves();
+							NumElements += bSourceElements ? GroupData.Guides.BulkData.GetNumSourceCurves() : GroupData.Guides.BulkData.GetNumCurves();
 						}
 						else if (DomainName == UOptimusGroomAssetComponentSource::FGuidesExecutionDomains::Edges)
 						{
-							NumElements += bSourceElements ? Instance->Guides.GetData().GetNumSourcePoints() - Instance->Guides.GetData().GetNumSourceCurves() :
-								Instance->Guides.GetData().GetNumPoints() - Instance->Guides.GetData().GetNumCurves();
+							NumElements += bSourceElements ? GroupData.Guides.BulkData.GetNumSourcePoints() - GroupData.Guides.BulkData.GetNumSourceCurves() :
+								GroupData.Guides.BulkData.GetNumPoints() - GroupData.Guides.BulkData.GetNumCurves();
 						}
 						else if (DomainName == UOptimusGroomAssetComponentSource::FGuidesExecutionDomains::Points)
 						{
-							NumElements += bSourceElements ? Instance->Guides.GetData().GetNumSourcePoints() : Instance->Guides.GetData().GetNumPoints();
+							NumElements += bSourceElements ? GroupData.Guides.BulkData.GetNumSourcePoints() : GroupData.Guides.BulkData.GetNumPoints();
 						}
 					}
 				}
@@ -361,6 +360,149 @@ namespace UE::Groom::Private
 		return TotalCount;
 	}
 
+	static int32 FillCollisionOverlappingComponents(const TSet<TObjectPtr<UGroomComponent>>& GroomComponents,
+		const TMap<TObjectPtr<UMeshComponent>, int32>& CollisionComponents, TArray<TPair<TObjectPtr<UMeshComponent>, int32>>& OverlappingComponents)
+	{
+		OverlappingComponents.Reset();
+		for(const TPair<TObjectPtr<UMeshComponent>, int32>& CollisionComponent : CollisionComponents)
+		{
+			if(CollisionComponent.Key)
+			{
+				const FBoxSphereBounds CollisionBounds = CollisionComponent.Key->GetLocalBounds().TransformBy(CollisionComponent.Key->GetComponentTransform());
+				for(const UGroomComponent* GroomComponent : GroomComponents)
+				{
+					if(GroomComponent)
+					{
+						const uint32 NumGroups = GroomComponent->GetGroupCount();
+						for (uint32 GroupIndex = 0; GroupIndex < NumGroups; ++GroupIndex)
+						{
+							if (const FHairGroupInstance* GroupInstance = GroomComponent->GetGroupInstance(GroupIndex))
+							{
+								if(FBoxSphereBounds::BoxesIntersect(GroupInstance->GetBounds(), CollisionBounds))
+								{
+									OverlappingComponents.Add(CollisionComponent);
+								}
+							}
+						}
+					}
+				}
+			}
+		}
+		return OverlappingComponents.Num();
+	}
+
+	static int32 GetCollisionNumInvocations(const TArray<TPair<TObjectPtr<UMeshComponent>, int32>>& CollisionComponents)
+	{
+		int32 NumInvocations = 0;
+		
+		for(const TPair<TObjectPtr<UMeshComponent>, int32>& CollisionComponent : CollisionComponents)
+		{
+			if(CollisionComponent.Key)
+			{
+				if(const USkinnedMeshComponent* SkinnedMeshComponent = Cast<const USkinnedMeshComponent>(CollisionComponent.Key))
+				{
+					if(SkinnedMeshComponent->MeshObject)
+					{	
+						const FSkeletalMeshRenderData& MeshRenderData = SkinnedMeshComponent->MeshObject->GetSkeletalMeshRenderData();
+						{
+							const int32 ValidLod = SkinnedMeshComponent->GetPredictedLODLevel();
+							if(MeshRenderData.LODRenderData.IsValidIndex(ValidLod))
+							{
+								NumInvocations += MeshRenderData.LODRenderData[ValidLod].RenderSections.Num();
+							}
+						}
+					}
+				}
+				else if(const UStaticMeshComponent* StaticMeshComponent = Cast<const UStaticMeshComponent>(CollisionComponent.Key))
+				{
+					if(StaticMeshComponent->GetStaticMesh())
+					{
+						if(const FStaticMeshRenderData* MeshRenderData = StaticMeshComponent->GetStaticMesh()->GetRenderData())
+						{
+							const int32 ValidLod = MeshRenderData->GetCurrentFirstLODIdx(0); 
+							if(MeshRenderData->LODResources.IsValidIndex(ValidLod))
+							{
+								NumInvocations += MeshRenderData->LODResources[ValidLod].Sections.Num();
+							}
+						}
+					}
+				}
+			}
+		}
+		return NumInvocations;		
+	}
+	
+	static int32 GetCollisionInvocationElementCounts(const TArray<TPair<TObjectPtr<UMeshComponent>, int32>>& CollisionComponents, const FName DomainName, TArray<int32>& InvocationCounts)
+	{
+		InvocationCounts.Reset();
+		int32 TotalCount = 0;
+		for(const TPair<TObjectPtr<UMeshComponent>, int32>& CollisionComponent : CollisionComponents)
+		{
+			if(CollisionComponent.Key)
+			{
+				if(const USkinnedMeshComponent* SkinnedMeshComponent = Cast<const USkinnedMeshComponent>(CollisionComponent.Key))
+				{
+					if(SkinnedMeshComponent->MeshObject)
+					{
+						const FSkeletalMeshRenderData& MeshRenderData = SkinnedMeshComponent->MeshObject->GetSkeletalMeshRenderData();
+						{
+							const int32 ValidLod = SkinnedMeshComponent->GetPredictedLODLevel();
+							if(MeshRenderData.LODRenderData.IsValidIndex(ValidLod))
+							{
+								const FSkeletalMeshLODRenderData& LodRenderData = MeshRenderData.LODRenderData[ValidLod];
+								
+								const int32 NumSections = LodRenderData.RenderSections.Num();
+								for(int32 SectionIndex = 0; SectionIndex < NumSections; ++SectionIndex)
+								{
+									if(DomainName == UOptimusGroomCollisionComponentSource::FCollisionExecutionDomains::Vertices)
+									{
+										InvocationCounts.Add(LodRenderData.RenderSections[SectionIndex].NumVertices);
+										TotalCount += InvocationCounts.Last();
+									}
+									else if(DomainName == UOptimusGroomCollisionComponentSource::FCollisionExecutionDomains::Triangles)
+									{
+										InvocationCounts.Add(LodRenderData.RenderSections[SectionIndex].NumTriangles);
+										TotalCount += InvocationCounts.Last();
+									}
+								}
+							}
+						}
+					}
+				}
+				else if(const UStaticMeshComponent* StaticMeshComponent = Cast<const UStaticMeshComponent>(CollisionComponent.Key))
+				{
+					if(StaticMeshComponent->GetStaticMesh())
+					{
+						if(const FStaticMeshRenderData* MeshRenderData = StaticMeshComponent->GetStaticMesh()->GetRenderData())
+						{
+							const int32 ValidLod = MeshRenderData->GetCurrentFirstLODIdx(0);
+							if(MeshRenderData->LODResources.IsValidIndex(ValidLod))
+							{
+								const FStaticMeshLODResources& LodRenderData = MeshRenderData->LODResources[ValidLod];
+
+								const int32 NumSections = LodRenderData.Sections.Num();
+								for(int32 SectionIndex = 0; SectionIndex < NumSections; ++SectionIndex)
+								{
+									if(DomainName == UOptimusGroomCollisionComponentSource::FCollisionExecutionDomains::Vertices)
+									{
+										InvocationCounts.Add(LodRenderData.Sections[SectionIndex].MaxVertexIndex - LodRenderData.Sections[SectionIndex].MinVertexIndex);
+										TotalCount += InvocationCounts.Last();
+									}
+									else if(DomainName == UOptimusGroomCollisionComponentSource::FCollisionExecutionDomains::Triangles)
+									{
+										InvocationCounts.Add(LodRenderData.Sections[SectionIndex].NumTriangles);
+										TotalCount += InvocationCounts.Last();
+									}
+								}
+							}
+						}
+					}
+				}
+			}
+		}
+		return TotalCount;
+	}
+
 	struct FGroupElements
 	{
 		TArray<int32> GroupIndices;
diff --git a/Engine/Plugins/Runtime/HairStrands/Source/HairStrandsEditor/Private/GroomActions.cpp b/Engine/Plugins/Runtime/HairStrands/Source/HairStrandsEditor/Private/GroomActions.cpp
--- a/Engine/Plugins/Runtime/HairStrands/Source/HairStrandsEditor/Private/GroomActions.cpp
+++ b/Engine/Plugins/Runtime/HairStrands/Source/HairStrandsEditor/Private/GroomActions.cpp
@@ -41,6 +41,7 @@
 #include "ToolMenus.h"
 #include "Dataflow/DataflowEditor.h"
 #include "Dataflow/DataflowEditorToolkit.h"
+#include "Dataflow/DataflowSimulationScene.h"
 #include "Dialog/SMessageDialog.h"
 
 #include "ThumbnailRendering/SceneThumbnailInfo.h"
@@ -126,6 +127,27 @@ EAssetCommandResult UAssetDefinition_GroomAsset::OpenAssets(const FAssetOpenArgs
 				
 				GroomAsset->GetDataflowSettings().GetDataflowAsset()->Material = HairMaterial;
 				AssetEditor->Initialize({ GroomAsset}, ActorClass);
+
+				auto ResetDataflowSimulation = [GroomAsset]()
+				{
+					if (UAssetEditorSubsystem* AssetEditorSubsystem = GEditor->GetEditorSubsystem<UAssetEditorSubsystem>())
+					{
+						TArray<IAssetEditorInstance*> EditorInstances = AssetEditorSubsystem->FindEditorsForAssetAndSubObjects(GroomAsset);
+						for (IAssetEditorInstance* EditorInstance : EditorInstances)
+						{
+							if (EditorInstance->GetEditorName() == "DataflowEditor")
+							{
+								if (FDataflowEditorToolkit* DataflowEditorToolkit = static_cast<FDataflowEditorToolkit*>(EditorInstance))
+								{
+									DataflowEditorToolkit->ResetDataflowSimulation();
+								}
+							}
+						}
+					}
+				};
+
+				GroomAsset->GetOnGroomAssetResourcesChanged().AddLambda(ResetDataflowSimulation);
+				GroomAsset->GetOnGroomAssetChanged().AddLambda(ResetDataflowSimulation);
 			}
 		}
 	}
diff --git a/Engine/Plugins/Runtime/HairStrands/Source/HairStrandsSolver/Private/GroomSolverComponent.cpp b/Engine/Plugins/Runtime/HairStrands/Source/HairStrandsSolver/Private/GroomSolverComponent.cpp
--- a/Engine/Plugins/Runtime/HairStrands/Source/HairStrandsSolver/Private/GroomSolverComponent.cpp
+++ b/Engine/Plugins/Runtime/HairStrands/Source/HairStrandsSolver/Private/GroomSolverComponent.cpp
@@ -1,15 +1,20 @@
 // Copyright Epic Games, Inc. All Rights Reserved.
 
 #include "GroomSolverComponent.h"
+
+#include <ThirdParty/Catch2/v3.4.0/extras/catch_amalgamated.hpp>
+#include <ThirdParty/ShaderConductor/ShaderConductor/External/DirectXShaderCompiler/include/dxc/DXIL/DxilConstants.h>
+
 #include "GroomComponent.h"
 #include "Animation/MeshDeformer.h"
 #include "Animation/MeshDeformerInstance.h"
+#include "VerseVM/VVMInstantiationContext.h"
 
 #include UE_INLINE_GENERATED_CPP_BY_NAME(GroomSolverComponent)
 
 namespace UE::Groom::Private
 {
-	FORCEINLINE void GatherViewLocations(const UWorld* LocalWorld, TArray<FVector>& ViewLocations)
+	static void GatherViewLocations(const UWorld* LocalWorld, TArray<FVector>& ViewLocations)
 	{
 		if(LocalWorld)
 		{
@@ -33,6 +38,35 @@ namespace UE::Groom::Private
 			}
 		}
 	}
+
+	static FBox GetGroomBounds(const UGroomComponent* GroomComponent, const int32 GroupIndex)
+	{
+		FBox LocalBounds(EForceInit::ForceInitToZero);
+		if(GroomComponent && GroomComponent->GroomAsset)
+		{
+			const FTransform LocalToWorld = GroomComponent->GetComponentTransform();
+			const FHairGroupPlatformData& GroupData = GroomComponent->GroomAsset->GetHairGroupsPlatformData()[GroupIndex];
+			if (GroupData.Guides.HasValidData())
+			{
+				// The hair width is not scaled by the transform scale. Since we increase the local bound by the hair width, 
+				// and that amount is then scaled by the local to world matrix which conatins a scale factor, we apply a 
+				// inverse scale to the radius computation
+				const float InvScale = 1.f / LocalToWorld.GetMaximumAxisScale();
+
+				// Take into account the strands size and add it as a bound 'border'
+				const FHairGroupDesc Desc = GroomComponent->GroomGroupsDesc[GroupIndex];
+				const FVector HairBorder(0.5f * Desc.HairWidth * FMath::Max3(1.0f, Desc.HairRootScale, Desc.HairTipScale) * InvScale);
+
+				LocalBounds += GroupData.Guides.GetBounds();
+
+				LocalBounds.Min -= HairBorder;
+				LocalBounds.Max += HairBorder;
+				
+				return LocalBounds.TransformBy(LocalToWorld);
+			}
+		}
+		return LocalBounds;
+	}
 }
 
 UGroomSolverComponent::UGroomSolverComponent(const FObjectInitializer& ObjectInitializer)
@@ -119,6 +153,17 @@ void UGroomSolverComponent::SelectDynamicCurves()
 				}
 				
 				const uint32 NumGroups = GroomComponent->GetGroupCount();
+
+				int32 TotalNumPoints = 0;
+				for (uint32 GroupIndex = 0; GroupIndex < NumGroups; ++GroupIndex)
+				{
+					const FHairGroupInstance* GroupInstance = GroomComponent->GetGroupInstance(GroupIndex);
+					if(GroomComponent->IsDeformationEnable(GroupIndex) && GroupInstance && GroupInstance->Guides.IsValid())
+					{
+						TotalNumPoints += (1.0-DistanceRatio) * GroupInstance->Guides.GetData().GetNumPoints();
+					}
+				}
+				const float PointsScale = FMath::Clamp(GetSolverSettings().MaxPointsCount / static_cast<float>(TotalNumPoints), 0.0f, 1.0f);
 				for (uint32 GroupIndex = 0; GroupIndex < NumGroups; ++GroupIndex)
 				{
 					const FHairGroupInstance* GroupInstance = GroomComponent->GetGroupInstance(GroupIndex);
@@ -128,7 +173,7 @@ void UGroomSolverComponent::SelectDynamicCurves()
 						const int32 NumTotalPoints = GroupInstance->Guides.GetData().GetNumPoints();
 
 						const int32 NumCurvePoints = NumTotalPoints / NumTotalCurves;
-						const int32 NumDynamicCurves = (1.0-DistanceRatio) * NumTotalCurves;
+						const int32 NumDynamicCurves = (1.0-DistanceRatio) * NumTotalCurves * PointsScale;
 						
 						const int32 NumDynamicPoints = NumDynamicCurves * NumCurvePoints;
 						const int32 NumKinematicPoints = NumTotalPoints-NumDynamicPoints;
@@ -168,10 +213,11 @@ void UGroomSolverComponent::SelectDynamicCurves()
 						
 						CurveOffset += NumTotalCurves;
 						PointOffset += NumTotalPoints;
+						
+						// Keep it for now for debugging since it is WIP
+                        UE_LOG(LogTemp, Log, TEXT("Num Curves = %d Distance Ratio = %f Simulated Curves = %d | Num LODs = %d | Points Scale = %f"),
+							NumTotalCurves, DistanceRatio, NumDynamicCurves, FMath::CeilLogTwo(NumTotalCurves), PointsScale);
 					}
-				
-					// Keep it for now for debugging since it is WIP
-					//UE_LOG(LogTemp, Log, TEXT("Num Curves = %d Distance Ratio = %f Simulated Curves = %f | Num LODs = %d"), NumCurves, DistanceRatio, (1.0-DistanceRatio) * NumCurves, FMath::CeilLogTwo(NumCurves));
 				}
 			}
 		}
@@ -245,6 +291,21 @@ void UGroomSolverComponent::ResetGroomComponents()
 	GroomComponents.Reset();
 }
 
+void UGroomSolverComponent::AddCollisionComponent(UMeshComponent* CollisionComponent, const int32 LODIndex)
+{
+	CollisionComponents.Add(CollisionComponent, LODIndex);
+}
+
+void UGroomSolverComponent::RemoveCollisionComponent(UMeshComponent* CollisionComponent)
+{
+	CollisionComponents.Remove(CollisionComponent);
+}
+
+void UGroomSolverComponent::ResetCollisionComponents()
+{
+	CollisionComponents.Reset();
+}
+
 void UGroomSolverComponent::SetDeformerSolver(UMeshDeformer* DeformerSolver)
 {
 	MeshDeformer = DeformerSolver;
diff --git a/Engine/Plugins/Runtime/HairStrands/Source/HairStrandsSolver/Public/GroomSolverComponent.h b/Engine/Plugins/Runtime/HairStrands/Source/HairStrandsSolver/Public/GroomSolverComponent.h
--- a/Engine/Plugins/Runtime/HairStrands/Source/HairStrandsSolver/Public/GroomSolverComponent.h
+++ b/Engine/Plugins/Runtime/HairStrands/Source/HairStrandsSolver/Public/GroomSolverComponent.h
@@ -20,16 +20,16 @@ struct FGroomSolverSettings
 	GENERATED_BODY()
 	
 	/** Number of temporal sub-steps */
-	UPROPERTY(EditAnywhere, Category = "Timing")
-	int32 NumSubSteps = 2;
+	//UPROPERTY(EditAnywhere, Category = "Timing")
+	//int32 NumSubSteps = 2;
 
 	/** Number of iterations for the constraint solver */
-	UPROPERTY(EditAnywhere, Category = "Timing")
-	int32 NumSolverIterations = 5;
+	//UPROPERTY(EditAnywhere, Category = "Timing")
+	//int32 NumSolverIterations = 5;
 
 	/** Frame rate for the simulation to run */
-	UPROPERTY(EditAnywhere, Category = "Timing")
-	float SimulationFrameRate = 30;
+	//UPROPERTY(EditAnywhere, Category = "Timing")
+	//float SimulationFrameRate = 30;
 
 	/** Maximum LOD distance (if distance between the component and the views is higher that this threshold, no simulation)*/
 	UPROPERTY(EditAnywhere, Category = "LODs", meta = (ClampMin = "0", UIMin = "0"))
@@ -39,6 +39,10 @@ struct FGroomSolverSettings
 	UPROPERTY(EditAnywhere, Category = "LODs", meta = (ClampMin = "0", UIMin = "0"))
 	float MinLODDistance = 100;
 
+	/** Maximum number of points the solver is going to simulate */
+	UPROPERTY(EditAnywhere, Category = "LODs")
+	float MaxPointsCount = 100000;
+
 	/** List of dynamic curves */
 	TArray<int32> CurveDynamicIndices;
 
@@ -143,17 +147,32 @@ public:
 	UFUNCTION(BlueprintCallable, Category = "Groom")
 	UE_API void AddGroomComponent(UGroomComponent* GroomComponent);
 	
-	/** Add a groom component to the solver */
+	/** Remove a groom component to the solver */
 	UFUNCTION(BlueprintCallable, Category = "Groom")
 	UE_API void RemoveGroomComponent(UGroomComponent* GroomComponent);
 
-	/** Reset the groom physics components */
+	/** Reset the solver groom components */
     UFUNCTION(BlueprintCallable, Category = "Groom")
     UE_API void ResetGroomComponents();
 
-	/** Reset the groom physics components */
+	/** Add a collision component to the solver */
+	UFUNCTION(BlueprintCallable, Category = "Groom")
+	UE_API void AddCollisionComponent(UMeshComponent* CollisionComponent, const int32 LODIndex);
+
+	/** Remove a collision component to the solver */
+	UFUNCTION(BlueprintCallable, Category = "Groom")
+	UE_API void RemoveCollisionComponent(UMeshComponent* CollisionComponent);
+
+	/** Reset the solver collision components */
+	UFUNCTION(BlueprintCallable, Category = "Groom")
+	UE_API void ResetCollisionComponents();
+
+	/** Get the solver groom components */
 	const TSet< TObjectPtr<UGroomComponent> >& GetGroomComponents() const {return GroomComponents;}
 
+	/** Get the solver collision components */
+	const TMap< TObjectPtr<UMeshComponent>, int32 >& GetCollisionComponents() const {return CollisionComponents;}
+
 	/* Change the MeshDeformer solver that is used for this Component. */
 	UFUNCTION(BlueprintCallable, Category = "Groom")
 	UE_API void SetDeformerSolver(UMeshDeformer* DeformerSolver);
@@ -195,6 +214,10 @@ private :
 
 	/** Groom solver proxy to be used in dataflow */
 	FDataflowGroomSolverProxy GroomSolverProxy;
+
+	/** List of collision components registered to the solver */
+	UPROPERTY(VisibleAnywhere, Category = "Groom")
+	TMap< TObjectPtr<UMeshComponent>, int32 > CollisionComponents;
 };
 
 #undef UE_API 
