Fix CIS with bad includes in the groom solver component

diff --git a/Engine/Plugins/Runtime/HairStrands/Source/HairStrandsSolver/Private/GroomSolverComponent.cpp b/Engine/Plugins/Runtime/HairStrands/Source/HairStrandsSolver/Private/GroomSolverComponent.cpp
--- a/Engine/Plugins/Runtime/HairStrands/Source/HairStrandsSolver/Private/GroomSolverComponent.cpp
+++ b/Engine/Plugins/Runtime/HairStrands/Source/HairStrandsSolver/Private/GroomSolverComponent.cpp
@@ -1,14 +1,9 @@
 // Copyright Epic Games, Inc. All Rights Reserved.
 
 #include "GroomSolverComponent.h"
-
-#include <ThirdParty/Catch2/v3.4.0/extras/catch_amalgamated.hpp>
-#include <ThirdParty/ShaderConductor/ShaderConductor/External/DirectXShaderCompiler/include/dxc/DXIL/DxilConstants.h>
-
 #include "GroomComponent.h"
 #include "Animation/MeshDeformer.h"
 #include "Animation/MeshDeformerInstance.h"
-#include "VerseVM/VVMInstantiationContext.h"
 
 #include UE_INLINE_GENERATED_CPP_BY_NAME(GroomSolverComponent)
 
diff --git a/Engine/Plugins/Runtime/HairStrands/Source/HairStrandsSolver/Public/GroomSolverComponent.h b/Engine/Plugins/Runtime/HairStrands/Source/HairStrandsSolver/Public/GroomSolverComponent.h
--- a/Engine/Plugins/Runtime/HairStrands/Source/HairStrandsSolver/Public/GroomSolverComponent.h
+++ b/Engine/Plugins/Runtime/HairStrands/Source/HairStrandsSolver/Public/GroomSolverComponent.h
@@ -18,18 +18,6 @@ USTRUCT()
 struct FGroomSolverSettings
 {
 	GENERATED_BODY()
-	
-	/** Number of temporal sub-steps */
-	//UPROPERTY(EditAnywhere, Category = "Timing")
-	//int32 NumSubSteps = 2;
-
-	/** Number of iterations for the constraint solver */
-	//UPROPERTY(EditAnywhere, Category = "Timing")
-	//int32 NumSolverIterations = 5;
-
-	/** Frame rate for the simulation to run */
-	//UPROPERTY(EditAnywhere, Category = "Timing")
-	//float SimulationFrameRate = 30;
 
 	/** Maximum LOD distance (if distance between the component and the views is higher that this threshold, no simulation)*/
 	UPROPERTY(EditAnywhere, Category = "LODs", meta = (ClampMin = "0", UIMin = "0"))
