Bending model for groom + geometric collision+ guides solver

#virtualized

diff --git a/Engine/Plugins/Runtime/HairStrands/HairStrands.uplugin b/Engine/Plugins/Runtime/HairStrands/HairStrands.uplugin
--- a/Engine/Plugins/Runtime/HairStrands/HairStrands.uplugin
+++ b/Engine/Plugins/Runtime/HairStrands/HairStrands.uplugin
@@ -23,7 +23,7 @@
 		{
 			"Name": "HairStrandsDeformer",
 			"Type": "Runtime",
-			"LoadingPhase": "Default",
+			"LoadingPhase": "PreDefault",
 			"PlatformDenyList": [
 				"Android",
 				"IOS"
@@ -52,7 +52,7 @@
 		{
 			"Name": "HairStrandsSolver",
 			"Type": "Runtime",
-			"LoadingPhase": "Default"
+			"LoadingPhase": "PreDefault"
 		}
 	],
 	"Plugins": [
diff --git a/Engine/Plugins/Runtime/HairStrands/Shaders/Private/Deformers/DeformerGroomCollisionRead.ush b/Engine/Plugins/Runtime/HairStrands/Shaders/Private/Deformers/DeformerGroomCollisionRead.ush
new file mode 100644
--- /dev/null
+++ b/Engine/Plugins/Runtime/HairStrands/Shaders/Private/Deformers/DeformerGroomCollisionRead.ush
@@ -0,0 +1,46 @@
+﻿// Copyright Epic Games, Inc. All Rights Reserved.
+
+#include "/Plugin/ComputeFramework/Private/ComputeKernelCommon.ush"
+
+uint {DataInterfaceName}_NumVertices; 
+uint {DataInterfaceName}_NumTriangles;
+uint {DataInterfaceName}_VertexOffset; 
+uint {DataInterfaceName}_TriangleOffset;
+uint {DataInterfaceName}_TotalNumVertices; 
+uint {DataInterfaceName}_TotalNumTriangles;
+uint {DataInterfaceName}_BaseVertexIndex; 
+uint {DataInterfaceName}_BaseTriangleIndex;
+
+Buffer<float> {DataInterfaceName}_VertexGlobalPositions;
+Buffer<uint> {DataInterfaceName}_TriangleVertexIndices;
+
+// Interface functions
+
+uint ReadNumVertices_{DataInterfaceName}()
+{
+	return {DataInterfaceName}_TotalNumVertices;
+}
+
+uint ReadNumTriangles_{DataInterfaceName}()
+{
+	return {DataInterfaceName}_TotalNumTriangles;
+}
+
+float3 ReadVertexGlobalPosition_{DataInterfaceName}(uint VertexIndex)
+{
+	VertexIndex -= {DataInterfaceName}_BaseVertexIndex;
+	const uint VertexOffset = 3 * (VertexIndex + {DataInterfaceName}_VertexOffset);
+	return (VertexIndex < {DataInterfaceName}_NumVertices) ?
+		float3({DataInterfaceName}_VertexGlobalPositions[VertexOffset], {DataInterfaceName}_VertexGlobalPositions[VertexOffset+1], {DataInterfaceName}_VertexGlobalPositions[VertexOffset+2]) : float3(0,0,0);
+}
+
+uint3 ReadTriangleVertexIndices_{DataInterfaceName}(uint TriangleIndex)
+{
+	TriangleIndex -= {DataInterfaceName}_BaseTriangleIndex;
+	const uint TriangleOffset = 3*TriangleIndex + {DataInterfaceName}_TriangleOffset;
+	return (TriangleIndex < {DataInterfaceName}_NumTriangles) ?
+		uint3({DataInterfaceName}_TriangleVertexIndices[TriangleOffset]+{DataInterfaceName}_BaseVertexIndex-{DataInterfaceName}_VertexOffset,
+			  {DataInterfaceName}_TriangleVertexIndices[TriangleOffset+1]+{DataInterfaceName}_BaseVertexIndex-{DataInterfaceName}_VertexOffset,
+			  {DataInterfaceName}_TriangleVertexIndices[TriangleOffset+2]+{DataInterfaceName}_BaseVertexIndex-{DataInterfaceName}_VertexOffset) : uint3(0,0,0);
+}
+
diff --git a/Engine/Plugins/Runtime/HairStrands/Shaders/Private/Deformers/DeformerGroomSolverRead.ush b/Engine/Plugins/Runtime/HairStrands/Shaders/Private/Deformers/DeformerGroomSolverRead.ush
--- a/Engine/Plugins/Runtime/HairStrands/Shaders/Private/Deformers/DeformerGroomSolverRead.ush
+++ b/Engine/Plugins/Runtime/HairStrands/Shaders/Private/Deformers/DeformerGroomSolverRead.ush
@@ -1,7 +1,38 @@
-﻿// Copyright Epic Games, Inc. All Rights Reserved.
+// Copyright Epic Games, Inc. All Rights Reserved.
 
 #include "/Plugin/ComputeFramework/Private/ComputeKernelCommon.ush"
 
+#ifdef RAYTRACING_SOLVER_BINDING
+
+#include "/Engine/Private/Common.ush"
+#include "/Engine/Private/SceneData.ush"
+#include "/Engine/Private/ShadingCommon.ush"
+#include "/Engine/Private/RayTracing/RayTracingCommon.ush"
+#include "/Engine/Private/RayTracing/RayTracingHitGroupCommon.ush"
+
+#ifndef DATA_INTERFACE_ACCELERATION_STRUCTURE_ONCE
+#define DATA_INTERFACE_ACCELERATION_STRUCTURE_ONCE 1
+
+struct FAccelerationStructure
+{
+	RaytracingAccelerationStructure	WorldTLAS;
+	int ValidTLAS;
+};
+
+#endif // DATA_INTERFACE_ACCELERATION_STRUCTURE_ONCE
+
+RaytracingAccelerationStructure	{DataInterfaceName}_WorldAccelerationStructure;
+
+FAccelerationStructure ReadWorldAccelerationStructure_{DataInterfaceName}()
+{
+	FAccelerationStructure AccelerationStructure;
+	AccelerationStructure.ValidTLAS = true;
+	AccelerationStructure.WorldTLAS = {DataInterfaceName}_WorldAccelerationStructure;
+	return AccelerationStructure;
+}
+
+#endif
+
 uint {DataInterfaceName}_NumSolverPoints;
 uint {DataInterfaceName}_NumSolverCurves;
 uint {DataInterfaceName}_NumSolverObjects;
@@ -12,6 +43,9 @@ uint {DataInterfaceName}_NumDynamicCurves;
 uint {DataInterfaceName}_NumKinematicPoints;
 uint {DataInterfaceName}_NumKinematicCurves;
 
+float3 {DataInterfaceName}_BoundingBoxMin;
+float3 {DataInterfaceName}_BoundingBoxMax;
+
 Buffer<uint> {DataInterfaceName}_ObjectCurveOffsets;
 Buffer<uint> {DataInterfaceName}_ObjectPointOffsets;
 
@@ -66,6 +100,16 @@ uint ReadNumKinematicCurves_{DataInterfaceName}()
 	return {DataInterfaceName}_NumKinematicCurves;
 }
 
+float3 ReadBoundingBoxMin_{DataInterfaceName}()
+{
+	return {DataInterfaceName}_BoundingBoxMin;
+}
+
+float3 ReadBoundingBoxMax_{DataInterfaceName}()
+{
+	return {DataInterfaceName}_BoundingBoxMax;
+}
+
 uint ReadObjectPointOffset_{DataInterfaceName}(uint ObjectIndex)
 {
 	return (ObjectIndex < ({DataInterfaceName}_NumSolverObjects+1)) ? {DataInterfaceName}_ObjectPointOffsets[ObjectIndex] : 0;
@@ -124,3 +168,4 @@ int ReadKinematicCurveIndex_{DataInterfaceName}(uint CurveIndex)
 
 
 
+
diff --git a/Engine/Plugins/Runtime/HairStrands/Source/HairStrandsDataflow/Private/GroomAssetTerminalNode.cpp b/Engine/Plugins/Runtime/HairStrands/Source/HairStrandsDataflow/Private/GroomAssetTerminalNode.cpp
--- a/Engine/Plugins/Runtime/HairStrands/Source/HairStrandsDataflow/Private/GroomAssetTerminalNode.cpp
+++ b/Engine/Plugins/Runtime/HairStrands/Source/HairStrandsDataflow/Private/GroomAssetTerminalNode.cpp
@@ -330,7 +330,7 @@ namespace UE::Groom::Private
 		}
 	}
 
-	FORCEINLINE void RegisterSkeletalMeshes(const FManagedArrayCollection& InCollection, const FName& SkelMeshAttribute, const FName& MeshLODAttribute, const FName& GeometryGroup, const int32 NumGeometry, UGroomAsset* GroomAsset)
+	static void RegisterSkeletalMeshes(const FManagedArrayCollection& InCollection, const FName& SkelMeshAttribute, const FName& MeshLODAttribute, const FName& GeometryGroup, const int32 NumGeometry, UGroomAsset* GroomAsset)
 	{
 		const TManagedArray<TObjectPtr<UObject>>& ObjectSkeletalMeshes =
 						InCollection.GetAttribute<TObjectPtr<UObject>>(SkelMeshAttribute, GeometryGroup);
diff --git a/Engine/Plugins/Runtime/HairStrands/Source/HairStrandsDeformer/HairStrandsDeformer.Build.cs b/Engine/Plugins/Runtime/HairStrands/Source/HairStrandsDeformer/HairStrandsDeformer.Build.cs
--- a/Engine/Plugins/Runtime/HairStrands/Source/HairStrandsDeformer/HairStrandsDeformer.Build.cs
+++ b/Engine/Plugins/Runtime/HairStrands/Source/HairStrandsDeformer/HairStrandsDeformer.Build.cs
@@ -1,5 +1,6 @@
 // Copyright Epic Games, Inc. All Rights Reserved.
 
+using System.IO;
 namespace UnrealBuildTool.Rules
 {
 	public class HairStrandsDeformer : ModuleRules
@@ -29,6 +30,12 @@ namespace UnrealBuildTool.Rules
 				{
 					"Shaders",
 				});
+			
+			PrivateIncludePaths.AddRange(
+				new string[] {
+					Path.Combine(EngineDirectory,"Source/Runtime/Renderer/Private"),
+				}
+			);
 		}
 	}
 }
diff --git a/Engine/Plugins/Runtime/HairStrands/Source/HairStrandsDeformer/Private/DeformerGroomCollisionRead.cpp b/Engine/Plugins/Runtime/HairStrands/Source/HairStrandsDeformer/Private/DeformerGroomCollisionRead.cpp
new file mode 100644
--- /dev/null
+++ b/Engine/Plugins/Runtime/HairStrands/Source/HairStrandsDeformer/Private/DeformerGroomCollisionRead.cpp
@@ -0,0 +1,258 @@
+// Copyright Epic Games, Inc. All Rights Reserved.
+
+#include "DeformerGroomCollisionRead.h"
+
+#include "CachedGeometry.h"
+#include "ComputeFramework/ShaderParamTypeDefinition.h"
+#include "ShaderCompilerCore.h"
+#include "ShaderParameterMetadataBuilder.h"
+#include "OptimusDataTypeRegistry.h"
+#include "GroomSolverComponent.h"
+#include "RenderGraphBuilder.h"
+#include "GlobalRenderResources.h"
+#include "DeformerGroomDomainsSource.h"
+#include "DeformerGroomInterfaceUtils.h"
+#include "SceneInterface.h"
+#include "ScenePrivate.h"
+#include "SystemTextures.h"
+#include "SkeletalRenderPublic.h"
+//#include "Editor/UnrealEdTypes.h"
+
+FString UOptimusGroomCollisionReadDataInterface::GetDisplayName() const
+{
+	return TEXT("Groom Collision");
+}
+
+TArray<FOptimusCDIPinDefinition> UOptimusGroomCollisionReadDataInterface::GetPinDefinitions() const
+{
+	static const FName CollisionVertices(UOptimusGroomCollisionComponentSource::FCollisionExecutionDomains::Vertices);
+	static const FName CollisionTriangles(UOptimusGroomCollisionComponentSource::FCollisionExecutionDomains::Triangles);
+
+	TArray<FOptimusCDIPinDefinition> Defs;
+	
+	Defs.Add({ "NumCollisionVertices",   "ReadNumVertices", true, "NumCollisionVertices" });
+	Defs.Add({ "VertexGlobalPosition", "ReadVertexGlobalPosition", CollisionVertices,   "ReadNumVertices", true, "VertexGlobalPosition"  });
+	Defs.Add({ "NumCollisionTriangles",   "ReadNumTriangles", true, "NumCollisionTriangles" });
+	Defs.Add({ "TriangleVertexIndices", "ReadTriangleVertexIndices", CollisionTriangles,   "ReadNumTriangles", true, "TriangleVertexIndices"  });
+
+	return Defs;
+}
+ 
+TSubclassOf<UActorComponent> UOptimusGroomCollisionReadDataInterface::GetRequiredComponentClass() const
+{
+	return UGroomSolverComponent::StaticClass();
+}
+
+void UOptimusGroomCollisionReadDataInterface::GetPermutations(FComputeKernelPermutationVector& OutPermutationVector) const
+{}
+
+void UOptimusGroomCollisionReadDataInterface::GetSupportedInputs(TArray<FShaderFunctionDefinition>& OutFunctions) const
+{
+	OutFunctions.AddDefaulted_GetRef()
+		.SetName(TEXT("ReadNumVertices"))
+		.AddReturnType(EShaderFundamentalType::Uint);
+
+	OutFunctions.AddDefaulted_GetRef()
+		.SetName(TEXT("ReadNumTriangles"))
+		.AddReturnType(EShaderFundamentalType::Uint);
+
+	OutFunctions.AddDefaulted_GetRef()
+		.SetName(TEXT("ReadVertexGlobalPosition"))
+		.AddReturnType(EShaderFundamentalType::Float, 3)
+		.AddParam(EShaderFundamentalType::Uint);
+		
+	OutFunctions.AddDefaulted_GetRef()
+		.SetName(TEXT("ReadTriangleVertexIndices"))
+		.AddReturnType(EShaderFundamentalType::Int, 3)
+		.AddParam(EShaderFundamentalType::Uint);
+}
+
+BEGIN_SHADER_PARAMETER_STRUCT(FOptimusGroomCollisionReadParameters, )
+	SHADER_PARAMETER(uint32, NumVertices)
+	SHADER_PARAMETER(uint32, NumTriangles)
+	SHADER_PARAMETER(uint32, VertexOffset)
+	SHADER_PARAMETER(uint32, TriangleOffset)
+	SHADER_PARAMETER(uint32, BaseVertexIndex)
+	SHADER_PARAMETER(uint32, BaseTriangleIndex)
+	SHADER_PARAMETER(uint32, TotalNumVertices)
+	SHADER_PARAMETER(uint32, TotalNumTriangles)
+	SHADER_PARAMETER_SRV(Buffer<float3>, VertexGlobalPositions)
+	SHADER_PARAMETER_SRV(Buffer<uint3>, TriangleVertexIndices)
+END_SHADER_PARAMETER_STRUCT()
+
+void UOptimusGroomCollisionReadDataInterface::GetShaderParameters(TCHAR const* UID, FShaderParametersMetadataBuilder& InOutBuilder, FShaderParametersMetadataAllocations& InOutAllocations) const
+{
+	InOutBuilder.AddNestedStruct<FOptimusGroomCollisionReadParameters>(UID);
+}
+
+TCHAR const* UOptimusGroomCollisionReadDataInterface::TemplateFilePath = TEXT("/Plugin/Runtime/HairStrands/Private/Deformers/DeformerGroomCollisionRead.ush");
+
+TCHAR const* UOptimusGroomCollisionReadDataInterface::GetShaderVirtualPath() const
+{
+	return TemplateFilePath;
+}
+
+void UOptimusGroomCollisionReadDataInterface::GetShaderHash(FString& InOutKey) const
+{
+	GetShaderFileHash(TemplateFilePath, EShaderPlatform::SP_PCD3D_SM5).AppendString(InOutKey);
+}
+
+void UOptimusGroomCollisionReadDataInterface::GetHLSL(FString& OutHLSL, FString const& InDataInterfaceName) const
+{
+	TMap<FString, FStringFormatArg> TemplateArgs =
+	{
+		{ TEXT("DataInterfaceName"), InDataInterfaceName },
+	};
+
+	FString TemplateFile;
+	LoadShaderSourceFile(TemplateFilePath, EShaderPlatform::SP_PCD3D_SM5, &TemplateFile, nullptr);
+	OutHLSL += FString::Format(*TemplateFile, TemplateArgs);
+}
+
+UComputeDataProvider* UOptimusGroomCollisionReadDataInterface::CreateDataProvider(TObjectPtr<UObject> InBinding, uint64 InInputMask, uint64 InOutputMask) const
+{
+	UOptimusGroomCollisionReadDataProvider* Provider = NewObject<UOptimusGroomCollisionReadDataProvider>();
+	Provider->SolverComponent = Cast<UGroomSolverComponent>(InBinding);
+	return Provider;
+}
+
+FComputeDataProviderRenderProxy* UOptimusGroomCollisionReadDataProvider::GetRenderProxy()
+{
+	return new FOptimusGroomCollisionReadDataProviderProxy(SolverComponent);
+}
+
+FOptimusGroomCollisionReadDataProviderProxy::FOptimusGroomCollisionReadDataProviderProxy(UGroomSolverComponent* SolverComponent)
+{
+	if(SolverComponent)
+	{
+		CollisionObjects = SolverComponent->GetCollisionComponents();
+	}
+}
+
+bool FOptimusGroomCollisionReadDataProviderProxy::IsValid(FValidationData const& InValidationData) const
+{
+	if (InValidationData.ParameterStructSize != sizeof(FParameters))
+	{
+		return false;
+	}
+
+	int32 NumSections = 0;
+	for(const TPair<TObjectPtr<UMeshComponent>, int32>& CollisionComponent : CollisionObjects)
+	{
+		if(CollisionComponent.Key)
+		{
+			if(const USkinnedMeshComponent* SkinnedMeshComponent = Cast<const USkinnedMeshComponent>(CollisionComponent.Key))
+			{
+				if(SkinnedMeshComponent->MeshObject)
+				{
+					const FSkeletalMeshRenderData& MeshRenderData = SkinnedMeshComponent->MeshObject->GetSkeletalMeshRenderData();
+					NumSections += MeshRenderData.LODRenderData[SkinnedMeshComponent->MeshObject->GetLOD()].RenderSections.Num();
+				}
+			}
+		}
+	}
+	if (InValidationData.NumInvocations != NumSections)
+	{
+		return false;
+	}
+	
+	return true;
+}
+
+void FOptimusGroomCollisionReadDataProviderProxy::AllocateResources(FRDGBuilder& GraphBuilder, FAllocationData const& InAllocationData)
+{
+	if (!FallbackStructuredSRV) { FallbackStructuredSRV  = GraphBuilder.CreateSRV(GSystemTextures.GetDefaultStructuredBuffer(GraphBuilder, 16u)); }
+
+	CollisionResources.Reset();
+	for(const TPair<TObjectPtr<UMeshComponent>, int32>& CollisionComponent : CollisionObjects)
+	{
+		if(CollisionComponent.Key)
+		{
+			if(const USkinnedMeshComponent* SkinnedMeshComponent = Cast<const USkinnedMeshComponent>(CollisionComponent.Key))
+			{
+				if(SkinnedMeshComponent->MeshObject)
+				{
+					FCachedGeometry CachedGeometry;
+					SkinnedMeshComponent->MeshObject->GetCachedGeometry(GraphBuilder, CachedGeometry);
+
+					const FSkeletalMeshRenderData& MeshRenderData = SkinnedMeshComponent->MeshObject->GetSkeletalMeshRenderData();
+					check(CachedGeometry.Sections.Num() == MeshRenderData.LODRenderData[SkinnedMeshComponent->MeshObject->GetLOD()].RenderSections.Num());
+
+					for(int32 SectionIndex = 0; SectionIndex < CachedGeometry.Sections.Num(); ++SectionIndex)
+					{
+						FCollisionResources& CollisionResource = CollisionResources.AddDefaulted_GetRef();
+						
+						CollisionResource.VertexPositions = CachedGeometry.Sections[SectionIndex].PositionBuffer;
+						CollisionResource.TriangleIndices = CachedGeometry.Sections[SectionIndex].IndexBuffer;
+						CollisionResource.NumTriangles = CachedGeometry.Sections[SectionIndex].NumPrimitives;
+						CollisionResource.NumVertices = CachedGeometry.Sections[SectionIndex].NumVertices;
+						CollisionResource.TriangleOffset = CachedGeometry.Sections[SectionIndex].IndexBaseIndex;
+						CollisionResource.VertexOffset = CachedGeometry.Sections[SectionIndex].VertexBaseIndex;
+					}
+				} 
+			}
+			else if(const UStaticMeshComponent* StaticMeshComponent = Cast<const UStaticMeshComponent>(CollisionComponent.Key))
+			{
+				if(StaticMeshComponent->GetStaticMesh())
+				{
+					if(const FStaticMeshRenderData* MeshRenderData = StaticMeshComponent->GetStaticMesh()->GetRenderData())
+					{
+						const int32 ValidLod = MeshRenderData->GetCurrentFirstLODIdx(0);
+						if(MeshRenderData->LODResources.IsValidIndex(ValidLod))
+						{
+							for(int32 SectionIndex = 0; SectionIndex < MeshRenderData->LODResources[ValidLod].Sections.Num(); ++SectionIndex)
+							{
+								FCollisionResources& CollisionResource = CollisionResources.AddDefaulted_GetRef();
+
+								const FStaticMeshSection& SectionData = MeshRenderData->LODResources[ValidLod].Sections[SectionIndex];
+						
+								CollisionResource.VertexPositions = MeshRenderData->LODResources[ValidLod].VertexBuffers.PositionVertexBuffer.GetSRV();
+								CollisionResource.TriangleIndices = GraphBuilder.RHICmdList.CreateShaderResourceView(MeshRenderData->LODResources[ValidLod].IndexBuffer.IndexBufferRHI,
+									FRHIViewDesc::CreateBufferSRV().SetType(FRHIViewDesc::EBufferType::Typed).SetFormat(
+										MeshRenderData->LODResources[ValidLod].IndexBuffer.Is32Bit() ? PF_R32_UINT : PF_R16_UINT));
+								CollisionResource.NumTriangles = SectionData.NumTriangles;
+								CollisionResource.NumVertices = SectionData.MaxVertexIndex - SectionData.MinVertexIndex;
+								CollisionResource.TriangleOffset = SectionData.FirstIndex;
+								CollisionResource.VertexOffset = SectionData.MinVertexIndex;
+							}
+						}
+					}
+				}
+			}
+		}
+	}
+}
+
+void FOptimusGroomCollisionReadDataProviderProxy::GatherDispatchData(FDispatchData const& InDispatchData)
+{
+	check(InDispatchData.NumInvocations == CollisionResources.Num());
+	const TStridedView<FParameters> ParameterArray = MakeStridedParameterView<FParameters>(InDispatchData);
+
+	int32 TotalNumVertices = 0;
+	int32 TotalNumTriangles = 0;
+	for (int32 InvocationIndex = 0; InvocationIndex < ParameterArray.Num(); ++InvocationIndex)
+	{
+		FParameters& Parameters = ParameterArray[InvocationIndex];
+
+		const FCollisionResources& CollisionResource = CollisionResources[InvocationIndex];
+		Parameters.VertexOffset = CollisionResource.VertexOffset;
+		Parameters.TriangleOffset = CollisionResource.TriangleOffset;
+		Parameters.TriangleVertexIndices = CollisionResource.TriangleIndices;
+		Parameters.VertexGlobalPositions = CollisionResource.VertexPositions;
+		Parameters.NumTriangles = CollisionResource.NumTriangles;
+		Parameters.NumVertices =  CollisionResource.NumVertices;
+		Parameters.BaseTriangleIndex = TotalNumTriangles;
+		Parameters.BaseVertexIndex = TotalNumVertices;
+
+		TotalNumVertices += CollisionResource.NumVertices;
+		TotalNumTriangles += CollisionResource.NumTriangles;
+	}
+
+	for (int32 InvocationIndex = 0; InvocationIndex < ParameterArray.Num(); ++InvocationIndex)
+	{
+		FParameters& Parameters = ParameterArray[InvocationIndex];
+		
+		Parameters.TotalNumTriangles = TotalNumTriangles;
+		Parameters.TotalNumVertices = TotalNumVertices;
+	}
+}
diff --git a/Engine/Plugins/Runtime/HairStrands/Source/HairStrandsDeformer/Private/DeformerGroomCollisionRead.h b/Engine/Plugins/Runtime/HairStrands/Source/HairStrandsDeformer/Private/DeformerGroomCollisionRead.h
new file mode 100644
--- /dev/null
+++ b/Engine/Plugins/Runtime/HairStrands/Source/HairStrandsDeformer/Private/DeformerGroomCollisionRead.h
@@ -0,0 +1,99 @@
+﻿// Copyright Epic Games, Inc. All Rights Reserved.
+
+#pragma once
+
+#include "RenderGraphDefinitions.h"
+#include "OptimusComputeDataInterface.h"
+#include "HairStrandsInterface.h"
+#include "IOptimusDeformerInstanceAccessor.h"
+#include "ComputeFramework/ComputeDataProvider.h"
+#include "DeformerGroomCollisionRead.generated.h"
+
+class FOptimusGroomCollisionReadParameters;
+class UGroomComponent;
+struct FHairGroupInstance;
+class UGroomAsset;
+class UMeshComponent;
+class FSkeletalMeshObject;
+class UGroomSolverComponent;
+
+/** Compute Framework Data Interface for reading groom collision. */
+UCLASS(Category = ComputeFramework)
+class UOptimusGroomCollisionReadDataInterface : public UOptimusComputeDataInterface
+{
+	GENERATED_BODY()
+
+public:
+	
+	//~ Begin UOptimusComputeDataInterface Interface
+	virtual FString GetDisplayName() const override;
+	virtual TArray<FOptimusCDIPinDefinition> GetPinDefinitions() const override;
+	virtual TSubclassOf<UActorComponent> GetRequiredComponentClass() const override;
+	//~ End UOptimusComputeDataInterface Interface
+	
+	//~ Begin UComputeDataInterface Interface
+	virtual TCHAR const* GetClassName() const override { return TEXT("GroomCollisionRead"); }
+	virtual void GetSupportedInputs(TArray<FShaderFunctionDefinition>& OutFunctions) const override;
+	virtual void GetShaderParameters(TCHAR const* UID, FShaderParametersMetadataBuilder& InOutBuilder, FShaderParametersMetadataAllocations& InOutAllocations) const override;
+	virtual TCHAR const* GetShaderVirtualPath() const override;
+	virtual void GetPermutations(FComputeKernelPermutationVector& OutPermutationVector) const override;
+	virtual void GetShaderHash(FString& InOutKey) const override;
+	virtual void GetHLSL(FString& OutHLSL, FString const& InDataInterfaceName) const override;
+	virtual UComputeDataProvider* CreateDataProvider(TObjectPtr<UObject> InBinding, uint64 InInputMask, uint64 InOutputMask) const override;
+	//~ End UComputeDataInterface Interface
+
+private:
+	/** File holding the hlsl implementation */
+	static TCHAR const* TemplateFilePath;
+};
+
+/** Compute Framework Data Provider for reading groom meshes. */
+UCLASS(BlueprintType, editinlinenew, Category = ComputeFramework)
+class UOptimusGroomCollisionReadDataProvider : public UComputeDataProvider
+{
+	GENERATED_BODY()
+
+public:
+	
+	UPROPERTY(EditAnywhere, BlueprintReadWrite, Category = Binding)
+	TObjectPtr<UGroomSolverComponent> SolverComponent = nullptr;
+
+	//~ Begin UComputeDataProvider Interface
+	virtual FComputeDataProviderRenderProxy* GetRenderProxy() override;
+	//~ End UComputeDataProvider Interface
+};
+
+class FOptimusGroomCollisionReadDataProviderProxy : public FComputeDataProviderRenderProxy
+{
+public:
+	FOptimusGroomCollisionReadDataProviderProxy(UGroomSolverComponent* SolverComponent);
+
+	//~ Begin FComputeDataProviderRenderProxy Interface
+	virtual bool IsValid(FValidationData const& InValidationData) const override;
+	virtual void AllocateResources(FRDGBuilder& GraphBuilder, FAllocationData const& InAllocationData) override;
+	virtual void GatherDispatchData(FDispatchData const& InDispatchData) override;
+	//~ End FComputeDataProviderRenderProxy Interface
+
+private:
+	using FParameters = FOptimusGroomCollisionReadParameters;
+	
+	/** Skeletal mesh objects to extract the bones from */
+	TMap<TObjectPtr<UMeshComponent>, int32> CollisionObjects;
+
+	/** Collision resources used to dispatch CS on GPU */
+	struct FCollisionResources
+	{
+		FRHIShaderResourceView* VertexPositions = nullptr;
+		FRHIShaderResourceView* TriangleIndices = nullptr;
+		int32 VertexOffset = 0;
+		int32 TriangleOffset = 0;
+		int32 NumVertices = 0;
+		int32 NumTriangles = 0;
+	};
+	
+	/**List of collision resources per invocations */
+	TArray<FCollisionResources> CollisionResources;
+
+	/** Fallback resources */
+	FRDGBufferSRVRef FallbackStructuredSRV = nullptr;
+};
diff --git a/Engine/Plugins/Runtime/HairStrands/Source/HairStrandsDeformer/Private/DeformerGroomDomainsSource.cpp b/Engine/Plugins/Runtime/HairStrands/Source/HairStrandsDeformer/Private/DeformerGroomDomainsSource.cpp
--- a/Engine/Plugins/Runtime/HairStrands/Source/HairStrandsDeformer/Private/DeformerGroomDomainsSource.cpp
+++ b/Engine/Plugins/Runtime/HairStrands/Source/HairStrandsDeformer/Private/DeformerGroomDomainsSource.cpp
@@ -87,9 +87,11 @@ FName UOptimusGroomSolverComponentSource::FSolverExecutionDomains::Points("Solve
 
 FName UOptimusGroomSolverComponentSource::FDynamicExecutionDomains::Points("DynamicPoints");
 FName UOptimusGroomSolverComponentSource::FDynamicExecutionDomains::Curves("DynamicCurves");
+FName UOptimusGroomSolverComponentSource::FDynamicExecutionDomains::Edges("DynamicEdges");
 
 FName UOptimusGroomSolverComponentSource::FKinematicExecutionDomains::Points("KinematicPoints");
 FName UOptimusGroomSolverComponentSource::FKinematicExecutionDomains::Curves("KinematicCurves");
+FName UOptimusGroomSolverComponentSource::FKinematicExecutionDomains::Edges("KinematicEdges");
 
 FText UOptimusGroomSolverComponentSource::GetDisplayName() const
 {
@@ -109,7 +111,8 @@ bool UOptimusGroomSolverComponentSource::IsUsableAsPrimarySource() const
 TArray<FName> UOptimusGroomSolverComponentSource::GetExecutionDomains() const
 {
 	return {FSolverExecutionDomains::Points, FSolverExecutionDomains::Edges, FSolverExecutionDomains::Curves, FSolverExecutionDomains::Objects,
-		FDynamicExecutionDomains::Points, FDynamicExecutionDomains::Curves, FKinematicExecutionDomains::Points, FKinematicExecutionDomains::Curves};
+		FDynamicExecutionDomains::Points, FDynamicExecutionDomains::Curves, FKinematicExecutionDomains::Points, FKinematicExecutionDomains::Curves,
+	FDynamicExecutionDomains::Edges, FKinematicExecutionDomains::Edges};
 }
 
 int32 UOptimusGroomSolverComponentSource::GetLodIndex(const UActorComponent* InComponent) const
@@ -159,6 +162,10 @@ bool UOptimusGroomSolverComponentSource::GetComponentElementCountsForExecutionDo
 		{
 			TotalCount = GroomSolver->GetSolverSettings().PointDynamicIndices.Num();
 		}
+		else if (DomainName == FDynamicExecutionDomains::Edges)
+		{
+			TotalCount = GroomSolver->GetSolverSettings().PointDynamicIndices.Num() - GroomSolver->GetSolverSettings().CurveDynamicIndices.Num();
+		}
 		else if (DomainName == FKinematicExecutionDomains::Curves)
 		{
 			TotalCount = GroomSolver->GetSolverSettings().CurveKinematicIndices.Num();
@@ -167,6 +174,10 @@ bool UOptimusGroomSolverComponentSource::GetComponentElementCountsForExecutionDo
 		{
 			TotalCount = GroomSolver->GetSolverSettings().PointKinematicIndices.Num();
 		}
+		else if (DomainName == FKinematicExecutionDomains::Edges)
+		{
+			TotalCount = GroomSolver->GetSolverSettings().PointKinematicIndices.Num() - GroomSolver->GetSolverSettings().CurveKinematicIndices.Num();
+		}
 		TotalCount = FMath::Max(TotalCount, 1u);
 		InvocationElementCounts.Init(TotalCount,1);
 		return TotalCount != 0;
diff --git a/Engine/Plugins/Runtime/HairStrands/Source/HairStrandsDeformer/Private/DeformerGroomDomainsSource.h b/Engine/Plugins/Runtime/HairStrands/Source/HairStrandsDeformer/Private/DeformerGroomDomainsSource.h
--- a/Engine/Plugins/Runtime/HairStrands/Source/HairStrandsDeformer/Private/DeformerGroomDomainsSource.h
+++ b/Engine/Plugins/Runtime/HairStrands/Source/HairStrandsDeformer/Private/DeformerGroomDomainsSource.h
@@ -71,12 +71,14 @@ public:
 	{
 		static FName Points;
 		static FName Curves;
+		static FName Edges;
 	};
 
 	struct FKinematicExecutionDomains
 	{
 		static FName Points;
 		static FName Curves;
+		static FName Edges;
 	};
 	
 	//~ Begin  UOptimusComponentSource implementations
diff --git a/Engine/Plugins/Runtime/HairStrands/Source/HairStrandsDeformer/Private/DeformerGroomSolverRead.cpp b/Engine/Plugins/Runtime/HairStrands/Source/HairStrandsDeformer/Private/DeformerGroomSolverRead.cpp
--- a/Engine/Plugins/Runtime/HairStrands/Source/HairStrandsDeformer/Private/DeformerGroomSolverRead.cpp
+++ b/Engine/Plugins/Runtime/HairStrands/Source/HairStrandsDeformer/Private/DeformerGroomSolverRead.cpp
@@ -18,6 +18,7 @@
 #include "OptimusDataTypeRegistry.h"
 #include "OptimusDeformerInstance.h"
 #include "OptimusValueContainerStruct.h"
+#include "ScenePrivate.h"
 #include "SystemTextures.h"
 
 #include UE_INLINE_GENERATED_CPP_BY_NAME(DeformerGroomSolverRead)
@@ -65,6 +66,8 @@ TArray<FOptimusCDIPinDefinition> UOptimusGroomSolverReadDataInterface::GetPinDef
 	Defs.Add({ "NumDynamicCurves",  "ReadNumDynamicCurves", true, "NumDynamicCurves" });
 	Defs.Add({ "NumKinematicPoints",  "ReadNumKinematicPoints", true, "NumKinematicPoints" });
 	Defs.Add({ "NumKinematicCurves",  "ReadNumKinematicCurves", true, "NumKinematicCurves" });
+	Defs.Add({ "BoundingBoxMin",  "ReadBoundingBoxMin", true, "BoundingBoxMin" });
+	Defs.Add({ "BoundingBoxMax",  "ReadBoundingBoxMax", true, "BoundingBoxMax" });
 	Defs.Add({ "ObjectPointOffset", "ReadObjectPointOffset", SolverObjects,          "ReadNumSolverObjects", true, "ObjectPointOffset"  });
 	Defs.Add({ "ObjectNumPoints", "ReadObjectNumPoints", SolverObjects,          "ReadNumSolverObjects", true, "ObjectNumPoints"  });
 	Defs.Add({ "ObjectCurveOffset", "ReadObjectCurveOffset", SolverObjects,          "ReadNumSolverObjects", true, "ObjectCurveOffset"  });
@@ -77,8 +80,26 @@ TArray<FOptimusCDIPinDefinition> UOptimusGroomSolverReadDataInterface::GetPinDef
 	Defs.Add({ "KinematicCurveIndex", "ReadKinematicCurveIndex", KinematicCurves,          "ReadNumKinematicCurves", true, "KinematicCurveIndex"  });
 	Defs.Add({ "ObjectDistanceLod", "ReadObjectDistanceLod", SolverObjects,          "ReadObjectDistanceLod", true, "ObjectDistanceLod"  });
 
+#ifdef RAYTRACING_SOLVER_BINDING
+	Defs.Add({ "WorldAccelerationStructure", "ReadWorldAccelerationStructure", false, "WorldAccelerationStructure"  });
+#endif
+
 	return Defs;
 }
+
+void UOptimusGroomSolverReadDataInterface::RegisterTypes() 
+{
+#ifdef RAYTRACING_SOLVER_BINDING
+	FOptimusDataTypeRegistry::Get().RegisterType(  
+		FName("FAccelerationStructure"),
+		FText::FromString(TEXT("FAccelerationStructure")),
+		FShaderValueType::Get(FName("FAccelerationStructure"), { FShaderValueType::FStructElement(FName("ValidTLAS"), FShaderValueType::Get(EShaderFundamentalType::Int)) }),
+		FName("FAccelerationStructure"),
+		nullptr,
+		FLinearColor(0.3f, 0.7f, 0.4f, 1.0f),
+		EOptimusDataTypeUsageFlags::DataInterfaceOutput | EOptimusDataTypeUsageFlags::PinType);
+#endif
+}
  
 TSubclassOf<UActorComponent> UOptimusGroomSolverReadDataInterface::GetRequiredComponentClass() const
 {
@@ -119,6 +140,14 @@ void UOptimusGroomSolverReadDataInterface::GetSupportedInputs(TArray<FShaderFunc
 		.SetName(TEXT("ReadNumKinematicCurves"))
 		.AddReturnType(EShaderFundamentalType::Uint);
 
+	OutFunctions.AddDefaulted_GetRef()
+		.SetName(TEXT("ReadBoundingBoxMin"))
+		.AddReturnType(EShaderFundamentalType::Float,3);
+
+	OutFunctions.AddDefaulted_GetRef()
+		.SetName(TEXT("ReadBoundingBoxMax"))
+		.AddReturnType(EShaderFundamentalType::Float,3);
+
 	OutFunctions.AddDefaulted_GetRef()
 		.SetName(TEXT("ReadObjectPointOffset"))
 		.AddReturnType(EShaderFundamentalType::Uint)
@@ -173,6 +202,18 @@ void UOptimusGroomSolverReadDataInterface::GetSupportedInputs(TArray<FShaderFunc
 		.SetName(TEXT("ReadObjectDistanceLod"))
 		.AddReturnType(EShaderFundamentalType::Uint)
 		.AddParam(EShaderFundamentalType::Uint);
+
+#ifdef RAYTRACING_SOLVER_BINDING
+	if(FOptimusDataTypeHandle DataType = FOptimusDataTypeRegistry::Get().FindType(FName("FAccelerationStructure")))
+	{
+		FShaderValueTypeHandle AccelerationStructureType = DataType->ShaderValueType;
+
+		OutFunctions.AddDefaulted_GetRef()
+			.SetName(TEXT("ReadWorldAccelerationStructure"))
+			.AddReturnType(AccelerationStructureType);
+	}
+
+#endif
 }
 
 BEGIN_SHADER_PARAMETER_STRUCT(FOptimusGroomSolverReadParameters, )
@@ -183,6 +224,8 @@ BEGIN_SHADER_PARAMETER_STRUCT(FOptimusGroomSolverReadParameters, )
 	SHADER_PARAMETER(uint32, NumDynamicCurves)
 	SHADER_PARAMETER(uint32, NumKinematicPoints)
 	SHADER_PARAMETER(uint32, NumKinematicCurves)
+	SHADER_PARAMETER(FVector3f, BoundingBoxMin)
+	SHADER_PARAMETER(FVector3f, BoundingBoxMax)
 	SHADER_PARAMETER_RDG_BUFFER_SRV(Buffer<uint>, ObjectPointOffsets)
 	SHADER_PARAMETER_RDG_BUFFER_SRV(Buffer<uint>, ObjectCurveOffsets)
 	SHADER_PARAMETER_RDG_BUFFER_SRV(Buffer<uint>, ObjectNumPoints)
@@ -194,6 +237,9 @@ BEGIN_SHADER_PARAMETER_STRUCT(FOptimusGroomSolverReadParameters, )
 	SHADER_PARAMETER_RDG_BUFFER_SRV(Buffer<int>, KinematicPointIndices)
 	SHADER_PARAMETER_RDG_BUFFER_SRV(Buffer<int>, KinematicCurveIndices)
 	SHADER_PARAMETER_RDG_BUFFER_SRV(Buffer<uint>, ObjectDistanceLods)
+#ifdef RAYTRACING_SOLVER_BINDING
+	SHADER_PARAMETER_SRV(RaytracingAccelerationStructure, WorldAccelerationStructure)
+#endif
 END_SHADER_PARAMETER_STRUCT()
 
 void UOptimusGroomSolverReadDataInterface::GetShaderParameters(TCHAR const* UID, FShaderParametersMetadataBuilder& InOutBuilder, FShaderParametersMetadataAllocations& InOutAllocations) const
@@ -302,6 +348,10 @@ FOptimusGroomSolverReadDataProviderProxy::FOptimusGroomSolverReadDataProviderPro
 	
 	SolverSettings = SolverComponent->GetSolverSettings();
 	bResetSimulationTrigger = bResetSimulation;
+
+#ifdef RAYTRACING_SOLVER_BINDING
+	SceneInterface = SolverComponent->GetScene();
+#endif
 }
 
 bool FOptimusGroomSolverReadDataProviderProxy::IsValid(FValidationData const& InValidationData) const
@@ -431,6 +481,16 @@ void FOptimusGroomSolverReadDataProviderProxy::CreateInternalBuffers(FRDGBuilder
 void FOptimusGroomSolverReadDataProviderProxy::AllocateResources(FRDGBuilder& GraphBuilder, FAllocationData const& InAllocationData)
 {
 	CreateInternalBuffers(GraphBuilder);
+
+#ifdef RAYTRACING_SOLVER_BINDING
+	if(SceneInterface)
+	{
+		if (const FScene* RenderScene = SceneInterface->GetRenderScene())
+		{
+			AccelerationStructure = RenderScene->RayTracingScene.CreateLayerViewRHI(GraphBuilder.RHICmdList, ERayTracingSceneLayer::Base);
+		}
+	}
+#endif
 }
 
 void FOptimusGroomSolverReadDataProviderProxy::GatherDispatchData(FDispatchData const& InDispatchData)
@@ -438,6 +498,12 @@ void FOptimusGroomSolverReadDataProviderProxy::GatherDispatchData(FDispatchData
 	const TStridedView<FParameters> ParameterArray = MakeStridedParameterView<FParameters>(InDispatchData);
 	check(!ParameterArray.IsEmpty());
 
+	FBox BoundingBox(ForceInitToZero);
+	for(const FHairGroupInstance* GroupInstance : GroupInstances)
+	{
+		BoundingBox += GroupInstance->GetBounds().GetBox();
+	}
+
 	for (int32 InvocationIndex = 0; InvocationIndex < ParameterArray.Num(); ++InvocationIndex)
 	{
 		FParameters& Parameters = ParameterArray[InvocationIndex];
@@ -446,6 +512,9 @@ void FOptimusGroomSolverReadDataProviderProxy::GatherDispatchData(FDispatchData
 		Parameters.NumSolverCurves = NumCurves;
 		Parameters.NumSolverObjects = GroupInstances.Num();
 
+		Parameters.BoundingBoxMin = FVector3f(BoundingBox.Min);
+		Parameters.BoundingBoxMax = FVector3f(BoundingBox.Max);
+
 		Parameters.NumDynamicPoints = SolverSettings.PointDynamicIndices.Num();
 		Parameters.NumDynamicCurves = SolverSettings.CurveDynamicIndices.Num();
 		Parameters.NumKinematicPoints = SolverSettings.PointKinematicIndices.Num();
@@ -467,9 +536,10 @@ void FOptimusGroomSolverReadDataProviderProxy::GatherDispatchData(FDispatchData
 		Parameters.KinematicCurveIndices = KinematicCurveIndicesResource;
 
 		Parameters.ObjectDistanceLods = ObjectDistanceLodsResource;
+
+#ifdef RAYTRACING_SOLVER_BINDING
+		Parameters.WorldAccelerationStructure = AccelerationStructure;
+#endif
 	}
 }
 
-void FOptimusGroomSolverReadDataProviderProxy::PostGraphDispatch(FRDGBuilder& GraphBuilder) const
-{
-}
diff --git a/Engine/Plugins/Runtime/HairStrands/Source/HairStrandsDeformer/Private/DeformerGroomSolverRead.h b/Engine/Plugins/Runtime/HairStrands/Source/HairStrandsDeformer/Private/DeformerGroomSolverRead.h
--- a/Engine/Plugins/Runtime/HairStrands/Source/HairStrandsDeformer/Private/DeformerGroomSolverRead.h
+++ b/Engine/Plugins/Runtime/HairStrands/Source/HairStrandsDeformer/Private/DeformerGroomSolverRead.h
@@ -35,6 +35,7 @@ public:
 	virtual TArray<FOptimusCDIPinDefinition> GetPinDefinitions() const override;
 	virtual TSubclassOf<UActorComponent> GetRequiredComponentClass() const override;
 	virtual TArray<FOptimusCDIPropertyPinDefinition> GetPropertyPinDefinitions() const override;
+	virtual void RegisterTypes() override;
 	//~ End UOptimusComputeDataInterface Interface
 	
 	//~ Begin UComputeDataInterface Interface
@@ -97,7 +98,6 @@ public:
 	virtual void AllocateResources(FRDGBuilder& GraphBuilder, FAllocationData const& InAllocationData) override;
 	virtual void GatherPermutations(FPermutationData& InOutPermutationData) const override;
 	virtual void GatherDispatchData(FDispatchData const& InDispatchData) override;
-	virtual void PostGraphDispatch(FRDGBuilder& GraphBuilder) const override;
 	//~ End FComputeDataProviderRenderProxy Interface
 
 private:
@@ -152,4 +152,14 @@ private:
 
 	/** Reset simulation trigger */
 	bool bResetSimulationTrigger = false;
+
+#ifdef RAYTRACING_SOLVER_BINDING
+	
+	/** Scene interface used to retrieve the TLAS */
+	FSceneInterface* SceneInterface;
+
+	/** Acceleration structure used for geometric collisions */
+	FRHIShaderResourceView* AccelerationStructure = nullptr;
+
+#endif
 };
diff --git a/Engine/Plugins/Runtime/HairStrands/Source/HairStrandsDeformer/Private/HairStrandsDeformerModule.cpp b/Engine/Plugins/Runtime/HairStrands/Source/HairStrandsDeformer/Private/HairStrandsDeformerModule.cpp
--- a/Engine/Plugins/Runtime/HairStrands/Source/HairStrandsDeformer/Private/HairStrandsDeformerModule.cpp
+++ b/Engine/Plugins/Runtime/HairStrands/Source/HairStrandsDeformer/Private/HairStrandsDeformerModule.cpp
@@ -1,5 +1,29 @@
 // Copyright Epic Games, Inc. All Rights Reserved.
 
+#include "OptimusComputeDataInterface.h"
 #include "Modules/ModuleManager.h"
+#include "Modules/ModuleInterface.h"
 
-IMPLEMENT_MODULE(FDefaultModuleImpl, HairStrandsDeformer);
+/**
+ *  FHairStrandsDeformerModule
+ */
+class FHairStrandsDeformerModule : public IModuleInterface
+{
+public:
+
+	//~ Begin IModuleInterface interface
+	virtual void StartupModule() override;
+	virtual void ShutdownModule() override;
+	//~ End IModuleInterface interface
+};
+
+void FHairStrandsDeformerModule::StartupModule()
+{
+	//UOptimusComputeDataInterface::RegisterAllTypes();
+}
+
+void FHairStrandsDeformerModule::ShutdownModule()
+{
+}
+
+IMPLEMENT_MODULE(FHairStrandsDeformerModule, HairStrandsDeformer);
diff --git a/Engine/Plugins/Runtime/HairStrands/Source/HairStrandsSolver/Private/GroomSolverComponent.cpp b/Engine/Plugins/Runtime/HairStrands/Source/HairStrandsSolver/Private/GroomSolverComponent.cpp
--- a/Engine/Plugins/Runtime/HairStrands/Source/HairStrandsSolver/Private/GroomSolverComponent.cpp
+++ b/Engine/Plugins/Runtime/HairStrands/Source/HairStrandsSolver/Private/GroomSolverComponent.cpp
@@ -210,8 +210,8 @@ void UGroomSolverComponent::SelectDynamicCurves()
 						PointOffset += NumTotalPoints;
 						
 						// Keep it for now for debugging since it is WIP
-                        UE_LOG(LogTemp, Log, TEXT("Num Curves = %d Distance Ratio = %f Simulated Curves = %d | Num LODs = %d | Points Scale = %f"),
-							NumTotalCurves, DistanceRatio, NumDynamicCurves, FMath::CeilLogTwo(NumTotalCurves), PointsScale);
+                        //UE_LOG(LogTemp, Log, TEXT("Num Curves = %d Distance Ratio = %f Simulated Curves = %d | Num LODs = %d | Points Scale = %f"),
+						//	NumTotalCurves, DistanceRatio, NumDynamicCurves, FMath::CeilLogTwo(NumTotalCurves), PointsScale);
 					}
 				}
 			}
