Fix incorrect vertex element type declaration causing assert at runtime on some platforms.

#test MH test on editor, console platforms

diff --git a/Engine/Plugins/Runtime/HairStrands/Source/HairStrandsCore/Private/HairCardsVertexFactory.cpp b/Engine/Plugins/Runtime/HairStrands/Source/HairStrandsCore/Private/HairCardsVertexFactory.cpp
--- a/Engine/Plugins/Runtime/HairStrands/Source/HairStrandsCore/Private/HairCardsVertexFactory.cpp
+++ b/Engine/Plugins/Runtime/HairStrands/Source/HairStrandsCore/Private/HairCardsVertexFactory.cpp
@@ -404,11 +404,12 @@ void FHairCardsVertexFactory::InitResources(FRHICommandListBase& RHICmdList)
 			// Add additional TexCoord
 			const uint32 UVSize = sizeof(FVector2f);
 			const uint32 NumUVComponentCount = FMath::DivideAndRoundUp(NumUVs, 2u); // Two UVs are packed per Float4
+			const EVertexElementType UVComponentVertexElementType = VET_Float2;
 
 			// UV[0] is rootUV which is not part of the regular UVs, hence UV starts at index 1
 			for (uint32 UVIt = 1; UVIt < NumUVs; UVIt +=2)
 			{
-				Elements.Add(this->AccessStreamComponent(FVertexStreamComponent(UVsBuffer, 0, UVIt * UVSize, FHairCardsUVFormat::SizeInByte * NumUVComponentCount, FHairCardsUVFormat::VertexElementType, EVertexStreamUsage::Default), StartStreamIndex++));
+				Elements.Add(this->AccessStreamComponent(FVertexStreamComponent(UVsBuffer, 0, UVIt * UVSize, FHairCardsUVFormat::SizeInByte * NumUVComponentCount, UVComponentVertexElementType, EVertexStreamUsage::Default), StartStreamIndex++));
 			}
 		};
 
