Fix groom deformers not working and crashing

diff --git a/Engine/Plugins/Runtime/HairStrands/Shaders/Private/Deformers/DeformerGroomAttributeRead.ush b/Engine/Plugins/Runtime/HairStrands/Shaders/Private/Deformers/DeformerGroomAttributeRead.ush
--- a/Engine/Plugins/Runtime/HairStrands/Shaders/Private/Deformers/DeformerGroomAttributeRead.ush
+++ b/Engine/Plugins/Runtime/HairStrands/Shaders/Private/Deformers/DeformerGroomAttributeRead.ush
@@ -8,7 +8,7 @@ uint {DataInterfaceName}_BaseValueIndex;
 uint {DataInterfaceName}_TotalValueCount;
 uint {DataInterfaceName}_bHasIndexMapping;
 
-Buffer<float> {DataInterfaceName}_ValueTypedData;
+StructuredBuffer<float> {DataInterfaceName}_ValueTypedData;
 
 ByteAddressBuffer {DataInterfaceName}_IndexMappingBuffer;
 
diff --git a/Engine/Plugins/Runtime/HairStrands/Shaders/Private/Deformers/DeformerGroomMeshesRead.ush b/Engine/Plugins/Runtime/HairStrands/Shaders/Private/Deformers/DeformerGroomMeshesRead.ush
--- a/Engine/Plugins/Runtime/HairStrands/Shaders/Private/Deformers/DeformerGroomMeshesRead.ush
+++ b/Engine/Plugins/Runtime/HairStrands/Shaders/Private/Deformers/DeformerGroomMeshesRead.ush
@@ -4,7 +4,7 @@
 
 uint {DataInterfaceName}_NumBones; 
 
-Buffer<float4> {DataInterfaceName}_BoneTransformMatrices;
+StructuredBuffer<float4> {DataInterfaceName}_BoneTransformMatrices;
 
 // Interface functions
 
diff --git a/Engine/Plugins/Runtime/HairStrands/Shaders/Private/Deformers/DeformerGroomSolverRead.ush b/Engine/Plugins/Runtime/HairStrands/Shaders/Private/Deformers/DeformerGroomSolverRead.ush
--- a/Engine/Plugins/Runtime/HairStrands/Shaders/Private/Deformers/DeformerGroomSolverRead.ush
+++ b/Engine/Plugins/Runtime/HairStrands/Shaders/Private/Deformers/DeformerGroomSolverRead.ush
@@ -46,22 +46,22 @@ uint {DataInterfaceName}_NumKinematicCurves;
 float3 {DataInterfaceName}_BoundingBoxMin;
 float3 {DataInterfaceName}_BoundingBoxMax;
 
-Buffer<uint> {DataInterfaceName}_ObjectCurveOffsets;
-Buffer<uint> {DataInterfaceName}_ObjectPointOffsets;
+StructuredBuffer<uint> {DataInterfaceName}_ObjectCurveOffsets;
+StructuredBuffer<uint> {DataInterfaceName}_ObjectPointOffsets;
 
-Buffer<uint> {DataInterfaceName}_ObjectNumCurves;
-Buffer<uint> {DataInterfaceName}_ObjectNumPoints;
+StructuredBuffer<uint> {DataInterfaceName}_ObjectNumCurves;
+StructuredBuffer<uint> {DataInterfaceName}_ObjectNumPoints;
 
-Buffer<int> {DataInterfaceName}_PointObjectIndices;
-Buffer<int> {DataInterfaceName}_CurveObjectIndices;
+StructuredBuffer<int> {DataInterfaceName}_PointObjectIndices;
+StructuredBuffer<int> {DataInterfaceName}_CurveObjectIndices;
 
-Buffer<int> {DataInterfaceName}_DynamicPointIndices;
-Buffer<int> {DataInterfaceName}_DynamicCurveIndices;
+StructuredBuffer<int> {DataInterfaceName}_DynamicPointIndices;
+StructuredBuffer<int> {DataInterfaceName}_DynamicCurveIndices;
 
-Buffer<int> {DataInterfaceName}_KinematicPointIndices;
-Buffer<int> {DataInterfaceName}_KinematicCurveIndices;
+StructuredBuffer<int> {DataInterfaceName}_KinematicPointIndices;
+StructuredBuffer<int> {DataInterfaceName}_KinematicCurveIndices;
 
-Buffer<uint> {DataInterfaceName}_ObjectDistanceLods;
+StructuredBuffer<uint> {DataInterfaceName}_ObjectDistanceLods;
 
 // Interface functions
 
diff --git a/Engine/Plugins/Runtime/HairStrands/Source/HairStrandsDeformer/Private/DeformerGroomAttributeRead.cpp b/Engine/Plugins/Runtime/HairStrands/Source/HairStrandsDeformer/Private/DeformerGroomAttributeRead.cpp
--- a/Engine/Plugins/Runtime/HairStrands/Source/HairStrandsDeformer/Private/DeformerGroomAttributeRead.cpp
+++ b/Engine/Plugins/Runtime/HairStrands/Source/HairStrandsDeformer/Private/DeformerGroomAttributeRead.cpp
@@ -222,7 +222,7 @@ BEGIN_SHADER_PARAMETER_STRUCT(FOptimusGroomAttributeReadParameters, )
 	SHADER_PARAMETER(uint32, BaseValueIndex)
 	SHADER_PARAMETER(uint32, TotalValueCount)
 	SHADER_PARAMETER(uint32, bHasIndexMapping)
-	SHADER_PARAMETER_RDG_BUFFER_SRV(Buffer<float>, ValueTypedData)
+	SHADER_PARAMETER_RDG_BUFFER_SRV(StructuredBuffer<float>, ValueTypedData)
 	SHADER_PARAMETER_RDG_BUFFER_SRV(ByteAddressBuffer, IndexMappingBuffer)
 END_SHADER_PARAMETER_STRUCT()
 
diff --git a/Engine/Plugins/Runtime/HairStrands/Source/HairStrandsDeformer/Private/DeformerGroomMeshesRead.cpp b/Engine/Plugins/Runtime/HairStrands/Source/HairStrandsDeformer/Private/DeformerGroomMeshesRead.cpp
--- a/Engine/Plugins/Runtime/HairStrands/Source/HairStrandsDeformer/Private/DeformerGroomMeshesRead.cpp
+++ b/Engine/Plugins/Runtime/HairStrands/Source/HairStrandsDeformer/Private/DeformerGroomMeshesRead.cpp
@@ -61,7 +61,7 @@ void UOptimusGroomMeshesReadDataInterface::GetSupportedInputs(TArray<FShaderFunc
 
 BEGIN_SHADER_PARAMETER_STRUCT(FOptimusGroomMeshesReadParameters, )
 	SHADER_PARAMETER(uint32, NumBones)
-	SHADER_PARAMETER_RDG_BUFFER_SRV(Buffer<float4>, BoneTransformMatrices)
+	SHADER_PARAMETER_RDG_BUFFER_SRV(StructuredBuffer<float4>, BoneTransformMatrices)
 END_SHADER_PARAMETER_STRUCT()
 
 void UOptimusGroomMeshesReadDataInterface::GetShaderParameters(TCHAR const* UID, FShaderParametersMetadataBuilder& InOutBuilder, FShaderParametersMetadataAllocations& InOutAllocations) const
diff --git a/Engine/Plugins/Runtime/HairStrands/Source/HairStrandsDeformer/Private/DeformerGroomSolverRead.cpp b/Engine/Plugins/Runtime/HairStrands/Source/HairStrandsDeformer/Private/DeformerGroomSolverRead.cpp
--- a/Engine/Plugins/Runtime/HairStrands/Source/HairStrandsDeformer/Private/DeformerGroomSolverRead.cpp
+++ b/Engine/Plugins/Runtime/HairStrands/Source/HairStrandsDeformer/Private/DeformerGroomSolverRead.cpp
@@ -226,17 +226,17 @@ BEGIN_SHADER_PARAMETER_STRUCT(FOptimusGroomSolverReadParameters, )
 	SHADER_PARAMETER(uint32, NumKinematicCurves)
 	SHADER_PARAMETER(FVector3f, BoundingBoxMin)
 	SHADER_PARAMETER(FVector3f, BoundingBoxMax)
-	SHADER_PARAMETER_RDG_BUFFER_SRV(Buffer<uint>, ObjectPointOffsets)
-	SHADER_PARAMETER_RDG_BUFFER_SRV(Buffer<uint>, ObjectCurveOffsets)
-	SHADER_PARAMETER_RDG_BUFFER_SRV(Buffer<uint>, ObjectNumPoints)
-	SHADER_PARAMETER_RDG_BUFFER_SRV(Buffer<uint>, ObjectNumCurves)
-	SHADER_PARAMETER_RDG_BUFFER_SRV(Buffer<int>, PointObjectIndices)
-	SHADER_PARAMETER_RDG_BUFFER_SRV(Buffer<int>, CurveObjectIndices)
-	SHADER_PARAMETER_RDG_BUFFER_SRV(Buffer<int>, DynamicPointIndices)
-	SHADER_PARAMETER_RDG_BUFFER_SRV(Buffer<int>, DynamicCurveIndices)
-	SHADER_PARAMETER_RDG_BUFFER_SRV(Buffer<int>, KinematicPointIndices)
-	SHADER_PARAMETER_RDG_BUFFER_SRV(Buffer<int>, KinematicCurveIndices)
-	SHADER_PARAMETER_RDG_BUFFER_SRV(Buffer<uint>, ObjectDistanceLods)
+	SHADER_PARAMETER_RDG_BUFFER_SRV(StructuredBuffer<uint>, ObjectPointOffsets)
+	SHADER_PARAMETER_RDG_BUFFER_SRV(StructuredBuffer<uint>, ObjectCurveOffsets)
+	SHADER_PARAMETER_RDG_BUFFER_SRV(StructuredBuffer<uint>, ObjectNumPoints)
+	SHADER_PARAMETER_RDG_BUFFER_SRV(StructuredBuffer<uint>, ObjectNumCurves)
+	SHADER_PARAMETER_RDG_BUFFER_SRV(StructuredBuffer<int>, PointObjectIndices)
+	SHADER_PARAMETER_RDG_BUFFER_SRV(StructuredBuffer<int>, CurveObjectIndices)
+	SHADER_PARAMETER_RDG_BUFFER_SRV(StructuredBuffer<int>, DynamicPointIndices)
+	SHADER_PARAMETER_RDG_BUFFER_SRV(StructuredBuffer<int>, DynamicCurveIndices)
+	SHADER_PARAMETER_RDG_BUFFER_SRV(StructuredBuffer<int>, KinematicPointIndices)
+	SHADER_PARAMETER_RDG_BUFFER_SRV(StructuredBuffer<int>, KinematicCurveIndices)
+	SHADER_PARAMETER_RDG_BUFFER_SRV(StructuredBuffer<uint>, ObjectDistanceLods)
 #ifdef RAYTRACING_SOLVER_BINDING
 	SHADER_PARAMETER_SRV(RaytracingAccelerationStructure, WorldAccelerationStructure)
 #endif
