Fix haircards causing crash during vulkan pso precompile

diff --git a/Engine/Plugins/Runtime/HairStrands/Source/HairStrandsCore/Private/HairCardsVertexFactory.cpp b/Engine/Plugins/Runtime/HairStrands/Source/HairStrandsCore/Private/HairCardsVertexFactory.cpp
--- a/Engine/Plugins/Runtime/HairStrands/Source/HairStrandsCore/Private/HairCardsVertexFactory.cpp
+++ b/Engine/Plugins/Runtime/HairStrands/Source/HairStrandsCore/Private/HairCardsVertexFactory.cpp
@@ -340,7 +340,8 @@ void FHairCardsVertexFactory::GetPSOPrecacheVertexFetchElements(EVertexInputStre
 
 		Elements.Add(FVertexElement(2+NumStreams, 0, FHairCardsMaterialFormat::VertexElementType, 4, FHairCardsMaterialFormat::SizeInByte, false)); // Material
 		Elements.Add(FVertexElement(2+NumStreams, 0, FHairCardsMaterialFormat::VertexElementType, 5, FHairCardsMaterialFormat::SizeInByte, false)); // VertexColor
-		Elements.Add(FVertexElement(2+NumStreams, 0, FHairCardsUVFormat::VertexElementType, 6, FHairCardsUVFormat::SizeInByte, false));
+
+		Elements.Add(FVertexElement(3+NumStreams, 0, FHairCardsUVFormat::VertexElementType, 6, FHairCardsUVFormat::SizeInByte, false));
 	}
 }
 
