Fix dataflow crash when reseting simulation while groom deformers are running

diff --git a/Engine/Plugins/Runtime/HairStrands/Source/HairStrandsDeformer/Private/DeformerGroomAttributeRead.h b/Engine/Plugins/Runtime/HairStrands/Source/HairStrandsDeformer/Private/DeformerGroomAttributeRead.h
--- a/Engine/Plugins/Runtime/HairStrands/Source/HairStrandsDeformer/Private/DeformerGroomAttributeRead.h
+++ b/Engine/Plugins/Runtime/HairStrands/Source/HairStrandsDeformer/Private/DeformerGroomAttributeRead.h
@@ -141,7 +141,7 @@ private:
 	void CreateInternalBuffers(FRDGBuilder& GraphBuilder);
 
 	/** List of instances (invocations) used in that data interface */
-	TArray<const FHairGroupInstance*> GroupInstances;
+	TArray<TRefCountPtr<const FHairGroupInstance>> GroupInstances;
 
 	/** Groom assets group indices the proxy instance is linked to */
 	TArray<TPair<UGroomAsset*, UE::Groom::Private::FGroupElements>> GroupElements;
diff --git a/Engine/Plugins/Runtime/HairStrands/Source/HairStrandsDeformer/Private/DeformerGroomDomainsExec.h b/Engine/Plugins/Runtime/HairStrands/Source/HairStrandsDeformer/Private/DeformerGroomDomainsExec.h
--- a/Engine/Plugins/Runtime/HairStrands/Source/HairStrandsDeformer/Private/DeformerGroomDomainsExec.h
+++ b/Engine/Plugins/Runtime/HairStrands/Source/HairStrandsDeformer/Private/DeformerGroomDomainsExec.h
@@ -110,7 +110,7 @@ private:
 	using FParameters = FGroomExecDataInterfaceParameters;
 
 	/** List of instances (invocations) used in that data interface */
-	TArray<const FHairGroupInstance*> GroupInstances;
+	TArray<TRefCountPtr<const FHairGroupInstance>> GroupInstances;
 
 	/** Number of elements for each invocation */
 	TArray<int32> GroupCounts;
diff --git a/Engine/Plugins/Runtime/HairStrands/Source/HairStrandsDeformer/Private/DeformerGroomGuidesRead.h b/Engine/Plugins/Runtime/HairStrands/Source/HairStrandsDeformer/Private/DeformerGroomGuidesRead.h
--- a/Engine/Plugins/Runtime/HairStrands/Source/HairStrandsDeformer/Private/DeformerGroomGuidesRead.h
+++ b/Engine/Plugins/Runtime/HairStrands/Source/HairStrandsDeformer/Private/DeformerGroomGuidesRead.h
@@ -77,7 +77,7 @@ private:
 	using FParameters = FOptimusGroomGuidesReadParameters;
 
 	/** List of instances (invocations) used in that data interface */
-	TArray<const FHairGroupInstance*> GroupInstances;
+	TArray<TRefCountPtr<const FHairGroupInstance>> GroupInstances;
 
 	/** Binding Resources used to dispatch CS on GPU */
 	struct FBindingResources
diff --git a/Engine/Plugins/Runtime/HairStrands/Source/HairStrandsDeformer/Private/DeformerGroomGuidesWrite.h b/Engine/Plugins/Runtime/HairStrands/Source/HairStrandsDeformer/Private/DeformerGroomGuidesWrite.h
--- a/Engine/Plugins/Runtime/HairStrands/Source/HairStrandsDeformer/Private/DeformerGroomGuidesWrite.h
+++ b/Engine/Plugins/Runtime/HairStrands/Source/HairStrandsDeformer/Private/DeformerGroomGuidesWrite.h
@@ -77,7 +77,7 @@ private:
 	using FParameters = FOptimusGroomGuidesWriteParameters;
 
 	/** List of instances (invocations) used in that data interface */
-	TArray<const FHairGroupInstance*> GroupInstances;
+	TArray<TRefCountPtr<const FHairGroupInstance>> GroupInstances;
 
 	/** Output mask for gpu buffer format*/
 	uint64 OutputMask = 0;
diff --git a/Engine/Plugins/Runtime/HairStrands/Source/HairStrandsDeformer/Private/DeformerGroomInterfaceUtils.h b/Engine/Plugins/Runtime/HairStrands/Source/HairStrandsDeformer/Private/DeformerGroomInterfaceUtils.h
--- a/Engine/Plugins/Runtime/HairStrands/Source/HairStrandsDeformer/Private/DeformerGroomInterfaceUtils.h
+++ b/Engine/Plugins/Runtime/HairStrands/Source/HairStrandsDeformer/Private/DeformerGroomInterfaceUtils.h
@@ -46,7 +46,7 @@ namespace UE::Groom::Private
 		return nullptr;
 	}
 	
-	FORCEINLINE void GroomComponentsToInstances(const TArray<const UGroomComponent*>& GroomComponents, TArray<const FHairGroupInstance*>& GroupInstances)
+	FORCEINLINE void GroomComponentsToInstances(const TArray<const UGroomComponent*>& GroomComponents, TArray<TRefCountPtr<const FHairGroupInstance>>& GroupInstances)
  	{
  		GroupInstances.Reset();
  		for (const UGroomComponent* GroomComponent : GroomComponents)
@@ -58,14 +58,15 @@ namespace UE::Groom::Private
  				{
  					if(HasDeformationEnabledOrHasMeshDeformer(GroomComponent, GroupIndex) && GroomComponent->GetGroupInstance(GroupIndex))
  					{
- 						GroupInstances.Add(GroomComponent->GetGroupInstance(GroupIndex));
+						const TRefCountPtr<const FHairGroupInstance> GroupInstance(GroomComponent->GetGroupInstance(GroupIndex));
+ 						GroupInstances.Add(GroupInstance);
  					}
  				}
  			}
  		}
  	}
 	FORCEINLINE void GroomComponentsToSkelmeshes(const TArray<const UGroomComponent*>& GroomComponents, TArray<const FSkeletalMeshObject*>& SkeletalMeshes, TArray<FMatrix44f>& SkeletalTransforms,
-		TArray<TArray<FMatrix44f>>& BonesRefToLocals, TArray<TArray<FMatrix44f>>& BindTransforms, TArray<const FHairGroupInstance*>& GroupInstances)
+		TArray<TArray<FMatrix44f>>& BonesRefToLocals, TArray<TArray<FMatrix44f>>& BindTransforms, TArray<TRefCountPtr<const FHairGroupInstance>>& GroupInstances)
 	{
 		SkeletalMeshes.Reset();
 		SkeletalTransforms.Reset();
@@ -123,7 +124,7 @@ namespace UE::Groom::Private
 									BonesRefToLocals.Add(RefToLocals);
 									BindTransforms.Add(RefBases);
 									SkeletalMeshes.Add(SkelMesh->MeshObject);
-									GroupInstances.Add(GroomComponent->GetGroupInstance(GroupIndex));
+									GroupInstances.Add(TRefCountPtr<const FHairGroupInstance>(GroomComponent->GetGroupInstance(GroupIndex)));
 								}
 							}
 						}
@@ -149,7 +150,7 @@ namespace UE::Groom::Private
 		}
 	}
 	
-	FORCEINLINE void GatherGroupInstances(const UActorComponent* ActorComponent, TArray<const FHairGroupInstance*>& GroupInstances)
+	FORCEINLINE void GatherGroupInstances(const UActorComponent* ActorComponent, TArray<TRefCountPtr<const FHairGroupInstance>>& GroupInstances)
 	{
 		TArray<const UGroomComponent*> GroomComponents;
 		GatherGroomComponents(ActorComponent, GroomComponents);
@@ -158,7 +159,7 @@ namespace UE::Groom::Private
 	}
 
 	FORCEINLINE void GatherGroupSkelmeshes(const UActorComponent* ActorComponent, TArray<const FSkeletalMeshObject*>& SkeletalMeshes, TArray<FMatrix44f>& SkeletalTransforms,
-		TArray<TArray<FMatrix44f>>& BonesRefToLocals, TArray<TArray<FMatrix44f>>& BindTransforms, TArray<const FHairGroupInstance*>& GroupInstances)
+		TArray<TArray<FMatrix44f>>& BonesRefToLocals, TArray<TArray<FMatrix44f>>& BindTransforms, TArray<TRefCountPtr<const FHairGroupInstance>>& GroupInstances)
 	{
 		TArray<const UGroomComponent*> GroomComponents;
 		GatherGroomComponents(ActorComponent, GroomComponents);
@@ -180,9 +181,9 @@ namespace UE::Groom::Private
 		return false;
 	}
 
-	static bool HaveGuidesInstanceResources(const TArray<const FHairGroupInstance*>& GroupInstances)
+	static bool HaveGuidesInstanceResources(const TArray<TRefCountPtr<const FHairGroupInstance>>& GroupInstances)
 	{
-		for(const FHairGroupInstance* GroupInstance : GroupInstances)
+		for(TRefCountPtr<const FHairGroupInstance> GroupInstance : GroupInstances)
 		{
 			if(!GroupInstance || (GroupInstance && !HaveValidInstanceResources(GroupInstance->Guides)))
 			{
@@ -192,9 +193,9 @@ namespace UE::Groom::Private
 		return true;
 	}
 
-	FORCEINLINE bool HaveStrandsInstanceResources(const TArray<const FHairGroupInstance*>& GroupInstances)
+	FORCEINLINE bool HaveStrandsInstanceResources(const TArray<TRefCountPtr<const FHairGroupInstance>>& GroupInstances)
 	{
-		for(const FHairGroupInstance* GroupInstance : GroupInstances)
+		for(TRefCountPtr<const FHairGroupInstance> GroupInstance : GroupInstances)
 		{
 			if(!GroupInstance || (GroupInstance && !HaveValidInstanceResources(GroupInstance->Strands)))
 			{
@@ -221,9 +222,9 @@ namespace UE::Groom::Private
 		return false;
 	}
 
-	static bool HaveGuidesSkinnedResources(const TArray<const FHairGroupInstance*>& GroupInstances)
+	static bool HaveGuidesSkinnedResources(const TArray<TRefCountPtr<const FHairGroupInstance>>& GroupInstances)
 	{
-		for(const FHairGroupInstance* GroupInstance : GroupInstances)
+		for(TRefCountPtr<const FHairGroupInstance> GroupInstance : GroupInstances)
 		{
 			if(GroupInstance && GroupInstance->HairGroupPublicData)
 			{
@@ -243,9 +244,9 @@ namespace UE::Groom::Private
 		return true;
 	}
 
-	FORCEINLINE bool HaveStrandsSkinnedResources(const TArray<const FHairGroupInstance*>& GroupInstances)
+	FORCEINLINE bool HaveStrandsSkinnedResources(const TArray<TRefCountPtr<const FHairGroupInstance>>& GroupInstances)
 	{
-		for(const FHairGroupInstance* GroupInstance : GroupInstances)
+		for(TRefCountPtr<const FHairGroupInstance> GroupInstance : GroupInstances)
 		{
 			if(GroupInstance && GroupInstance->HairGroupPublicData)
 			{
@@ -376,7 +377,7 @@ namespace UE::Groom::Private
 						const uint32 NumGroups = GroomComponent->GetGroupCount();
 						for (uint32 GroupIndex = 0; GroupIndex < NumGroups; ++GroupIndex)
 						{
-							if (const FHairGroupInstance* GroupInstance = GroomComponent->GetGroupInstance(GroupIndex))
+							if (TRefCountPtr<const FHairGroupInstance> GroupInstance = TRefCountPtr<const FHairGroupInstance>(GroomComponent->GetGroupInstance(GroupIndex)))
 							{
 								if(FBoxSphereBounds::BoxesIntersect(GroupInstance->GetBounds(), CollisionBounds))
 								{
@@ -507,7 +508,7 @@ namespace UE::Groom::Private
 	{
 		TArray<int32> GroupIndices;
 		TArray<int32> GroupOffsets;
-		TArray<const FHairGroupInstance*> GroupInstances;
+		TArray<TRefCountPtr<const FHairGroupInstance>> GroupInstances;
 	};
 	
 	FORCEINLINE void GetGroomInvocationElementGroups(const TArray<const UGroomComponent*>& GroomComponents, const FName DomainName, TArray<TPair<UGroomAsset*, FGroupElements>>& InvocationGroups,  const bool bSourceElements = false)
diff --git a/Engine/Plugins/Runtime/HairStrands/Source/HairStrandsDeformer/Private/DeformerGroomMeshesRead.h b/Engine/Plugins/Runtime/HairStrands/Source/HairStrandsDeformer/Private/DeformerGroomMeshesRead.h
--- a/Engine/Plugins/Runtime/HairStrands/Source/HairStrandsDeformer/Private/DeformerGroomMeshesRead.h
+++ b/Engine/Plugins/Runtime/HairStrands/Source/HairStrandsDeformer/Private/DeformerGroomMeshesRead.h
@@ -89,7 +89,7 @@ private:
 	TArray<TArray<FMatrix44f>> BindTransforms;
 
 	/** List of instances (invocations) used in that data interface */
-	TArray<const FHairGroupInstance*> GroupInstances;
+	TArray<TRefCountPtr<const FHairGroupInstance>> GroupInstances;
 	
 	/** Bones Resources used to dispatch CS on GPU */
 	TArray<FRDGBufferSRVRef> BoneMatricesResources;
diff --git a/Engine/Plugins/Runtime/HairStrands/Source/HairStrandsDeformer/Private/DeformerGroomSolverRead.cpp b/Engine/Plugins/Runtime/HairStrands/Source/HairStrandsDeformer/Private/DeformerGroomSolverRead.cpp
--- a/Engine/Plugins/Runtime/HairStrands/Source/HairStrandsDeformer/Private/DeformerGroomSolverRead.cpp
+++ b/Engine/Plugins/Runtime/HairStrands/Source/HairStrandsDeformer/Private/DeformerGroomSolverRead.cpp
@@ -314,7 +314,7 @@ FComputeDataProviderRenderProxy* UOptimusGroomSolverReadDataProvider::GetRenderP
 			}
 			if(bResetSimulationCount > 0)
 			{
-				TArray<const FHairGroupInstance*> GroupInstances;
+				TArray<TRefCountPtr<const FHairGroupInstance>> GroupInstances;
 				UE::Groom::Private::GroomComponentsToInstances(GroomComponents, GroupInstances);
 				
 				if(!UE::Groom::Private::HaveGuidesInstanceResources(GroupInstances) || !UE::Groom::Private::HaveGuidesSkinnedResources(GroupInstances))
diff --git a/Engine/Plugins/Runtime/HairStrands/Source/HairStrandsDeformer/Private/DeformerGroomSolverRead.h b/Engine/Plugins/Runtime/HairStrands/Source/HairStrandsDeformer/Private/DeformerGroomSolverRead.h
--- a/Engine/Plugins/Runtime/HairStrands/Source/HairStrandsDeformer/Private/DeformerGroomSolverRead.h
+++ b/Engine/Plugins/Runtime/HairStrands/Source/HairStrandsDeformer/Private/DeformerGroomSolverRead.h
@@ -148,7 +148,7 @@ private:
 	FRDGBufferSRVRef ObjectDistanceLodsResource = nullptr;
 
 	/** List of instances (invocations) used in that data interface */
-	TArray<const FHairGroupInstance*> GroupInstances;
+	TArray<TRefCountPtr<const FHairGroupInstance>> GroupInstances;
 
 	/** Reset simulation trigger */
 	bool bResetSimulationTrigger = false;
diff --git a/Engine/Plugins/Runtime/HairStrands/Source/HairStrandsDeformer/Private/DeformerGroomStrandsRead.h b/Engine/Plugins/Runtime/HairStrands/Source/HairStrandsDeformer/Private/DeformerGroomStrandsRead.h
--- a/Engine/Plugins/Runtime/HairStrands/Source/HairStrandsDeformer/Private/DeformerGroomStrandsRead.h
+++ b/Engine/Plugins/Runtime/HairStrands/Source/HairStrandsDeformer/Private/DeformerGroomStrandsRead.h
@@ -73,7 +73,7 @@ private:
 	using FParameters = FOptimusGroomStrandsReadParameters;
 
 	/** List of instances (invocations) used in that data interface */
-	TArray<const FHairGroupInstance*> GroupInstances;
+	TArray<TRefCountPtr<const FHairGroupInstance>> GroupInstances;
 
 	/** Binding Resources */
 	struct FBindingResources
diff --git a/Engine/Plugins/Runtime/HairStrands/Source/HairStrandsDeformer/Private/DeformerGroomStrandsWrite.h b/Engine/Plugins/Runtime/HairStrands/Source/HairStrandsDeformer/Private/DeformerGroomStrandsWrite.h
--- a/Engine/Plugins/Runtime/HairStrands/Source/HairStrandsDeformer/Private/DeformerGroomStrandsWrite.h
+++ b/Engine/Plugins/Runtime/HairStrands/Source/HairStrandsDeformer/Private/DeformerGroomStrandsWrite.h
@@ -78,7 +78,7 @@ private:
 	using FParameters = FOptimusGroomStrandsWriteParameters;
 
 	/** List of instances (invocations) used in that data interface */
-	TArray<const FHairGroupInstance*> GroupInstances;
+	TArray<TRefCountPtr<const FHairGroupInstance>> GroupInstances;
 
 	/** Output mask for gpu buffer format*/
 	uint64 OutputMask = 0;
diff --git a/Engine/Plugins/Runtime/HairStrands/Source/HairStrandsSolver/Private/AddSolverDeformerNode.cpp b/Engine/Plugins/Runtime/HairStrands/Source/HairStrandsSolver/Private/AddSolverDeformerNode.cpp
--- a/Engine/Plugins/Runtime/HairStrands/Source/HairStrandsSolver/Private/AddSolverDeformerNode.cpp
+++ b/Engine/Plugins/Runtime/HairStrands/Source/HairStrandsSolver/Private/AddSolverDeformerNode.cpp
@@ -114,17 +114,23 @@ namespace UE::Groom::Private
 
 	FORCEINLINE void CreateDeformerInstance(UOptimusDeformer* DeformerGraph, UMeshDeformerInstance* DeformerInstance, const FGuid InstanceGuid, UDataflow* DataflowObject)
 	{
-		if(DeformerGraph && DeformerInstance)
+		if(DeformerGraph && DeformerInstance && DeformerInstance->GetOuter())
 		{
 			FFunctionGraphTask::CreateAndDispatchWhenReady([DeformerGraph, DeformerInstance, InstanceGuid, DataflowObject]()
 			{
-				if (!UE::IsSavingPackage(nullptr) && !IsGarbageCollectingAndLockingUObjectHashTables())
+				if (UMeshComponent* MeshComponent = Cast<UMeshComponent>(DeformerInstance->GetOuter()))
 				{
-					if (UOptimusDeformerDynamicInstanceManager* DeformerInstanceManager = Cast<UOptimusDeformerDynamicInstanceManager>(DeformerInstance))
+					if (!MeshComponent->IsBeingDestroyed() && MeshComponent->GetScene())
 					{
-						if (!DeformerInstanceManager->GetDeformerInstance(InstanceGuid))
+						if (!UE::IsSavingPackage(nullptr) && !IsGarbageCollectingAndLockingUObjectHashTables())
 						{
-							DeformerInstanceManager->AddProducerDeformer(DataflowObject, InstanceGuid, DeformerGraph);
+							if (UOptimusDeformerDynamicInstanceManager* DeformerInstanceManager = Cast<UOptimusDeformerDynamicInstanceManager>(DeformerInstance))
+							{
+								if (!DeformerInstanceManager->GetDeformerInstance(InstanceGuid))
+								{
+									DeformerInstanceManager->AddProducerDeformer(DataflowObject, InstanceGuid, DeformerGraph);
+								}
+							}
 						}
 					}
 				}
diff --git a/Engine/Plugins/Runtime/HairStrands/Source/HairStrandsSolver/Private/GroomSolverComponent.cpp b/Engine/Plugins/Runtime/HairStrands/Source/HairStrandsSolver/Private/GroomSolverComponent.cpp
--- a/Engine/Plugins/Runtime/HairStrands/Source/HairStrandsSolver/Private/GroomSolverComponent.cpp
+++ b/Engine/Plugins/Runtime/HairStrands/Source/HairStrandsSolver/Private/GroomSolverComponent.cpp
@@ -267,14 +267,28 @@ void UGroomSolverComponent::DestroyRenderState_Concurrent()
 
 void UGroomSolverComponent::AddGroomComponent(UGroomComponent* GroomPhysicsComponent)
 {
-	GroomComponents.Add(GroomPhysicsComponent);
-	GroomPhysicsComponent->SetGroomSolver(this);
+	if(GroomPhysicsComponent)
+	{ 
+		GroomComponents.Add(GroomPhysicsComponent);
+		GroomPhysicsComponent->SetGroomSolver(this);
+
+#if WITH_EDITOR
+		GroomPhysicsComponent->InvalidateAndRecreate();
+#endif
+	}
 }
 
 void UGroomSolverComponent::RemoveGroomComponent(UGroomComponent* GroomPhysicsComponent)
 {
-	GroomComponents.Remove(GroomPhysicsComponent);
-	GroomPhysicsComponent->SetGroomSolver(nullptr);
+	if (GroomPhysicsComponent)
+	{
+		GroomComponents.Remove(GroomPhysicsComponent);
+		GroomPhysicsComponent->SetGroomSolver(nullptr);
+
+#if WITH_EDITOR
+		GroomPhysicsComponent->InvalidateAndRecreate();
+#endif
+	}
 }
 
 void UGroomSolverComponent::ResetGroomComponents()
