Fix dataflow crash when the binding asset is changed while dataflow editor is opened

diff --git a/Engine/Plugins/Runtime/HairStrands/Source/HairStrandsCore/Private/HairStrandsMeshProjection.cpp b/Engine/Plugins/Runtime/HairStrands/Source/HairStrandsCore/Private/HairStrandsMeshProjection.cpp
--- a/Engine/Plugins/Runtime/HairStrands/Source/HairStrandsCore/Private/HairStrandsMeshProjection.cpp
+++ b/Engine/Plugins/Runtime/HairStrands/Source/HairStrandsCore/Private/HairStrandsMeshProjection.cpp
@@ -279,28 +279,29 @@ static bool AddHairStrandUpdateMeshTrianglesPass(
 			{ 
 				return In.SectionIndex == MeshSectionIndex; 
 			});
-		check(RemappedSectionIndex != INDEX_NONE);
-
-		const FCachedGeometry::Section& MeshSectionData = MeshLODData.Sections[RemappedSectionIndex];
-		const uint32 SectionIndex = MeshSectionData.SectionIndex;
-		SectionDatas[SectionIndex].TotalIndexCount	= MeshSectionData.TotalIndexCount;
-		SectionDatas[SectionIndex].TotalVertexCount	= MeshSectionData.TotalVertexCount;
-		SectionDatas[SectionIndex].IndexBaseIndex	= MeshSectionData.IndexBaseIndex;
-		SectionDatas[SectionIndex].UVsChannelOffset	= MeshSectionData.UVsChannelOffset;
-		SectionDatas[SectionIndex].UVsChannelCount	= MeshSectionData.UVsChannelCount;
-		SectionDatas[SectionIndex].Pad				= 0u;
-		SectionDatas[SectionIndex].bIsSwapped		= bUseRDGPositionBuffer ? 
-			(MeshSectionData.RDGPositionBuffer != CommonParameters.RDGMeshPositionBuffer ? 1u : 0u) :
-			(MeshSectionData.PositionBuffer    != CommonParameters.MeshPositionBuffer    ? 1u : 0u);
-		SectionDatas[SectionIndex].bUseFaceNormal	= 0u;
-
-		// Sanity check
-		check(MeshSectionData.UVsChannelOffset < 255);
-		check(MeshSectionData.UVsChannelCount < 255);
-		check(CommonParameters.RDGMeshPositionBuffer == MeshSectionData.RDGPositionBuffer || CommonParameters.RDGMeshPositionBuffer == MeshSectionData.RDGPreviousPositionBuffer);
-		check(CommonParameters.MeshPositionBuffer    == MeshSectionData.PositionBuffer    || CommonParameters.MeshPositionBuffer    == MeshSectionData.PreviousPositionBuffer);
-		check(CommonParameters.MeshIndexBuffer		 == MeshSectionData.IndexBuffer);
-		check(CommonParameters.MeshTangentBuffer	 == MeshSectionData.TangentBuffer);
+		if(RemappedSectionIndex != INDEX_NONE)
+		{ 
+			const FCachedGeometry::Section& MeshSectionData = MeshLODData.Sections[RemappedSectionIndex];
+			const uint32 SectionIndex = MeshSectionData.SectionIndex;
+			SectionDatas[SectionIndex].TotalIndexCount	= MeshSectionData.TotalIndexCount;
+			SectionDatas[SectionIndex].TotalVertexCount	= MeshSectionData.TotalVertexCount;
+			SectionDatas[SectionIndex].IndexBaseIndex	= MeshSectionData.IndexBaseIndex;
+			SectionDatas[SectionIndex].UVsChannelOffset	= MeshSectionData.UVsChannelOffset;
+			SectionDatas[SectionIndex].UVsChannelCount	= MeshSectionData.UVsChannelCount;
+			SectionDatas[SectionIndex].Pad				= 0u;
+			SectionDatas[SectionIndex].bIsSwapped		= bUseRDGPositionBuffer ? 
+				(MeshSectionData.RDGPositionBuffer != CommonParameters.RDGMeshPositionBuffer ? 1u : 0u) :
+				(MeshSectionData.PositionBuffer    != CommonParameters.MeshPositionBuffer    ? 1u : 0u);
+			SectionDatas[SectionIndex].bUseFaceNormal	= 0u;
+
+			// Sanity check
+			check(MeshSectionData.UVsChannelOffset < 255);
+			check(MeshSectionData.UVsChannelCount < 255);
+			check(CommonParameters.RDGMeshPositionBuffer == MeshSectionData.RDGPositionBuffer || CommonParameters.RDGMeshPositionBuffer == MeshSectionData.RDGPreviousPositionBuffer);
+			check(CommonParameters.MeshPositionBuffer    == MeshSectionData.PositionBuffer    || CommonParameters.MeshPositionBuffer    == MeshSectionData.PreviousPositionBuffer);
+			check(CommonParameters.MeshIndexBuffer		 == MeshSectionData.IndexBuffer);
+			check(CommonParameters.MeshTangentBuffer	 == MeshSectionData.TangentBuffer);
+		}
 	}
 
 	// If no previous position buffer available, reusing the current position buffers
