Fix groom crash on AMD when strands are stretched or invalid

diff --git a/Engine/Plugins/Runtime/HairStrands/Shaders/Private/NiagaraDataInterfaceHairStrandsTemplate.ush b/Engine/Plugins/Runtime/HairStrands/Shaders/Private/NiagaraDataInterfaceHairStrandsTemplate.ush
--- a/Engine/Plugins/Runtime/HairStrands/Shaders/Private/NiagaraDataInterfaceHairStrandsTemplate.ush
+++ b/Engine/Plugins/Runtime/HairStrands/Shaders/Private/NiagaraDataInterfaceHairStrandsTemplate.ush
@@ -365,7 +365,7 @@ void InitGridSamples_{ParameterName}(in float3 NodePosition, in float3 NodeVeloc
 		const float SegmentMass =  (SharedInverseMass[GGroupThreadId.x-1] * SegmentWeight.x + SharedInverseMass[GGroupThreadId.x] * SegmentWeight.y);
 		const float SegmentLength = length(OutDeltaPosition);
 
-		OutNumSamples = ceil(SegmentLength / GridLength);
+		OutNumSamples = isfinite(SegmentLength) ? min(ceil(SegmentLength / GridLength), 32) : 1;
 		OutSampleMass = SegmentMass / OutNumSamples;
 	}
 	else
