Added logging to catch null bone buffer from groom skin cache system and log information about it.

diff --git a/Engine/Plugins/Runtime/HairStrands/Source/HairStrandsCore/Private/GroomManager.cpp b/Engine/Plugins/Runtime/HairStrands/Source/HairStrandsCore/Private/GroomManager.cpp
--- a/Engine/Plugins/Runtime/HairStrands/Source/HairStrandsCore/Private/GroomManager.cpp
+++ b/Engine/Plugins/Runtime/HairStrands/Source/HairStrandsCore/Private/GroomManager.cpp
@@ -333,6 +333,10 @@ struct FHairGeometryCache
 				UpdateSection.BoneBuffer 			= FSkeletalMeshDeformerHelpers::GetBoneBufferForReading(InMeshObject, InLODIndex, SectionIndex, false);
 				UpdateSection.BonePrevBuffer 		= FSkeletalMeshDeformerHelpers::GetBoneBufferForReading(InMeshObject, InLODIndex, SectionIndex, true);
 				Data.RequestedSectionBits[SectionIndex] = true;
+
+				UE_CLOG(!UpdateSection.BoneBuffer, LogHairStrands, Error,
+					TEXT("Failed to get bone buffer from mesh object for hair skin cache. MeshObject: %s, LOD: %d, Section: %d, IsGPUSkinMesh %d, IsSectionValid %d"),
+					*InMeshObject->GetDebugName().ToString(), InLODIndex, SectionIndex, InMeshObject->IsGPUSkinMesh() ? 1 : 0, Section.IsValid() ? 1 : 0);
 			}
 		}
 
