Fix groom simulation reset during LOD transition

diff --git a/Engine/Plugins/Runtime/HairStrands/Source/HairStrandsCore/Private/Niagara/NiagaraDataInterfaceHairStrands.cpp b/Engine/Plugins/Runtime/HairStrands/Source/HairStrandsCore/Private/Niagara/NiagaraDataInterfaceHairStrands.cpp
--- a/Engine/Plugins/Runtime/HairStrands/Source/HairStrandsCore/Private/Niagara/NiagaraDataInterfaceHairStrands.cpp
+++ b/Engine/Plugins/Runtime/HairStrands/Source/HairStrandsCore/Private/Niagara/NiagaraDataInterfaceHairStrands.cpp
@@ -3518,7 +3518,7 @@ void UNiagaraDataInterfaceHairStrands::SetShaderParameters(const FNiagaraDataInt
 
 		// Simulation setup (we update the rest configuration based on the deformed positions 
 		// if in restupdate mode or if we are resetting the sim and using RBF transfer since the rest positions are not matching the physics asset)
-		const bool bShouldReset = HairStrandsBuffer->bShouldReset || !HairStrandsBuffer->bValidGeometryType || (HairStrandsBuffer->CurrentMeshLOD != MeshLODIndex);
+		const bool bShouldReset = HairStrandsBuffer->bShouldReset || !HairStrandsBuffer->bValidGeometryType;
 		HairStrandsBuffer->ResetCount = bShouldReset ? 0 : FMath::Min(GHairSimulationMaxDelay + 1, HairStrandsBuffer->ResetCount+ 1);
 		
 		const int32 NeedResetValue = HairStrandsBuffer->ResetCount <= GHairSimulationMaxDelay;
