Sync the Inline skinned mesh task group instead of SkinCache as it's using the manual skin cache for grooms.

diff --git a/Engine/Plugins/Runtime/HairStrands/Source/HairStrandsCore/Private/GroomManager.cpp b/Engine/Plugins/Runtime/HairStrands/Source/HairStrandsCore/Private/GroomManager.cpp
--- a/Engine/Plugins/Runtime/HairStrands/Source/HairStrandsCore/Private/GroomManager.cpp
+++ b/Engine/Plugins/Runtime/HairStrands/Source/HairStrandsCore/Private/GroomManager.cpp
@@ -659,7 +659,7 @@ static FCachedGeometry GetCacheGeometryForHair(
 				else if (IsHairManualSkinCacheEnabled())
 				{
 					// Sync the async update tasks since we will be accessing bone buffers.
-					FSkeletalMeshUpdater::WaitForStage(GraphBuilder, ESkeletalMeshUpdateStage::SkinCache);
+					FSkeletalMeshUpdater::WaitForStage(GraphBuilder, ESkeletalMeshUpdateStage::Inline);
 
 					uint32 OutTotalSectionCount = 0;
 					GetOrAllocateCachedGeometry(GraphBuilder, ShaderMap, SceneProxy, RootResource, bOutputTriangleData, Out, OutHairGeometryCache, OutTotalSectionCount);
