Fix crash due to missing bone buffer when using hair, by adding bone buffer building fallback for hair manual skin cache.

When a skel. mesh uses CPU skinning, it has no bone buffer. This CL adds the building and upload on the GPU of its bone buffer for proper deformation.

diff --git a/Engine/Plugins/Runtime/HairStrands/Source/HairStrandsCore/Private/GroomManager.cpp b/Engine/Plugins/Runtime/HairStrands/Source/HairStrandsCore/Private/GroomManager.cpp
--- a/Engine/Plugins/Runtime/HairStrands/Source/HairStrandsCore/Private/GroomManager.cpp
+++ b/Engine/Plugins/Runtime/HairStrands/Source/HairStrandsCore/Private/GroomManager.cpp
@@ -213,6 +213,51 @@ static FORCEINLINE uint32 GetHash(const FHairGeometryCacheKey& In)
 	return Murmur32({In.ComponentId, In.LODIndex}); 
 }
 
+static FRDGBufferRef UpdateBoneData(
+	FRDGBuilder& GraphBuilder,
+	const FSkeletalMeshObject* InMeshObject, 
+	TConstArrayView<FMatrix44f> ReferenceToLocalMatrices,
+	TConstArrayView<FBoneIndexType> BoneMap)
+{
+	const uint32 NumBones = BoneMap.Num();
+	check(IsInParallelRenderingThread());
+	check(NumBones > 0);
+	
+	TArray<FMatrix3x4> ChunkMatrices;
+	ChunkMatrices.SetNum(NumBones);
+
+	constexpr int32 PreFetchStride = 2; // FPlatformMisc::Prefetch stride
+	for (uint32 BoneIdx = 0; BoneIdx < NumBones; BoneIdx++)
+	{
+		const FBoneIndexType RefToLocalIdx = BoneMap[BoneIdx];
+		FPlatformMisc::Prefetch(ReferenceToLocalMatrices.GetData() + RefToLocalIdx + PreFetchStride);
+		FPlatformMisc::Prefetch(ReferenceToLocalMatrices.GetData() + RefToLocalIdx + PreFetchStride, PLATFORM_CACHE_LINE_SIZE);
+
+		FMatrix3x4& BoneMat = ChunkMatrices[BoneIdx];
+		const FMatrix44f& RefToLocal = ReferenceToLocalMatrices[RefToLocalIdx];
+		RefToLocal.To3x4MatrixTranspose((float*)BoneMat.M);
+	}
+
+	const uint32 BytesPerElement = sizeof(FMatrix3x4);
+	FRDGBufferRef BoneMatricesBuffer = CreateVertexBuffer(GraphBuilder, TEXT("Hair.SkinnedBoneMatrices"), FRDGBufferDesc::CreateBufferDesc(BytesPerElement, ChunkMatrices.Num()), ChunkMatrices.GetData(), BytesPerElement * ChunkMatrices.Num());
+	return BoneMatricesBuffer;
+}
+
+static FRDGBufferRef CreateBoneBuffer(FRDGBuilder& GraphBuilder, const FSkeletalMeshObject* InMeshObject, int32 LODIndex, int32 SectionIndex)
+{
+	const TArray<FSkelMeshRenderSection>& Sections = InMeshObject->GetRenderSections(LODIndex);
+	if (Sections.IsValidIndex(SectionIndex))
+	{
+		const FSkelMeshRenderSection& Section = Sections[SectionIndex];
+		if (Section.IsValid())
+		{
+			return UpdateBoneData(GraphBuilder, InMeshObject, InMeshObject->GetReferenceToLocalMatrices(), Section.BoneMap);
+		}
+	}
+
+	return nullptr;
+}
+
 struct FHairGeometryCache
 {
 	enum class ECacheType
@@ -232,6 +277,7 @@ struct FHairGeometryCache
 		FRDGBufferSRVRef PositionSRV				= nullptr;
 		FRDGBufferSRVRef PreviousPositionSRV		= nullptr;
 		FRDGBufferSRVRef TangentSRV					= nullptr;
+		TArray<TRefCountPtr<FRDGPooledBuffer>> BoneBuffers;
 
 		FHairGeometryCacheKey Key;
 		uint32 Hash 				= 0;
@@ -313,7 +359,7 @@ struct FHairGeometryCache
 				Data.TangentBuffer = GraphBuilder.CreateBuffer(FRDGBufferDesc::CreateBufferDesc(TangentStride, TangentElementCount), TEXT("Hair.SkinnedDeformedTangents"));
 				Data.TangentSRV = GraphBuilder.CreateSRV(Data.TangentBuffer, TangentFormat);
 			}
-
+			Data.BoneBuffers.Init(nullptr, Data.TotalSectionCount);
 			Data.RequestedSections.Reserve(Data.TotalSectionCount);
 			Data.RequestedSectionBits.Init(false, Data.TotalSectionCount);
 		}
@@ -334,6 +380,18 @@ struct FHairGeometryCache
 				UpdateSection.BonePrevBuffer 		= FSkeletalMeshDeformerHelpers::GetBoneBufferForReading(InMeshObject, InLODIndex, SectionIndex, true);
 				Data.RequestedSectionBits[SectionIndex] = true;
 
+				if (UpdateSection.BoneBuffer == nullptr)
+				{
+					check(InMeshObject->IsCPUSkinned());
+					if (Data.BoneBuffers[SectionIndex] == nullptr)
+					{
+						FRDGBufferRef BoneBuffer = CreateBoneBuffer(GraphBuilder, InMeshObject, InLODIndex, SectionIndex);
+						Data.BoneBuffers[SectionIndex] = GraphBuilder.ConvertToExternalBuffer(BoneBuffer);
+					}
+					UpdateSection.BoneBuffer = Data.BoneBuffers[SectionIndex]->GetOrCreateSRV(GraphBuilder.RHICmdList, FRHIBufferSRVCreateInfo(PF_A32B32G32R32F));
+				}
+				check(UpdateSection.BoneBuffer);
+
 				UE_CLOG(!UpdateSection.BoneBuffer, LogHairStrands, Error,
 					TEXT("Failed to get bone buffer from mesh object for hair skin cache. MeshObject: %s, LOD: %d, Section: %d, IsGPUSkinMesh %d, IsSectionValid %d"),
 					*InMeshObject->GetDebugName().ToString(), InLODIndex, SectionIndex, InMeshObject->IsGPUSkinMesh() ? 1 : 0, Section.IsValid() ? 1 : 0);
