Fix determinism in hair card builder.

#FailedValidations Preflight

diff --git a/Engine/Plugins/Runtime/HairStrands/Source/HairStrandsCore/Private/HairCardsBuilder.cpp b/Engine/Plugins/Runtime/HairStrands/Source/HairStrandsCore/Private/HairCardsBuilder.cpp
--- a/Engine/Plugins/Runtime/HairStrands/Source/HairStrandsCore/Private/HairCardsBuilder.cpp
+++ b/Engine/Plugins/Runtime/HairStrands/Source/HairStrandsCore/Private/HairCardsBuilder.cpp
@@ -121,15 +121,15 @@ static bool InternalCreateCardsGuides(
 
 	const uint32 NumCards = InCards.IndexOffsets.Num();
 	const uint32 MaxGuidePoints = InCards.Indices.Num() / 3;
-	OutCardLengths.SetNum(NumCards);
+	OutCardLengths.SetNumZeroed(NumCards);
 
 	OutGuides.StrandsPoints.PointsPosition.Reserve(MaxGuidePoints);
 	OutGuides.StrandsPoints.PointsRadius.Reserve(MaxGuidePoints);
 	OutGuides.StrandsPoints.PointsCoordU.Reserve(MaxGuidePoints);
 	OutGuides.StrandsCurves.SetNum(NumCards, 0u);
 
-	InCards.CoordU.SetNum(InCards.Positions.Num());
-	InCards.LocalUVs.SetNum(InCards.Positions.Num());
+	InCards.CoordU.SetNumZeroed(InCards.Positions.Num());
+	InCards.LocalUVs.SetNumZeroed(InCards.Positions.Num());
 
 	// Find the principal direction of the atlas by comparing the number of segment along U and along V
 	bool bIsMainDirectionU = true;
@@ -579,17 +579,17 @@ static bool InternalImportGeometry_WithGeneratedGuides(
 	const uint32 NumUVs = NumUVSource + 1 /*RootUV*/;
 
 	// Fill vertex data
-	Out.Cards.Positions.SetNum(VertexCount);
-	Out.Cards.Normals.SetNum(VertexCount);
-	Out.Cards.Tangents.SetNum(VertexCount);
+	Out.Cards.Positions.SetNumZeroed(VertexCount);
+	Out.Cards.Normals.SetNumZeroed(VertexCount);
+	Out.Cards.Tangents.SetNumZeroed(VertexCount);
 	Out.Cards.VertexColors.Empty();
-	Out.Cards.UVs.SetNum(NumUVs);
+	Out.Cards.UVs.SetNumZeroed(NumUVs);
 	for (uint32 UVIt = 0;UVIt<NumUVs; ++UVIt)
 	{
 		Out.Cards.UVs[UVIt].UV.Init(FVector2f::ZeroVector, VertexCount);
 	}
 	TArray<float> TangentFrameSigns;
-	TangentFrameSigns.SetNum(VertexCount);
+	TangentFrameSigns.SetNumZeroed(VertexCount);
 
 	Out.Cards.BoundingBox.Init();
 
@@ -604,7 +604,7 @@ static bool InternalImportGeometry_WithGeneratedGuides(
 	const bool bHasVertexColor = VertexColors.GetNumElements() > 0;
 	if (bHasVertexColor)
 	{
-		Out.Cards.VertexColors.SetNum(VertexCount);
+		Out.Cards.VertexColors.SetNumZeroed(VertexCount);
 	}
 	for (const FVertexID VertexId : MeshDescription->Vertices().GetElementIDs())
 	{
@@ -638,14 +638,14 @@ static bool InternalImportGeometry_WithGeneratedGuides(
 
 	// Fill in render resources (do we need to keep it separated? e.g, format compression, packing)
 	const uint32 PointCount = Out.Cards.Positions.Num();
-	OutBulk.Positions.SetNum(PointCount);
-	OutBulk.Normals.SetNum(PointCount * FHairCardsNormalFormat::ComponentCount);
-	OutBulk.UVs.SetNum(GetNumPackedUVs(PointCount, NumUVs));
-	OutBulk.Materials.SetNum(PointCount);
+	OutBulk.Positions.SetNumZeroed(PointCount);
+	OutBulk.Normals.SetNumZeroed(PointCount * FHairCardsNormalFormat::ComponentCount);
+	OutBulk.UVs.SetNumZeroed(GetNumPackedUVs(PointCount, NumUVs));
+	OutBulk.Materials.SetNumZeroed(PointCount);
 	OutBulk.VertexColors.Empty();
 	if (bHasVertexColor)
 	{
-		OutBulk.VertexColors.SetNum(PointCount);
+		OutBulk.VertexColors.SetNumZeroed(PointCount);
 	}
 	for (uint32 PointIt = 0; PointIt < PointCount; ++PointIt)
 	{
@@ -659,7 +659,7 @@ static bool InternalImportGeometry_WithGeneratedGuides(
 		}
 	}
 
-	OutBulk.Indices.SetNum(IndexCount);
+	OutBulk.Indices.SetNumZeroed(IndexCount);
 	for (uint32 IndexIt = 0; IndexIt < IndexCount; ++IndexIt)
 	{
 		OutBulk.Indices[IndexIt] = Out.Cards.Indices[IndexIt];
@@ -683,7 +683,7 @@ static bool InternalImportGeometry_WithGeneratedGuides(
 
 		// Fill out the interpolation data
 		static_assert(sizeof(FHairCardsInterpolationVertex) == HAIR_INTERPOLATION_CARDS_GUIDE_STRIDE);
-		OutInterpolationBulkData.Interpolation.SetNum(PointCount);
+		OutInterpolationBulkData.Interpolation.SetNumZeroed(PointCount);
 		for (uint32 PointIt = 0; PointIt < PointCount; ++PointIt)
 		{
 			const uint32 InterpVertexIndex = InterpolationData.PointsSimCurvesVertexIndex[PointIt];
@@ -850,7 +850,7 @@ static bool InternalImportGeometry_WithGeneratedGuides(
 			FVector2f RootUV1;
 		};
 		TArray<FRootUVData> OutRootUVDatas;
-		OutRootUVDatas.SetNum(CardsRoots.Num());
+		OutRootUVDatas.SetNumZeroed(CardsRoots.Num());
 
 		// 3. Find cards root / curve root
 
@@ -980,14 +980,14 @@ static bool InternalImportGeometry_WithImportedGuides(
 	const uint32 NumUVs = NumSourceUVs + 1u /*RootUV*/;
 
 	// Write out BulkData
-	OutBulk.Positions.SetNum(PointCount);
-	OutBulk.Normals.SetNum(PointCount * FHairCardsNormalFormat::ComponentCount);
-	OutBulk.UVs.SetNum(GetNumPackedUVs(PointCount, NumUVs));
-	OutBulk.Materials.SetNum(PointCount);
+	OutBulk.Positions.SetNumZeroed(PointCount);
+	OutBulk.Normals.SetNumZeroed(PointCount * FHairCardsNormalFormat::ComponentCount);
+	OutBulk.UVs.SetNumZeroed(GetNumPackedUVs(PointCount, NumUVs));
+	OutBulk.Materials.SetNumZeroed(PointCount);
 	OutBulk.VertexColors.Empty();
 	if (bHasVertexColor)
 	{
-		OutBulk.VertexColors.SetNum(PointCount);
+		OutBulk.VertexColors.SetNumZeroed(PointCount);
 	}
 	FBox3f BoundingBox = FBox3f(EForceInit::ForceInit);
 	for (const FVertexID VertexId : MeshDescription->Vertices().GetElementIDs())
@@ -1067,7 +1067,7 @@ static bool InternalImportGeometry_WithImportedGuides(
 
 		FHairCardsInterpolationDatas InterpolationData;
 		InterpolationData.SetNum(PointCount);
-		OutInterpolationBulkData.Interpolation.SetNum(PointCount);
+		OutInterpolationBulkData.Interpolation.SetNumZeroed(PointCount);
 
 		TArray<float> CoordUs;
 		CoordUs.Init(0, PointCount);
@@ -1258,12 +1258,12 @@ void BuildGeometry(
 
 	FHairMeshesDatas Out;
 
-	Out.Meshes.Positions.SetNum(TotalPointCount);
-	Out.Meshes.Normals.SetNum(TotalPointCount);
-	Out.Meshes.Tangents.SetNum(TotalPointCount);
-	Out.Meshes.Indices.SetNum(TotalIndexCount);
-	Out.Meshes.UVs.SetNum(HAIR_CARD_UV_COUNT);
-	Out.Meshes.UVs[HAIR_CARD_UV_CARDUV].UV.SetNum(TotalPointCount);
+	Out.Meshes.Positions.SetNumZeroed(TotalPointCount);
+	Out.Meshes.Normals.SetNumZeroed(TotalPointCount);
+	Out.Meshes.Tangents.SetNumZeroed(TotalPointCount);
+	Out.Meshes.Indices.SetNumZeroed(TotalIndexCount);
+	Out.Meshes.UVs.SetNumZeroed(HAIR_CARD_UV_COUNT);
+	Out.Meshes.UVs[HAIR_CARD_UV_CARDUV].UV.SetNumZeroed(TotalPointCount);
 	Out.Meshes.UVs[HAIR_CARD_UV_ROOTUV].UV.Init(FVector2f(0,0), TotalPointCount);
 
 	Out.Meshes.Positions[0] = Center + FVector3f(-Extent.X, -Extent.Y, -Extent.Z);
@@ -1349,9 +1349,9 @@ void BuildGeometry(
 	// Fill in render resources (do we need to keep it separated? e.g, format compression, packing)
 	const uint32 NumUVs = 2; // RootUV+CardUV
 	const uint32 PointCount = Out.Meshes.Positions.Num();
-	OutBulk.Positions.SetNum(PointCount);
-	OutBulk.Normals.SetNum(PointCount * FHairCardsNormalFormat::ComponentCount);
-	OutBulk.UVs.SetNum(FHairCardsBuilder::GetNumPackedUVs(PointCount, NumUVs));
+	OutBulk.Positions.SetNumZeroed(PointCount);
+	OutBulk.Normals.SetNumZeroed(PointCount * FHairCardsNormalFormat::ComponentCount);
+	OutBulk.UVs.SetNumZeroed(FHairCardsBuilder::GetNumPackedUVs(PointCount, NumUVs));
 	OutBulk.VertexColors.Empty();
 	for (uint32 PointIt = 0; PointIt < PointCount; ++PointIt)
 	{
@@ -1368,7 +1368,7 @@ void BuildGeometry(
 	OutBulk.Header.NumUVs = NumUVs;
 
 	const uint32 IndexCount = Out.Meshes.Indices.Num();
-	OutBulk.Indices.SetNum(IndexCount);
+	OutBulk.Indices.SetNumZeroed(IndexCount);
 	for (uint32 IndexIt = 0; IndexIt < IndexCount; ++IndexIt)
 	{
 		OutBulk.Indices[IndexIt] = Out.Meshes.Indices[IndexIt];
@@ -1394,12 +1394,12 @@ void ImportGeometry(
 	const uint32 NumUVsSource = LODData.VertexBuffers.StaticMeshVertexBuffer.GetNumTexCoords();
 	const uint32 NumUVs = NumUVsSource + 1; // +1 for RootUV
 	
-	Out.Meshes.Positions.SetNum(VertexCount);
-	Out.Meshes.Normals.SetNum(VertexCount);
-	Out.Meshes.Tangents.SetNum(VertexCount);
+	Out.Meshes.Positions.SetNumZeroed(VertexCount);
+	Out.Meshes.Normals.SetNumZeroed(VertexCount);
+	Out.Meshes.Tangents.SetNumZeroed(VertexCount);
 	Out.Meshes.VertexColors.Empty();
-	Out.Meshes.Indices.SetNum(IndexCount);
-	Out.Meshes.UVs.SetNum(NumUVs);
+	Out.Meshes.Indices.SetNumZeroed(IndexCount);
+	Out.Meshes.UVs.SetNumZeroed(NumUVs);
 	for (uint32 UVIt = 0; UVIt < NumUVs; ++UVIt)
 	{
 		Out.Meshes.UVs[UVIt].UV.Init(FVector2f::ZeroVector, VertexCount);
@@ -1407,7 +1407,7 @@ void ImportGeometry(
 
 	if (bHasVertexColor)
 	{
-		Out.Meshes.VertexColors.SetNum(VertexCount);
+		Out.Meshes.VertexColors.SetNumZeroed(VertexCount);
 		LODData.VertexBuffers.ColorVertexBuffer.GetVertexColors(Out.Meshes.VertexColors);
 	}
 
@@ -1438,13 +1438,13 @@ void ImportGeometry(
 	OutBulk.Header.NumUVs = NumUVs;
 
 	const uint32 PointCount = Out.Meshes.Positions.Num();
-	OutBulk.Positions.SetNum(PointCount);
-	OutBulk.Normals.SetNum(PointCount * FHairCardsNormalFormat::ComponentCount);
-	OutBulk.UVs.SetNum(FHairCardsBuilder::GetNumPackedUVs(PointCount, NumUVs));
+	OutBulk.Positions.SetNumZeroed(PointCount);
+	OutBulk.Normals.SetNumZeroed(PointCount * FHairCardsNormalFormat::ComponentCount);
+	OutBulk.UVs.SetNumZeroed(FHairCardsBuilder::GetNumPackedUVs(PointCount, NumUVs));
 	OutBulk.VertexColors.Empty();
 	if (bHasVertexColor)
 	{
-		OutBulk.VertexColors.SetNum(PointCount);
+		OutBulk.VertexColors.SetNumZeroed(PointCount);
 	}
 	for (uint32 PointIt = 0; PointIt < PointCount; ++PointIt)
 	{
@@ -1462,7 +1462,7 @@ void ImportGeometry(
 		}
 	}
 
-	OutBulk.Indices.SetNum(IndexCount);
+	OutBulk.Indices.SetNumZeroed(IndexCount);
 	for (uint32 IndexIt = 0; IndexIt < IndexCount; ++IndexIt)
 	{
 		OutBulk.Indices[IndexIt] = Out.Meshes.Indices[IndexIt];
diff --git a/Engine/Plugins/Runtime/HairStrands/Source/HairStrandsCore/Private/HairCardsDatas.cpp b/Engine/Plugins/Runtime/HairStrands/Source/HairStrandsCore/Private/HairCardsDatas.cpp
--- a/Engine/Plugins/Runtime/HairStrands/Source/HairStrandsCore/Private/HairCardsDatas.cpp
+++ b/Engine/Plugins/Runtime/HairStrands/Source/HairStrandsCore/Private/HairCardsDatas.cpp
@@ -54,9 +54,9 @@ void FHairMeshesBulkData::Serialize(FArchive& Ar)
 /////////////////////////////////////////////////////////////////////////////////////////
 void FHairCardsInterpolationDatas::SetNum(const uint32 NumPoints)
 {
-	PointsSimCurvesIndex.SetNum(NumPoints);
-	PointsSimCurvesVertexIndex.SetNum(NumPoints);
-	PointsSimCurvesVertexLerp.SetNum(NumPoints);
+	PointsSimCurvesIndex.SetNumZeroed(NumPoints);
+	PointsSimCurvesVertexIndex.SetNumZeroed(NumPoints);
+	PointsSimCurvesVertexLerp.SetNumZeroed(NumPoints);
 }
 
 void FHairCardsInterpolationDatas::Reset()
