Do not force async compute pipeline when when hair raytracing geometry is disabled.

diff --git a/Engine/Plugins/Runtime/HairStrands/Source/HairStrandsCore/Private/GroomManager.cpp b/Engine/Plugins/Runtime/HairStrands/Source/HairStrandsCore/Private/GroomManager.cpp
--- a/Engine/Plugins/Runtime/HairStrands/Source/HairStrandsCore/Private/GroomManager.cpp
+++ b/Engine/Plugins/Runtime/HairStrands/Source/HairStrandsCore/Private/GroomManager.cpp
@@ -2040,7 +2040,7 @@ static void RunHairStrandsInterpolation_Cards(
 	ERHIPipeline CardPositionExternalAccessPipeline = ERHIPipeline::Graphics;
 #if RHI_RAYTRACING
 	// Also add async compute as queue which can access these external hair resources
-	if (GSupportsEfficientAsyncCompute && GRHIGlobals.RayTracing.SupportsAsyncRayTraceDispatch)
+	if (GSupportsEfficientAsyncCompute && GRHIGlobals.RayTracing.SupportsAsyncRayTraceDispatch && IsHairRayTracingEnabled())
 	{
 		EnumAddFlags(CardPositionExternalAccessPipeline, ERHIPipeline::AsyncCompute);
 	}
