[GroomBuilder]
- Slight optim that can have a big effect on instrumented builds by avoiding global memory access we replace 3 calls to _Instrument_Load by a single one

diff --git a/Engine/Plugins/Runtime/HairStrands/Source/HairStrandsCore/Private/GroomBuilder.cpp b/Engine/Plugins/Runtime/HairStrands/Source/HairStrandsCore/Private/GroomBuilder.cpp
--- a/Engine/Plugins/Runtime/HairStrands/Source/HairStrandsCore/Private/GroomBuilder.cpp
+++ b/Engine/Plugins/Runtime/HairStrands/Source/HairStrandsCore/Private/GroomBuilder.cpp
@@ -145,14 +145,15 @@ namespace HairStrandsBuilder
 				FVector3f PreviousPosition(0.0, 0.0, 0.0);
 				for (uint32 PointIndex = 0; PointIndex < StrandCount; ++PointIndex, ++PositionIterator, ++RadiusIterator, ++CoordUIterator)
 				{
-					Out.BoundingBox += *PositionIterator;
+					const FVector3f CurrentPosition = *PositionIterator;
+					Out.BoundingBox += CurrentPosition;
 					Out.BoundingBox.IsValid = 1;
 					if (PointIndex > 0)
 					{
-						StrandLength += (*PositionIterator - PreviousPosition).Size();
+						StrandLength += (CurrentPosition - PreviousPosition).Size();
 					}
 					*CoordUIterator = StrandLength;
-					PreviousPosition = *PositionIterator;
+					PreviousPosition = CurrentPosition;
 				}
 				*LengthIterator = StrandLength;
 			}
