Fix hair card only groom asset.

#Synced-CL 48446704

diff --git a/Engine/Plugins/Runtime/HairStrands/Source/HairStrandsCore/Private/GroomAsset.cpp b/Engine/Plugins/Runtime/HairStrands/Source/HairStrandsCore/Private/GroomAsset.cpp
--- a/Engine/Plugins/Runtime/HairStrands/Source/HairStrandsCore/Private/GroomAsset.cpp
+++ b/Engine/Plugins/Runtime/HairStrands/Source/HairStrandsCore/Private/GroomAsset.cpp
@@ -2988,66 +2988,73 @@ static bool InternalImportCardGeometry(
 {
 	UStaticMesh* CardsMesh = InDesc->ImportedMesh;
 	check(CardsMesh);
-	check(InHairDescriptionGroups.IsValid());
-	check(InHairDescriptionGroups.HairGroups.IsValidIndex(InDesc->GroupIndex));
 
 	// * Create a transient FHairStrandsData in order to extract RootUV and transfer them to cards data
 	// * Voxelize hair strands data (all group), to transfer group index from strands to cards
 	//FHairGroupInfo DummyInfo;
+	if (InHairDescriptionGroups.IsValid() && InHairDescriptionGroups.HairGroups.IsValidIndex(InDesc->GroupIndex))
+	{
+		FHairStrandsDatas TempHairGuidesData;
+		FHairStrandsDatas TempHairStrandsData;
+		FGroomBuilder::BuildData(
+			InHairDescriptionGroups.HairGroups[InDesc->GroupIndex], 
+			InHairInterpolationGroups[InDesc->GroupIndex],
+			TempHairStrandsData,
+			TempHairGuidesData);
 
-	FHairStrandsDatas TempHairGuidesData;
-	FHairStrandsDatas TempHairStrandsData;
-	FGroomBuilder::BuildData(
-		InHairDescriptionGroups.HairGroups[InDesc->GroupIndex], 
-		InHairInterpolationGroups[InDesc->GroupIndex],
-		TempHairStrandsData,
-		TempHairGuidesData);
-
-	FHairStrandsVoxelData TempHairStrandsVoxelData;
-	FGroomBuilder::VoxelizeGroupIndex(InHairDescriptionGroups, TempHairStrandsVoxelData);
+		FHairStrandsVoxelData TempHairStrandsVoxelData;
+		FGroomBuilder::VoxelizeGroupIndex(InHairDescriptionGroups, TempHairStrandsVoxelData);
 
-	// Transfer RootUV from strands to guides
-	const bool bNeedRootUVTransfer = InDesc->GuideType == EHairCardsGuideType::GuideBased;
-	if (bNeedRootUVTransfer)
-	{
-		struct FRoot
+		// Transfer RootUV from strands to guides
+		const bool bNeedRootUVTransfer = InDesc->GuideType == EHairCardsGuideType::GuideBased;
+		if (bNeedRootUVTransfer)
 		{
-			FVector3f Position;
-			FVector2f RootUV;
-		};
+			struct FRoot
+			{
+				FVector3f Position;
+				FVector2f RootUV;
+			};
 
-		// Extra all roots from strands
-		TArray<FRoot> StrandsRoots;
-		StrandsRoots.SetNum(TempHairStrandsData.GetNumCurves());
-		for (uint32 CurveIt = 0, CurveCount = TempHairStrandsData.GetNumCurves(); CurveIt < CurveCount; ++CurveIt)
-		{
-			StrandsRoots[CurveIt].Position = TempHairStrandsData.StrandsPoints.PointsPosition[TempHairStrandsData.StrandsCurves.CurvesOffset[CurveIt]];
-			StrandsRoots[CurveIt].RootUV   = TempHairStrandsData.StrandsCurves.CurvesRootUV[CurveIt];
-		}
+			// Extra all roots from strands
+			TArray<FRoot> StrandsRoots;
+			StrandsRoots.SetNum(TempHairStrandsData.GetNumCurves());
+			for (uint32 CurveIt = 0, CurveCount = TempHairStrandsData.GetNumCurves(); CurveIt < CurveCount; ++CurveIt)
+			{
+				StrandsRoots[CurveIt].Position = TempHairStrandsData.StrandsPoints.PointsPosition[TempHairStrandsData.StrandsCurves.CurvesOffset[CurveIt]];
+				StrandsRoots[CurveIt].RootUV   = TempHairStrandsData.StrandsCurves.CurvesRootUV[CurveIt];
+			}
 
-		// Find closest strands root for each guide roots, and assign RootUV
-		for (uint32 CurveIt = 0, CurveCount = TempHairGuidesData.GetNumCurves(); CurveIt < CurveCount; ++CurveIt)
-		{
-			float MaxDistance = FLT_MAX;
-			const FVector3f Position = TempHairGuidesData.StrandsPoints.PointsPosition[TempHairGuidesData.StrandsCurves.CurvesOffset[CurveIt]];
-			for (const FRoot& Root : StrandsRoots)
+			// Find closest strands root for each guide roots, and assign RootUV
+			for (uint32 CurveIt = 0, CurveCount = TempHairGuidesData.GetNumCurves(); CurveIt < CurveCount; ++CurveIt)
 			{
-				const float Dist = (Position - Root.Position).Length();
-				if (Dist < MaxDistance)
+				float MaxDistance = FLT_MAX;
+				const FVector3f Position = TempHairGuidesData.StrandsPoints.PointsPosition[TempHairGuidesData.StrandsCurves.CurvesOffset[CurveIt]];
+				for (const FRoot& Root : StrandsRoots)
 				{
-					TempHairGuidesData.StrandsCurves.CurvesRootUV[CurveIt] = Root.RootUV;
-					MaxDistance = Dist;
+					const float Dist = (Position - Root.Position).Length();
+					if (Dist < MaxDistance)
+					{
+						TempHairGuidesData.StrandsCurves.CurvesRootUV[CurveIt] = Root.RootUV;
+						MaxDistance = Dist;
+					}
 				}
 			}
 		}
+		return FHairCardsBuilder::ImportGeometry(CardsMesh, TempHairGuidesData, TempHairStrandsData, TempHairStrandsVoxelData, InDesc->GuideType == EHairCardsGuideType::Generated, OutBulkData, OutGuidesData, OutInterpolationBulkData);
+	}
+	else
+	{
+		// If no hair description is available, build hair card with dummy hair strands data
+		FHairStrandsDatas DummyHairGuidesData;
+		FHairStrandsDatas DummyHairStrandsData;
+		FHairStrandsVoxelData DummyHairStrandsVoxelData;
+		return FHairCardsBuilder::ImportGeometry(CardsMesh, DummyHairGuidesData, DummyHairStrandsData, DummyHairStrandsVoxelData, true /*bGenerateGuidesFromCardGeometry*/, OutBulkData, OutGuidesData, OutInterpolationBulkData);
 	}
-
-	return FHairCardsBuilder::ImportGeometry(CardsMesh, TempHairGuidesData, TempHairStrandsData, TempHairStrandsVoxelData, InDesc->GuideType == EHairCardsGuideType::Generated, OutBulkData, OutGuidesData, OutInterpolationBulkData);
 }
 
 bool UGroomAsset::HasChanged_Cards(uint32 GroupIndex, TArray<bool>& OutIsValid) const 
 {
-	if (!IsHairStrandsEnabled(EHairStrandsShaderType::Cards) || (GetHairGroupsCards().Num() == 0 && CachedHairGroupsCards.Num() == 0)|| !CanRebuildFromDescription())
+	if (!IsHairStrandsEnabled(EHairStrandsShaderType::Cards) || (GetHairGroupsCards().Num() == 0 && CachedHairGroupsCards.Num() == 0))
 	{
 		return false;
 	}
@@ -3144,11 +3151,18 @@ bool UGroomAsset::BuildHairGroup_Cards(uint32 GroupIndex)
 
 				// 2. (Re)Build data for the sim guides (since there are transient/no-cached)
 				const FHairDescriptionGroups& LocalHairDescriptionGroups = GetHairDescriptionGroups();
-				check(LocalHairDescriptionGroups.IsValid());
 				FHairStrandsDatas SimGuidesData;
-				FHairStrandsDatas RenStrandsData; // Dummy, not used
-				FGroomBuilder::BuildData(LocalHairDescriptionGroups.HairGroups[GroupIndex], GetHairGroupsInterpolation()[GroupIndex], RenStrandsData, SimGuidesData);
-
+				if (LocalHairDescriptionGroups.IsValid())
+				{
+					FHairStrandsDatas RenStrandsData; // Dummy, not used
+					FGroomBuilder::BuildData(LocalHairDescriptionGroups.HairGroups[GroupIndex], GetHairGroupsInterpolation()[GroupIndex], RenStrandsData, SimGuidesData);
+				}
+				else
+				{
+					// Use cards' guides as simulation guides for building interpolation.
+					// but these cards' guides won't be simulated, since they are not written as main guides
+					SimGuidesData = LODGuidesData;
+				}
 				// 3. Compute interpolation data & bulk data between the sim guides and the cards guides
 				FHairStrandsInterpolationDatas LODGuideInterpolationData;
 				FGroomBuilder::BuildInterplationData(DummyGroupInfo, LODGuidesData, SimGuidesData, CardsInterpolationSettings, LODGuideInterpolationData);
