Temporarily removed IsCPUSkinned check when creating Groom Bone Buffer.

#Synced-CL 49723226

diff --git a/Engine/Plugins/Runtime/HairStrands/Source/HairStrandsCore/Private/GroomManager.cpp b/Engine/Plugins/Runtime/HairStrands/Source/HairStrandsCore/Private/GroomManager.cpp
--- a/Engine/Plugins/Runtime/HairStrands/Source/HairStrandsCore/Private/GroomManager.cpp
+++ b/Engine/Plugins/Runtime/HairStrands/Source/HairStrandsCore/Private/GroomManager.cpp
@@ -382,7 +382,6 @@ struct FHairGeometryCache
 
 				if (UpdateSection.BoneBuffer == nullptr)
 				{
-					check(InMeshObject->IsCPUSkinned());
 					if (Data.BoneBuffers[SectionIndex] == nullptr)
 					{
 						FRDGBufferRef BoneBuffer = CreateBoneBuffer(GraphBuilder, InMeshObject, InLODIndex, SectionIndex);
