HairStrands: Fix for the instance scale factor being applied twice for raytraced hair primitives

Use the raytracing scale factor to compenstate for this effect

#Synced-CL 49711686

diff --git a/Engine/Plugins/Runtime/HairStrands/Source/HairStrandsCore/Private/HairStrandsInterpolation.cpp b/Engine/Plugins/Runtime/HairStrands/Source/HairStrandsCore/Private/HairStrandsInterpolation.cpp
--- a/Engine/Plugins/Runtime/HairStrands/Source/HairStrandsCore/Private/HairStrandsInterpolation.cpp
+++ b/Engine/Plugins/Runtime/HairStrands/Source/HairStrandsCore/Private/HairStrandsInterpolation.cpp
@@ -1819,6 +1819,7 @@ FHairGroupPublicData::FVertexFactoryInput InternalComputeHairStrandsVertexInputD
 	OutVFInput.Strands.Common.RootScale = Instance->Strands.Modifier.HairRootScale;
 	OutVFInput.Strands.Common.TipScale = Instance->Strands.Modifier.HairTipScale;
 	OutVFInput.Strands.Common.RaytracingRadiusScale = (GHairRaytracingRadiusScale > 0 ? GHairRaytracingRadiusScale : Instance->Strands.Modifier.HairRaytracingRadiusScale);
+	OutVFInput.Strands.Common.RaytracingRadiusScale /= InstanceScale; // In raytracing, the BLAS is automatically scaled by the instance transform, so cancel out the InstanceScale factor above to avoid applying it twice
 	OutVFInput.Strands.Common.Length = Instance->Strands.Modifier.HairLength;
 	OutVFInput.Strands.Common.Density = Instance->Strands.Modifier.HairShadowDensity;
 	OutVFInput.Strands.Common.Flags = 0;
