Fix crash when no rest resources are availables + prevent access to the groom component from the render thread

#Synced-CL 49728154

diff --git a/Engine/Plugins/Runtime/HairStrands/Source/HairStrandsCore/Private/GroomBuilder.cpp b/Engine/Plugins/Runtime/HairStrands/Source/HairStrandsCore/Private/GroomBuilder.cpp
--- a/Engine/Plugins/Runtime/HairStrands/Source/HairStrandsCore/Private/GroomBuilder.cpp
+++ b/Engine/Plugins/Runtime/HairStrands/Source/HairStrandsCore/Private/GroomBuilder.cpp
@@ -4119,7 +4119,8 @@ struct FStrandsPositionData
 	uint32 ReadbackDelay = 0;
 	uint32 CurveCount = 0;
 	uint32 PointCount = 0;
-	const UGroomComponent* Component = nullptr;
+	FSceneInterface* Scene;
+	TRefCountPtr<FHairGroupInstance> GroupInstance = nullptr;
 	TSharedPtr<FStrandsPositionOutput> Output;
 	bool bReadGuides = false;
 };
@@ -4142,12 +4143,12 @@ void RunHairStrandsPositionQueries(FRDGBuilder& GraphBuilder, FGlobalShaderMap*
 		FStrandsPositionData* Q = nullptr;
 		while (GStrandsPositionReadbacks.Dequeue(Q))
 		{
-			if(Q && Q->Component && Q->Component->GetWorld() &&
-				SceneView && SceneView->Family && (Q->Component->GetWorld()->Scene == SceneView->Family->Scene))
+			if(Q && Q->Scene && Q->GroupInstance &&
+				SceneView && SceneView->Family && (Q->Scene == SceneView->Family->Scene))
 			{
 				if (Q->GPUPositions->IsReady() || Q->ReadbackDelay == 0)
 				{
-					const FHairGroupInstance* Instance = Q->Component->GetGroupInstance(Q->GroupIt);
+					const TRefCountPtr<FHairGroupInstance> Instance = Q->GroupInstance;
 
 					check(Q->GroupIt < uint32(Q->Output->Groups.Num()));
 					FStrandsPositionOutput::FGroup& Group = Q->Output->Groups[Q->GroupIt];
@@ -4197,11 +4198,10 @@ void RunHairStrandsPositionQueries(FRDGBuilder& GraphBuilder, FGlobalShaderMap*
 		FStrandsPositionData* Q = nullptr;
 		while (GStrandsPositionQueries.Dequeue(Q))
 		{
-			if(Q && Q->Component && Q->Component->GetWorld() &&
-				SceneView && SceneView->Family && (Q->Component->GetWorld()->Scene == SceneView->Family->Scene))
+			if(Q && Q->Scene && Q->GroupInstance &&
+				SceneView && SceneView->Family && (Q->Scene == SceneView->Family->Scene))
 			{
-				check(Q->Component);
-				if (const FHairGroupInstance* Instance = Q->Component->GetGroupInstance(Q->GroupIt))
+				if (const TRefCountPtr<FHairGroupInstance> Instance = Q->GroupInstance)
 				{
 					Q->CurveCount = Q->bReadGuides ? Instance->Guides.GetData().GetNumCurves() : Instance->HairGroupPublicData->GetActiveStrandsCurveCount();
 					Q->PointCount = Q->bReadGuides ? Instance->Guides.GetData().GetNumPoints() : Instance->HairGroupPublicData->GetActiveStrandsPointCount();
@@ -4239,6 +4239,10 @@ void RunHairStrandsPositionQueries(FRDGBuilder& GraphBuilder, FGlobalShaderMap*
 	}
 }
 
+// Forward declaration of the GroomComponent function that is retrieving a ref counted group instance ptr given an index
+TRefCountPtr<FHairGroupInstance> GetRefCountedHairGroupInstance(const UGroomComponent* GroomComponent, const int32 GroupIndex);
+
+
 #if WITH_EDITOR
 bool RequestStrandsPosition(const UGroomComponent* Component, TSharedPtr<FStrandsPositionOutput> Output, const bool bReadGuides)
 {
@@ -4253,7 +4257,8 @@ bool RequestStrandsPosition(const UGroomComponent* Component, TSharedPtr<FStrand
 		FStrandsPositionData* Data = new FStrandsPositionData();
 		Data->GroupIt = GroupIt;
 		Data->GroupCount = GroupCount;
-		Data->Component = Component;
+		Data->GroupInstance = GetRefCountedHairGroupInstance(Component, GroupIt);
+		Data->Scene = (Component && Component->GetWorld()) ? Component->GetWorld()->Scene : nullptr;
 		Data->ReadbackDelay = 3;
 		Data->CurveCount = 0;
 		Data->PointCount = 0;
diff --git a/Engine/Plugins/Runtime/HairStrands/Source/HairStrandsCore/Private/Niagara/NiagaraDataInterfaceHairStrands.cpp b/Engine/Plugins/Runtime/HairStrands/Source/HairStrandsCore/Private/Niagara/NiagaraDataInterfaceHairStrands.cpp
--- a/Engine/Plugins/Runtime/HairStrands/Source/HairStrandsCore/Private/Niagara/NiagaraDataInterfaceHairStrands.cpp
+++ b/Engine/Plugins/Runtime/HairStrands/Source/HairStrandsCore/Private/Niagara/NiagaraDataInterfaceHairStrands.cpp
@@ -343,7 +343,8 @@ static FHairGroupInstanceRDG Convert(FRDGBuilder& GraphBuilder, const TRefCountP
 			Out.PreviousDeformedPositionOffset = PrevDeformedPositionOffset.Buffer;
 		}
 
-		const bool bHasSkinnedInterpolation = In->Guides.DeformedRootResource && In->Guides.DeformedRootResource->IsValid(MeshLODIndex);
+		const bool bHasSkinnedInterpolation = In->Guides.DeformedRootResource && In->Guides.RestRootResource && 
+			In->Guides.DeformedRootResource->IsValid(MeshLODIndex) && In->Guides.RestRootResource->IsValid(MeshLODIndex);
 		if (bHasSkinnedInterpolation)
 		{
 			if (const FHairStrandsLODRestRootResource* RootResource = In->Guides.RestRootResource->GetLOD(MeshLODIndex))
