Fix typo causing inverted CPU/GPU memory report for groom binding.

#Synced-CL 49808724

diff --git a/Engine/Plugins/Runtime/HairStrands/Source/HairStrandsCore/Private/GroomAsset.cpp b/Engine/Plugins/Runtime/HairStrands/Source/HairStrandsCore/Private/GroomAsset.cpp
--- a/Engine/Plugins/Runtime/HairStrands/Source/HairStrandsCore/Private/GroomAsset.cpp
+++ b/Engine/Plugins/Runtime/HairStrands/Source/HairStrandsCore/Private/GroomAsset.cpp
@@ -207,20 +207,20 @@ void FGroomAssetMemoryStats::Accumulate(const FGroomAssetMemoryStats& In)
 uint32 FGroomAssetMemoryStats::GetTotalCPUSize() const
 {
 	uint32 Out = 0;
-	Out += GPU.Guides ;
-	Out += GPU.Strands;
-	Out += GPU.Cards  ;
-	Out += GPU.Meshes ;
+	Out += CPU.Guides ;
+	Out += CPU.Strands;
+	Out += CPU.Cards  ;
+	Out += CPU.Meshes ;
 	return Out;
 }
 
 uint32 FGroomAssetMemoryStats::GetTotalGPUSize() const
 {
 	uint32 Out = 0;
-	Out += CPU.Guides ;
-	Out += CPU.Strands;
-	Out += CPU.Cards  ;
-	Out += CPU.Meshes ;
+	Out += GPU.Guides ;
+	Out += GPU.Strands;
+	Out += GPU.Cards  ;
+	Out += GPU.Meshes ;
 	return Out;
 }
 
@@ -333,18 +333,18 @@ void FGroomBindingAssetMemoryStats::Accumulate(const FGroomBindingAssetMemorySta
 uint32 FGroomBindingAssetMemoryStats::GetTotalCPUSize() const
 {
 	uint32 Out = 0;
-	Out += GPU.Guides ;
-	Out += GPU.Strands;
-	Out += GPU.Cards  ;
+	Out += CPU.Guides ;
+	Out += CPU.Strands;
+	Out += CPU.Cards  ;
 	return Out;
 }
 
 uint32 FGroomBindingAssetMemoryStats::GetTotalGPUSize() const
 {
 	uint32 Out = 0;
-	Out += CPU.Guides ;
-	Out += CPU.Strands;
-	Out += CPU.Cards  ;
+	Out += GPU.Guides ;
+	Out += GPU.Strands;
+	Out += GPU.Cards  ;
 	return Out;
 }
 
