Fix hair mesh data fetching from DDC.

The data was query twice. Once with the new API, once with the old API.

#Synced-CL 49808724

diff --git a/Engine/Plugins/Runtime/HairStrands/Source/HairStrandsCore/Private/GroomAsset.cpp b/Engine/Plugins/Runtime/HairStrands/Source/HairStrandsCore/Private/GroomAsset.cpp
--- a/Engine/Plugins/Runtime/HairStrands/Source/HairStrandsCore/Private/GroomAsset.cpp
+++ b/Engine/Plugins/Runtime/HairStrands/Source/HairStrandsCore/Private/GroomAsset.cpp
@@ -2941,8 +2941,7 @@ bool UGroomAsset::CacheMeshesData(uint32 GroupIndex)
 	// * If data is available, populate LOD data
 	// * Otherwise build the LOD data
 	FHairGroupPlatformData& HairGroupData = GetHairGroupsPlatformData()[GroupIndex];
-	TArray<uint8> DerivedData;
-	if (GetDerivedDataCacheRef().GetSynchronous(*DerivedDataKey, DerivedData, GetPathName()))
+	if (Data)
 	{
 		UE_CLOG(IsHairStrandsDDCLogEnable(), LogHairStrands, Log, TEXT("[Groom/DDC] Meshes - Found (Groom:%s Group:%d)."), *GetName(), GroupIndex);
 
