Fix invalid change tracking in groom card and groom mesh.

#Synced-CL 49808724

diff --git a/Engine/Plugins/Runtime/HairStrands/Source/HairStrandsCore/Private/GroomAssetCards.cpp b/Engine/Plugins/Runtime/HairStrands/Source/HairStrandsCore/Private/GroomAssetCards.cpp
--- a/Engine/Plugins/Runtime/HairStrands/Source/HairStrandsCore/Private/GroomAssetCards.cpp
+++ b/Engine/Plugins/Runtime/HairStrands/Source/HairStrandsCore/Private/GroomAssetCards.cpp
@@ -26,7 +26,7 @@ void FHairGroupCardsTextures::SetLayout(EHairTextureLayout InLayout)
 
 void FHairGroupCardsTextures::SetTexture(int32 SlotIdx, UTexture2D* Texture)
 {
-	check(SlotIdx < Textures.Num());
+	check(Textures.IsValidIndex(SlotIdx));
 	Textures[SlotIdx] = Texture;
 }
 
@@ -60,7 +60,7 @@ bool FHairGroupsCardsSourceDescription::HasMeshChanged() const
 #if WITH_EDITORONLY_DATA
 	if (ImportedMesh)
 	{
-		return ImportedMeshKey == GetMeshKey();
+		return ImportedMeshKey != GetMeshKey();
 	}
 #endif
 	return false;
diff --git a/Engine/Plugins/Runtime/HairStrands/Source/HairStrandsCore/Private/GroomAssetInterpolation.cpp b/Engine/Plugins/Runtime/HairStrands/Source/HairStrandsCore/Private/GroomAssetInterpolation.cpp
--- a/Engine/Plugins/Runtime/HairStrands/Source/HairStrandsCore/Private/GroomAssetInterpolation.cpp
+++ b/Engine/Plugins/Runtime/HairStrands/Source/HairStrandsCore/Private/GroomAssetInterpolation.cpp
@@ -44,7 +44,11 @@ bool FHairLODSettings::operator==(const FHairLODSettings& A) const
 		CurveDecimation == A.CurveDecimation &&
 		ThicknessScale == A.ThicknessScale &&
 		bVisible == A.bVisible &&
-		ScreenSize == A.ScreenSize;
+		ScreenSize == A.ScreenSize &&
+		GeometryType == A.GeometryType &&
+		BindingType == A.BindingType &&
+		Simulation == A.Simulation &&
+		GlobalInterpolation == A.GlobalInterpolation;
 }
 
 bool FHairGroupsLOD::operator==(const FHairGroupsLOD& A) const
diff --git a/Engine/Plugins/Runtime/HairStrands/Source/HairStrandsCore/Private/GroomAssetMeshes.cpp b/Engine/Plugins/Runtime/HairStrands/Source/HairStrandsCore/Private/GroomAssetMeshes.cpp
--- a/Engine/Plugins/Runtime/HairStrands/Source/HairStrandsCore/Private/GroomAssetMeshes.cpp
+++ b/Engine/Plugins/Runtime/HairStrands/Source/HairStrandsCore/Private/GroomAssetMeshes.cpp
@@ -44,7 +44,7 @@ bool FHairGroupsMeshesSourceDescription::HasMeshChanged() const
 #if WITH_EDITORONLY_DATA
 	if (ImportedMesh)
 	{
-		return ImportedMeshKey == GetMeshKey();
+		return ImportedMeshKey != GetMeshKey();
 	}
 #endif
 	return false;
