Fix memory leak when iteratering on groom binding with hair cards.

Hair cards root resources were not released correclty.

#Synced-CL 49835472

#Synced-CL 49835472

diff --git a/Engine/Plugins/Runtime/HairStrands/Source/HairStrandsCore/Private/GroomBindingAsset.cpp b/Engine/Plugins/Runtime/HairStrands/Source/HairStrandsCore/Private/GroomBindingAsset.cpp
--- a/Engine/Plugins/Runtime/HairStrands/Source/HairStrandsCore/Private/GroomBindingAsset.cpp
+++ b/Engine/Plugins/Runtime/HairStrands/Source/HairStrandsCore/Private/GroomBindingAsset.cpp
@@ -431,8 +431,9 @@ void UGroomBindingAsset::ReleaseResource(bool bResetLoadedSize)
 	{
 		FHairStrandsRestRootResource* InSimRootResources = ResourceToDelete.SimRootResources;
 		FHairStrandsRestRootResource* InRenRootResources = ResourceToDelete.RenRootResources;
+		TArray<FHairStrandsRestRootResource*> InCardsRootResources = ResourceToDelete.CardsRootResources;
 		ENQUEUE_RENDER_COMMAND(ReleaseHairStrandsResourceCommand)(UE::RenderCommandPipe::Groom,
-			[InSimRootResources, InRenRootResources](FRHICommandList& RHICmdList)
+			[InSimRootResources, InRenRootResources, InCardsRootResources](FRHICommandList& RHICmdList)
 		{
 			if (InSimRootResources)
 			{
@@ -444,8 +445,15 @@ void UGroomBindingAsset::ReleaseResource(bool bResetLoadedSize)
 				InRenRootResources->ReleaseResource();
 				delete InRenRootResources;
 			}
+			for (FHairStrandsRestRootResource* InCardRootResource : InCardsRootResources)
+			{
+				if (InCardRootResource)
+				{
+					InCardRootResource->ReleaseResource();
+					delete InCardRootResource;
+				}
+			}
 		});
-		//	#hair_todo: release cards root resources
 	}
 
 	// Process data to be deleted (should happen only in editor)
