Use BitCast for reinterpreting bits instead of manual casting.

#Synced-CL 49835472

diff --git a/Engine/Plugins/Runtime/HairStrands/Source/HairStrandsCore/Private/GroomBindingBuilder.cpp b/Engine/Plugins/Runtime/HairStrands/Source/HairStrandsCore/Private/GroomBindingBuilder.cpp
--- a/Engine/Plugins/Runtime/HairStrands/Source/HairStrandsCore/Private/GroomBindingBuilder.cpp
+++ b/Engine/Plugins/Runtime/HairStrands/Source/HairStrandsCore/Private/GroomBindingBuilder.cpp
@@ -134,7 +134,7 @@ uint32 FHairStrandsRootUtils::PackUVs(const FVector2f& UV)
 float FHairStrandsRootUtils::PackUVsToFloat(const FVector2f& UV)
 {
 	uint32 Encoded = PackUVs(UV);
-	return *((float*)(&Encoded));
+	return BitCast<float>(Encoded);
 }
 
 static float PackNormalToFloat(const FVector3f& InN)
@@ -146,7 +146,7 @@ static float PackNormalToFloat(const FVector3f& InN)
 		(uint32(FMath::Clamp(NN.X, 0.f, 1.f) * 1023.0f) & 0x3FF)    |
 		(uint32(FMath::Clamp(NN.Y, 0.f, 1.f) * 1023.0f) & 0x3FF)<<10|
 		(uint32(FMath::Clamp(NN.Z, 0.f, 1.f) * 1023.0f) & 0x3FF)<<20;
-	return *((float*)(&Encoded));
+	return BitCast<float>(Encoded);
 }
 
 //////////////////////////////////////////////////////////////////////////
