Fix unclear transient card data.

#Synced-CL 49835472

diff --git a/Engine/Plugins/Runtime/HairStrands/Source/HairStrandsCore/Public/HairCardsDatas.h b/Engine/Plugins/Runtime/HairStrands/Source/HairStrandsCore/Public/HairCardsDatas.h
--- a/Engine/Plugins/Runtime/HairStrands/Source/HairStrandsCore/Public/HairCardsDatas.h
+++ b/Engine/Plugins/Runtime/HairStrands/Source/HairStrandsCore/Public/HairCardsDatas.h
@@ -100,11 +100,6 @@ struct FHairCardsGeometry
 	TArray<FVector2f>LocalUVs;// Transient data storing [0..1, 0..1] local UV of the cards. U is aligned with the card's main direction, and V is aligned perpendicularity
 	TArray<FLinearColor> VertexColors;
 
-	// Vertex offset and vertex count of each cards geometry
-	// No longer used, kept only for backward compatibility
-	TArray<uint32> PointOffsets;
-	TArray<uint32> PointCounts;
-
 	// Index offset and index count of each cards geometry
 	TArray<uint32> IndexOffsets;
 	TArray<uint32> IndexCounts;
@@ -119,6 +114,7 @@ struct FHairCardsGeometry
 		Positions.Reset();
 		Indices.Reset();
 		CoordU.Reset();
+		LocalUVs.Reset();
 		VertexColors.Reset();
 
 		IndexOffsets.Reset();
