Fix incorrect comment in Groom builder.

#Synced-CL 49835472

diff --git a/Engine/Plugins/Runtime/HairStrands/Source/HairStrandsCore/Public/GroomBuilder.h b/Engine/Plugins/Runtime/HairStrands/Source/HairStrandsCore/Public/GroomBuilder.h
--- a/Engine/Plugins/Runtime/HairStrands/Source/HairStrandsCore/Public/GroomBuilder.h
+++ b/Engine/Plugins/Runtime/HairStrands/Source/HairStrandsCore/Public/GroomBuilder.h
@@ -44,7 +44,7 @@ struct FHairStrandsVoxelData
 	FIntVector Resolution;
 	TArray<FData> Datas;
 
-	// Return the (closest) group index corresponding to position P
+	// Return the (closest) FData corresponding to position P
 	FData GetData(const FVector3f& P) const;
 	bool IsValid() const { return Datas.Num() > 0; }
 };
