Fix typo in RBF deformer causing invalid indexing of guide's point.

#Synced-CL 49835472

diff --git a/Engine/Plugins/Runtime/HairStrands/Source/HairStrandsCore/Private/GroomRBFDeformer.cpp b/Engine/Plugins/Runtime/HairStrands/Source/HairStrandsCore/Private/GroomRBFDeformer.cpp
--- a/Engine/Plugins/Runtime/HairStrands/Source/HairStrandsCore/Private/GroomRBFDeformer.cpp
+++ b/Engine/Plugins/Runtime/HairStrands/Source/HairStrandsCore/Private/GroomRBFDeformer.cpp
@@ -407,7 +407,7 @@ static void ApplyDeformationToGroom(const TArray<FRBFDeformedPositions>& Deforme
 				{
 					// Mark all invalid/trimmed control point as invalid and set their position to last control point to avoid bound computation issue
 					FlattenedIsValid.Add(false);
-					FlattenedDeformedPositions.Add(DeformedPositions[GroupIndex].GuideStrands[GroupInfo.CurrentRenderVertexIndex-1]);
+					FlattenedDeformedPositions.Add(DeformedPositions[GroupIndex].GuideStrands[GroupInfo.CurrentGuideVertexIndex-1]);
 				}
 			}
 		}
@@ -574,7 +574,8 @@ namespace GroomDerivedDataCacheUtils
 template<typename T>
 float SampleTexture(const FUintPoint& InCoord, const uint8 ComponentIndex, const uint8 NumComponent, const FUintPoint& Resolution, const T* InData)
 {
-	const uint32 Index = ComponentIndex + InCoord.X * NumComponent + InCoord.Y * Resolution.X * NumComponent;
+	const FUintPoint Coord(FMath::Min(Resolution.X - 1, InCoord.X), FMath::Min(Resolution.Y - 1, InCoord.Y));
+	const uint32 Index = ComponentIndex + Coord.X * NumComponent + Coord.Y * Resolution.X * NumComponent;
 	return InData[Index];
 }
 
@@ -592,7 +593,7 @@ float SampleTexture(const FVector2f& InUV, const uint8 InComponentIndex, const u
 	const T T00 = SampleTexture<T>(P00, InComponentIndex, InNumComponent, InResolution, TypedData);
 	const T T10 = SampleTexture<T>(P10, InComponentIndex, InNumComponent, InResolution, TypedData);
 	const T T11 = SampleTexture<T>(P11, InComponentIndex, InNumComponent, InResolution, TypedData);
-	const T T01 = SampleTexture<T>(P11, InComponentIndex, InNumComponent, InResolution, TypedData);
+	const T T01 = SampleTexture<T>(P01, InComponentIndex, InNumComponent, InResolution, TypedData);
 
 	const float V00 = float(T00) * InNormalizationScale;
 	const float V10 = float(T10) * InNormalizationScale;
