Fix typo in groom edit preventing to retrieve correct clump data.

#Synced-CL 49835472

diff --git a/Engine/Plugins/Runtime/HairStrands/Source/HairStrandsCore/Private/GroomEdit.cpp b/Engine/Plugins/Runtime/HairStrands/Source/HairStrandsCore/Private/GroomEdit.cpp
--- a/Engine/Plugins/Runtime/HairStrands/Source/HairStrandsCore/Private/GroomEdit.cpp
+++ b/Engine/Plugins/Runtime/HairStrands/Source/HairStrandsCore/Private/GroomEdit.cpp
@@ -140,8 +140,8 @@ static void EditableGroomGroupToHairDescription(FHairDescription& OutHairDescrip
 	bool bHasTripleWeights = false;
 	for (const FEditableHairStrand& Strands : In)
 	{
-		bHasTripleWeights = Strands.GuideIDs[0] != -1 && Strands.GuideIDs[1] != -1 && Strands.GuideIDs[2] != -1;
-		bHasSingleWeights = Strands.GuideIDs[0] != -1;
+		bHasTripleWeights = bHasTripleWeights || (Strands.GuideIDs[0] != -1 && Strands.GuideIDs[1] != -1 && Strands.GuideIDs[2] != -1);
+		bHasSingleWeights = bHasSingleWeights || (Strands.GuideIDs[0] != -1);
 	}
 
 	for (const FEditableHairStrand& Strands : In)
