Fix typo in GroomComponent hair cards rendering code.

#Synced-CL 49835472

diff --git a/Engine/Plugins/Runtime/HairStrands/Source/HairStrandsCore/Private/GroomComponent.cpp b/Engine/Plugins/Runtime/HairStrands/Source/HairStrandsCore/Private/GroomComponent.cpp
--- a/Engine/Plugins/Runtime/HairStrands/Source/HairStrandsCore/Private/GroomComponent.cpp
+++ b/Engine/Plugins/Runtime/HairStrands/Source/HairStrandsCore/Private/GroomComponent.cpp
@@ -962,7 +962,7 @@ public:
 
 			VertexFactory = (FVertexFactory*)Instance->Cards.LODs[IntLODIndex].GetVertexFactory();
 			check(VertexFactory);
-			PrimitiveIdMode = Instance->Meshes.LODs[IntLODIndex].GetVertexFactory()->GetPrimitiveIdMode(FeatureLevel);
+			PrimitiveIdMode = Instance->Cards.LODs[IntLODIndex].GetVertexFactory()->GetPrimitiveIdMode(FeatureLevel);
 			HairVertexCount = Instance->Cards.LODs[IntLODIndex].RestResource->GetPrimitiveCount() * 3;
 			MaxVertexIndex = HairVertexCount;
 			NumPrimitive = HairVertexCount / 3;
