Fix non-deterministic build in groom builder.

#Synced-CL 49835472

#Synced-CL 49835472

diff --git a/Engine/Plugins/Runtime/HairStrands/Source/HairStrandsCore/Private/GroomBuilder.cpp b/Engine/Plugins/Runtime/HairStrands/Source/HairStrandsCore/Private/GroomBuilder.cpp
--- a/Engine/Plugins/Runtime/HairStrands/Source/HairStrandsCore/Private/GroomBuilder.cpp
+++ b/Engine/Plugins/Runtime/HairStrands/Source/HairStrandsCore/Private/GroomBuilder.cpp
@@ -364,6 +364,7 @@ namespace HairStrandsBuilder
 				{
 					uint32* Data = (uint32*)AttributeColor.GetData();
 					uint32& Packed = Data[PointIndex + IndexOffset];
+					Packed = 0;
 					Packed |= FMath::Clamp(uint32(FMath::Sqrt(Points.PointsBaseColor[PointIndex + IndexOffset].R) * 0x7FF), 0u, 0x7FFu);
 					Packed |= FMath::Clamp(uint32(FMath::Sqrt(Points.PointsBaseColor[PointIndex + IndexOffset].G) * 0x7FF), 0u, 0x7FFu)<<11;
 					Packed |= FMath::Clamp(uint32(FMath::Sqrt(Points.PointsBaseColor[PointIndex + IndexOffset].B) * 0x3FF), 0u, 0x3FFu)<<22;
@@ -441,6 +442,7 @@ namespace HairStrandsBuilder
 				if (bHasMultipleClumpIDs)
 				{
 					uint64* Packed = (uint64*)AttributeClumpIDs.GetData();
+					Packed[CurveIndex]  = 0;
 					Packed[CurveIndex] |= uint64(FMath::Clamp(Curves.ClumpIDs[CurveIndex].X, 0, 65535));
 					Packed[CurveIndex] |= uint64(FMath::Clamp(Curves.ClumpIDs[CurveIndex].Y, 0, 65535))<<16;
 					Packed[CurveIndex] |= uint64(FMath::Clamp(Curves.ClumpIDs[CurveIndex].Z, 0, 65535))<<32;
