Fix typo in attribute name in GroomBindingAsset.

#Synced-CL 49835472

diff --git a/Engine/Plugins/Runtime/HairStrands/Source/HairStrandsCore/Private/GroomBindingAsset.cpp b/Engine/Plugins/Runtime/HairStrands/Source/HairStrandsCore/Private/GroomBindingAsset.cpp
--- a/Engine/Plugins/Runtime/HairStrands/Source/HairStrandsCore/Private/GroomBindingAsset.cpp
+++ b/Engine/Plugins/Runtime/HairStrands/Source/HairStrandsCore/Private/GroomBindingAsset.cpp
@@ -1919,7 +1919,7 @@ bool UGroomBindingAsset::RemoveHairGroupResourcesToDelete(FHairGroupResource& Ou
 FName UGroomBindingAsset::GetHairGroupResourcesMemberName()
 {
 	PRAGMA_DISABLE_DEPRECATION_WARNINGS
-	return GET_MEMBER_NAME_CHECKED(UGroomBindingAsset, GroupInfos);
+	return GET_MEMBER_NAME_CHECKED(UGroomBindingAsset, HairGroupResources);
 	PRAGMA_ENABLE_DEPRECATION_WARNINGS
 }
 
