* Fix incorrect hair cards builder logic for finding closest guide point.
* Fix incorrect logic for rebuilding mesh desription normal/tangent prior to build hair card data.

#Synced-CL 49835472

diff --git a/Engine/Plugins/Runtime/HairStrands/Source/HairStrandsCore/Private/HairCardsBuilder.cpp b/Engine/Plugins/Runtime/HairStrands/Source/HairStrandsCore/Private/HairCardsBuilder.cpp
--- a/Engine/Plugins/Runtime/HairStrands/Source/HairStrandsCore/Private/HairCardsBuilder.cpp
+++ b/Engine/Plugins/Runtime/HairStrands/Source/HairStrandsCore/Private/HairCardsBuilder.cpp
@@ -455,7 +455,7 @@ void SanitizeMeshDescription(FMeshDescription* MeshDescription)
 	bool bHasInvalidNormals = false;
 	bool bHasInvalidTangents = false;
 	FStaticMeshOperations::HasInvalidVertexInstanceNormalsOrTangents(*MeshDescription, bHasInvalidNormals, bHasInvalidTangents);
-	if (!bHasInvalidNormals || !bHasInvalidTangents)
+	if (bHasInvalidNormals || bHasInvalidTangents)
 	{
 		FStaticMeshOperations::ComputeTriangleTangentsAndNormals(*MeshDescription, THRESH_POINTS_ARE_SAME);
 
@@ -1096,14 +1096,17 @@ static bool InternalImportGeometry_WithImportedGuides(
 				const float GuideCoordU0 = InGuides.StrandsPoints.PointsCoordU[I0];
 				const float GuideCoordU1 = InGuides.StrandsPoints.PointsCoordU[I1];
 
+				const FVector3f G01Vector = GPoint1 - GPoint0;
+				const float G01Norm2 = FMath::Max(FVector3f::DotProduct(G01Vector,G01Vector), KINDA_SMALL_NUMBER);
+
 				// For each point of the mesh, compute its distance to the guide's segment
 				for (uint32 PointIt = 0; PointIt < PointCount; ++PointIt)
 				{
 					// Project P onto the segment.
 					const FVector3f P = OutBulk.Positions[PointIt];
-					float U = FVector3f::DotProduct(P-GPoint0, GPoint1-GPoint0);
+					float U = FVector3f::DotProduct(P-GPoint0, G01Vector) / G01Norm2;
 					U = FMath::Clamp(U, 0.f, 1.f);
-					const FVector3f PP = U * (GPoint1-GPoint0) + GPoint0;
+					const FVector3f PP = U * G01Vector + GPoint0;
 					const float Dist = (P-PP).Length();
 
 					// If new projection is closer than the older one, change guides
