Avoid groom resources allocation if a groom instance is not visible and has no simulation.

#Synced-CL 49918411

diff --git a/Engine/Plugins/Runtime/HairStrands/Source/HairStrandsCore/Private/GroomManager.cpp b/Engine/Plugins/Runtime/HairStrands/Source/HairStrandsCore/Private/GroomManager.cpp
--- a/Engine/Plugins/Runtime/HairStrands/Source/HairStrandsCore/Private/GroomManager.cpp
+++ b/Engine/Plugins/Runtime/HairStrands/Source/HairStrandsCore/Private/GroomManager.cpp
@@ -2765,6 +2765,31 @@ static bool SelectValidLOD(
 	const bool bSimulationCacheEnable		= Instance->HairGroupPublicData->bIsSimulationCacheEnable;
 	const bool bLODNeedsGuides				= bSimulationEnable || bDeformationEnable || bGlobalInterpolationEnable || bSimulationCacheEnable;
 
+	// If we're not rendering hair, and we don't need guides work, bail before any allocations.
+	// This prevents hitching for EHairGeometryType::NoneGeometry which isn't accounted for in GetHairResourceLoadingType() and just returns EHairResourceLoadingType::Sync
+	if (GeometryType == EHairGeometryType::NoneGeometry && !bLODNeedsGuides)
+	{
+		const int32 IntPrevHairLODIndex = FMath::Clamp(FMath::FloorToInt(PrevHairLODIndex), 0, HairLODCount - 1);
+		const EHairBindingType CurrBindingType = Instance->HairGroupPublicData->GetBindingType(IntHairLODIndex);
+		const EHairBindingType PrevBindingType = Instance->HairGroupPublicData->GetBindingType(IntPrevHairLODIndex);
+ 
+		Instance->HairGroupPublicData->SetLODVisibility(bIsVisible);
+		Instance->HairGroupPublicData->SetLODIndex(HairLODIndex);
+		Instance->HairGroupPublicData->SetLODBias(0);
+		Instance->HairGroupPublicData->SetMeshLODIndex(MeshLODIndex);
+ 
+		Instance->HairGroupPublicData->VFInput.GeometryType = GeometryType;
+		Instance->HairGroupPublicData->VFInput.BindingType = CurrBindingType;
+		Instance->HairGroupPublicData->VFInput.bHasLODSwitch = IntPrevHairLODIndex != IntHairLODIndex;
+		Instance->HairGroupPublicData->VFInput.bHasLODSwitchBindingType = CurrBindingType != PrevBindingType;
+ 
+		Instance->GeometryType = GeometryType;
+		Instance->BindingType = CurrBindingType;
+
+		// Return true to avoid further retry. While the selected LOD is not renderable, it is 'valid' as 'hidden'
+		return true;
+	}
+
 	const EHairResourceLoadingType LoadingType = GetHairResourceLoadingType(GeometryType, IntHairLODIndex);
 	EHairResourceStatus ResourceStatus;
 	ResourceStatus.Status 				= EHairResourceStatus::EStatus::None;
@@ -2780,7 +2805,7 @@ static bool SelectValidLOD(
 		if (Instance->Guides.DeformedResource)			
 		{ 
 			// Ensure the rest resources are correctly loaded prior to initialized and copy the rest position into the deformed positions
-			if (Instance->Guides.RestResource->bIsInitialized)
+			if (Instance->Guides.RestResource && Instance->Guides.RestResource->bIsInitialized)
 			{
 				const bool bNeedCopy = !Instance->Guides.DeformedResource->bIsInitialized; 
 				Instance->Guides.DeformedResource->Allocate(GraphBuilder, EHairResourceLoadingType::Sync); 
