Add groom binding data validation.

The validation check consistency between the groom asset and groom data used for building the groom binding.

The validation is enabled by default and only effective in the editor.

#Synced-CL 50480066

diff --git a/Engine/Plugins/Runtime/HairStrands/Source/HairStrandsCore/Private/GroomBindingBuilder.cpp b/Engine/Plugins/Runtime/HairStrands/Source/HairStrandsCore/Private/GroomBindingBuilder.cpp
--- a/Engine/Plugins/Runtime/HairStrands/Source/HairStrandsCore/Private/GroomBindingBuilder.cpp
+++ b/Engine/Plugins/Runtime/HairStrands/Source/HairStrandsCore/Private/GroomBindingBuilder.cpp
@@ -92,7 +92,7 @@ FGroomBindingBuilder::FInput::FInput(UGroomBindingAsset* BindingAsset, const ITa
 FString FGroomBindingBuilder::GetVersion()
 {
 	// Important to update the version when groom building changes
-	return TEXT("6d");
+	return TEXT("6d_563D4DE9469C41D590AE5810A85AA71B");
 }
 
 ///////////////////////////////////////////////////////////////////////////////////////////////////
@@ -2476,7 +2476,7 @@ namespace GroomBinding_Transfer
 }
 // namespace GroomBinding_Transfer
 
-static void InitHairStrandsRootData(TArray<FHairStrandsRootData>& Out, uint32 CurveCount, uint32 PointCount, uint32 LODCount, uint32 MinLOD, uint32 NumSamples)
+static void InitHairStrandsRootData(TArray<FHairStrandsRootData>& Out, uint32 CurveCount, uint32 PointCount, uint32 LODCount, uint32 MinLOD, uint32 NumSamples, uint32 InFlags, const FHairStrandsDatas::FHash& InHash)
 {
 	Out.SetNum(LODCount);
 	uint32 LODIndex = 0;
@@ -2498,6 +2498,8 @@ static void InitHairStrandsRootData(TArray<FHairStrandsRootData>& Out, uint32 Cu
 		OutLOD.MeshInterpolationWeightsBuffer.Empty();
 		OutLOD.MeshSampleIndicesBuffer.Empty();
 		OutLOD.RestSamplePositionsBuffer.Empty();
+		OutLOD.Flags = InFlags;
+		OutLOD.Hash = InHash;
 		OutLOD.LODIndex = LODIndex++;
 	}
 }
@@ -2577,6 +2579,8 @@ static void BuildRootBulkData(
 		Out.Header.UniqueTriangleCount 	= In.UniqueTriangleIndexBuffer.Num();
 		Out.Header.UniqueSectionIndices	= In.UniqueSectionIds;
 		Out.Header.MeshSectionCount		= In.MeshSectionCount;
+		Out.Header.Flags				= FHairStrandsBulkData::ConvertDataFlagsFromHairStrandsFlags(In.Flags);
+		Out.Header.Hash					= In.Hash;
 	}
 
 	// Data
@@ -2728,6 +2732,37 @@ static void BuildGroupPlatformData(
 
 } // namespace GroomBinding_BulkCopy
 
+///////////////////////////////////////////////////////////////////////////////////////////////////
+// Helper functions for binding validation code
+
+bool GetHairStrandsBindingDataValidation();
+
+#if WITH_EDITORONLY_DATA
+
+static FHairStrandsCommonDatas::FHash GetHash(const FHairStrandsDatas& In)
+{
+	return In.GetHash();
+}
+
+static uint32 GetFlags(const FHairStrandsDatas& In)
+{
+	return In.Flags;
+}
+
+#else
+
+static FHairStrandsCommonDatas::FHash GetHash(const FHairStrandsBulkData& In)
+{
+	return FHairStrandsCommonDatas::FHash();
+}
+
+static uint32 GetFlags(const FHairStrandsBulkData& In)
+{
+	return FHairStrandsBulkData::ConvertDataFlagsToHairStrandsFlags(In.Header.Flags);
+}
+
+#endif
+
 ///////////////////////////////////////////////////////////////////////////////////////////////////
 // Main entry (CPU path)
 bool BuildHairRootGroupData(const FGroomBindingBuilder::FInput& In, uint32 InGroupIndex, const ITargetPlatform* TargetPlatform, FHairRootGroupData& OutData)
@@ -2928,14 +2963,14 @@ bool BuildHairRootGroupData(const FGroomBindingBuilder::FInput& In, uint32 InGro
 			// Guides
 			const uint32 GuidesCurveCount = GuidesDataSource.GetNumCurves();
 			const uint32 GuidesPointCount = GuidesDataSource.GetNumPoints();
-			InitHairStrandsRootData(OutData.SimRootDatas, GuidesCurveCount, GuidesPointCount, MeshLODCount, In.TargetMeshMinLOD, In.NumInterpolationPoints);
+			InitHairStrandsRootData(OutData.SimRootDatas, GuidesCurveCount, GuidesPointCount, MeshLODCount, In.TargetMeshMinLOD, In.NumInterpolationPoints, GetFlags(GuidesDataSource), GetHash(GuidesDataSource));
 
 			// Strands
 			if (bNeedStrandsRoot)
 			{
 				const uint32 StrandsCurveCount = StrandsDataSource.GetNumCurves();
 				const uint32 StrandsPointCount = StrandsDataSource.GetNumPoints();
-				InitHairStrandsRootData(OutData.RenRootDatas, StrandsCurveCount, StrandsPointCount, MeshLODCount, In.TargetMeshMinLOD, In.NumInterpolationPoints);
+				InitHairStrandsRootData(OutData.RenRootDatas, StrandsCurveCount, StrandsPointCount, MeshLODCount, In.TargetMeshMinLOD, In.NumInterpolationPoints, GetFlags(StrandsDataSource), GetHash(StrandsDataSource));
 			}
 
 			// Cards
@@ -2949,19 +2984,15 @@ bool BuildHairRootGroupData(const FGroomBindingBuilder::FInput& In, uint32 InGro
 					FHairStrandsDatas LODGuidesData;
 					const bool bIsValid = In.GroomAsset->GetHairCardsGuidesDatas(InGroupIndex, CardsLODIt, LODGuidesData);
 #else
+					const FHairStrandsBulkData& LODGuidesData = GroupData.Cards.LODs[CardsLODIt].GuideBulkData;
 					const bool bIsValid = GroupData.Cards.LODs[CardsLODIt].BulkData.IsValid();
 #endif
 					if (bIsValid)
 					{
-#if WITH_EDITORONLY_DATA
 						const uint32 CardsCurveCount = LODGuidesData.GetNumCurves();
 						const uint32 CardsPointCount = LODGuidesData.GetNumPoints(); 
-#else
-						const uint32 CardsCurveCount = GroupData.Cards.LODs[CardsLODIt].GuideBulkData.GetNumCurves();
-						const uint32 CardsPointCount = GroupData.Cards.LODs[CardsLODIt].GuideBulkData.GetNumPoints();
-#endif
 
-						InitHairStrandsRootData(OutData.CardsRootDatas[CardsLODIt], CardsCurveCount, CardsPointCount, MeshLODCount, In.TargetMeshMinLOD, In.NumInterpolationPoints);
+						InitHairStrandsRootData(OutData.CardsRootDatas[CardsLODIt], CardsCurveCount, CardsPointCount, MeshLODCount, In.TargetMeshMinLOD, In.NumInterpolationPoints, GetFlags(LODGuidesData), GetHash(LODGuidesData));
 					}
 				}
 			}
@@ -3016,6 +3047,17 @@ bool BuildHairRootGroupData(const FGroomBindingBuilder::FInput& In, uint32 InGro
 			{
 				const GroomBinding_Mesh::FHairStrandRootPositions GuidesData(GuidesDataSource);
 
+				// Sanity check
+				#if WITH_EDITORONLY_DATA
+				if (GetHairStrandsBindingDataValidation())
+				{
+					for (const FHairStrandsRootData& RootData : OutData.SimRootDatas)
+					{
+						check(RootData.Hash == GuidesDataSource.GetHash());
+					}
+				}
+				#endif
+
 				if (!GroomBinding_RootProjection::Project(
 					GuidesData,
 					TargetMeshData.Get(),
@@ -3034,6 +3076,17 @@ bool BuildHairRootGroupData(const FGroomBindingBuilder::FInput& In, uint32 InGro
 			{
 				const GroomBinding_Mesh::FHairStrandRootPositions StrandsData(StrandsDataSource);
 
+				// Sanity check
+				#if WITH_EDITORONLY_DATA
+				if (GetHairStrandsBindingDataValidation())
+				{
+					for (const FHairStrandsRootData& RootData : OutData.RenRootDatas)
+					{
+						check(RootData.Hash == StrandsDataSource.GetHash());
+					}
+				}
+				#endif
+
 				if (!GroomBinding_RootProjection::Project(
 					StrandsData,
 					TargetMeshData.Get(),
@@ -3056,6 +3109,16 @@ bool BuildHairRootGroupData(const FGroomBindingBuilder::FInput& In, uint32 InGro
 #if WITH_EDITORONLY_DATA
 					FHairStrandsDatas LODGuidesDataSource;
 					const bool bIsValid = In.GroomAsset->GetHairCardsGuidesDatas(InGroupIndex, CardsLODIt, LODGuidesDataSource);
+
+					// Sanity check
+					if (GetHairStrandsBindingDataValidation())
+					{
+						const FHairStrandsCommonDatas::FHash SourceHash = LODGuidesDataSource.GetHash();
+						for (const FHairStrandsRootData& RootData : OutData.CardsRootDatas[CardsLODIt])
+						{
+							check(RootData.Hash == SourceHash);
+						}
+					}
 #else
 					const bool bIsValid = GroupData.Cards.LODs[CardsLODIt].IsValid();
 					const FHairStrandsBulkData& LODGuidesDataSource = GroupData.Cards.LODs[CardsLODIt].GuideBulkData;
diff --git a/Engine/Plugins/Runtime/HairStrands/Source/HairStrandsCore/Private/GroomBuilder.cpp b/Engine/Plugins/Runtime/HairStrands/Source/HairStrandsCore/Private/GroomBuilder.cpp
--- a/Engine/Plugins/Runtime/HairStrands/Source/HairStrandsCore/Private/GroomBuilder.cpp
+++ b/Engine/Plugins/Runtime/HairStrands/Source/HairStrandsCore/Private/GroomBuilder.cpp
@@ -55,7 +55,7 @@ bool DoesHairStrandsSupportCompressedPosition();
 
 FString FGroomBuilder::GetVersion()
 {
-	FString Out = TEXT("v21");
+	FString Out = TEXT("v21_102031A3EE7E4570B3093414EF3956F3");
 	if (IsRenderCurveEnabled())
 	{
 		Out += TEXT("_Curve4a");
@@ -621,6 +621,7 @@ namespace HairStrandsBuilder
 		OutBulkData.Header.ImportedAttributeFlags = HairStrands.GetAttributeFlags();
 		OutBulkData.Header.Transcoding.PositionOffset = FVector3f::ZeroVector;
 		OutBulkData.Header.Transcoding.PositionScale = FVector3f::ZeroVector;
+		OutBulkData.Header.Hash = HairStrands.GetHash();
 
 		// Transfer group trimmed pt./curve info from group to header. This allows to display this information 
 		// in the groom editor by reading the bulk data header
diff --git a/Engine/Plugins/Runtime/HairStrands/Source/HairStrandsCore/Private/GroomManager.cpp b/Engine/Plugins/Runtime/HairStrands/Source/HairStrandsCore/Private/GroomManager.cpp
--- a/Engine/Plugins/Runtime/HairStrands/Source/HairStrandsCore/Private/GroomManager.cpp
+++ b/Engine/Plugins/Runtime/HairStrands/Source/HairStrandsCore/Private/GroomManager.cpp
@@ -1544,6 +1544,9 @@ static void RunHairStrandsInterpolation_Strands(
 					InstanceData.Instance,
 					InstanceData.ActivePointCount,
 					InstanceData.ActiveCurveCount,
+					InstanceData.Instance->Strands.RestResource,
+					InstanceData.Instance->Guides.RestResource,
+					bValidGuide,
 					InstanceData.Instance->Strands.RestResource->BulkData.Header.MaxPointPerCurve,
 					InstanceData.MeshLODIndex,
 					InstanceData.Instance->Strands.Modifier.HairLengthScale,
@@ -2124,6 +2127,9 @@ static void RunHairStrandsInterpolation_Cards(
 					InstanceData.Instance,
 					LOD.Guides.RestResource->GetPointCount(),
 					LOD.Guides.RestResource->GetCurveCount(),
+					LOD.Guides.RestResource,
+					InstanceData.Instance->Guides.RestResource,
+					InstanceData.bValidGuide,
 					LOD.Guides.RestResource->BulkData.Header.MaxPointPerCurve,
 					InstanceData.MeshLODIndex,
 					1.0f,
diff --git a/Engine/Plugins/Runtime/HairStrands/Source/HairStrandsCore/Private/HairAttributes.cpp b/Engine/Plugins/Runtime/HairStrands/Source/HairStrandsCore/Private/HairAttributes.cpp
--- a/Engine/Plugins/Runtime/HairStrands/Source/HairStrandsCore/Private/HairAttributes.cpp
+++ b/Engine/Plugins/Runtime/HairStrands/Source/HairStrandsCore/Private/HairAttributes.cpp
@@ -3,6 +3,7 @@
 #include "HairAttributes.h"
 #include "HairStrandsDefinitions.h"
 #include "HairStrandsDatas.h"
+#include "Hash/CityHash.h"
 
 namespace HairAttribute
 {
@@ -181,4 +182,13 @@ void FHairStrandsDatas::CopyCurve(const FHairStrandsDatas& In, FHairStrandsDatas
 	{
 		Out.StrandsCurves.CurvesMapping[OutIndex] = In.StrandsCurves.CurvesMapping.IsValidIndex(InIndex) ? In.StrandsCurves.CurvesMapping[InIndex] : InIndex;
 	}
+}
+
+FHairStrandsDatas::FHash FHairStrandsCommonDatas::GetHash() const
+{
+	FHairStrandsDatas::FHash Out;
+	Out.CurveCountHash 	= CityHash64WithSeed((const char*)StrandsCurves.CurvesCount.GetData(), StrandsCurves.CurvesCount.Num() * sizeof(uint16), Out.CurveCountHash);
+	Out.CurveOffsetHash = CityHash64WithSeed((const char*)StrandsCurves.CurvesOffset.GetData(), StrandsCurves.CurvesOffset.Num() * sizeof(uint32), Out.CurveOffsetHash);
+	Out.PositionHash 	= CityHash64WithSeed((const char*)StrandsPoints.PointsPosition.GetData(), StrandsPoints.PointsPosition.Num() * sizeof(FVector3f), Out.PositionHash);
+	return Out;
 }
\ No newline at end of file
diff --git a/Engine/Plugins/Runtime/HairStrands/Source/HairStrandsCore/Private/HairCardsBuilder.cpp b/Engine/Plugins/Runtime/HairStrands/Source/HairStrandsCore/Private/HairCardsBuilder.cpp
--- a/Engine/Plugins/Runtime/HairStrands/Source/HairStrandsCore/Private/HairCardsBuilder.cpp
+++ b/Engine/Plugins/Runtime/HairStrands/Source/HairStrandsCore/Private/HairCardsBuilder.cpp
@@ -361,6 +361,7 @@ static bool InternalCreateCardsGuides(
 		OutGuides.StrandsCurves.CurvesCount[CardIt] = PointCount;
 		OutGuides.StrandsCurves.CurvesOffset[CardIt] = PointOffset;
 		OutGuides.StrandsCurves.CurvesLength[CardIt] = TotalLength;
+		OutGuides.Flags = 0u; // Card's guide don't have any particular flag (no shuffling, no extra point, no decimated points/curves, ...)
 		if (OutGuides.StrandsCurves.CurvesRootUV.Num()) OutGuides.StrandsCurves.CurvesRootUV[CardIt] = FVector2f(0, 0);
 
 		float CurrentLength = 0;
diff --git a/Engine/Plugins/Runtime/HairStrands/Source/HairStrandsCore/Private/HairStrandsDatas.cpp b/Engine/Plugins/Runtime/HairStrands/Source/HairStrandsCore/Private/HairStrandsDatas.cpp
--- a/Engine/Plugins/Runtime/HairStrands/Source/HairStrandsCore/Private/HairStrandsDatas.cpp
+++ b/Engine/Plugins/Runtime/HairStrands/Source/HairStrandsCore/Private/HairStrandsDatas.cpp
@@ -585,6 +585,26 @@ void FHairStreamingRequest::FChunk::Release()
 /////////////////////////////////////////////////////////////////////////////////////////
 // Rest bulk data
 
+uint32 FHairStrandsBulkData::ConvertDataFlagsToHairStrandsFlags(uint32 In)
+{
+	uint32 Out = 0;
+	Out |= !!(In & EDataFlags::DataFlags_HasTrimmedPoint) 		? uint32(EHairStrandsFlags::HasTrimmedPoint)	: 0u;
+	Out |= !!(In & EDataFlags::DataFlags_HasTrimmedCurve) 		? uint32(EHairStrandsFlags::HasTrimmedCurve)	: 0u;
+	Out |= !!(In & EDataFlags::DataFlags_HasShuffledCurves) 	? uint32(EHairStrandsFlags::HasShuffledCurve)	: 0u;
+	Out |= !!(In & EDataFlags::DataFlags_HasExtraEndingPoint) 	? uint32(EHairStrandsFlags::HasExtraEndingPoint): 0u;
+	return Out;
+}
+
+uint32 FHairStrandsBulkData::ConvertDataFlagsFromHairStrandsFlags(uint32 In)
+{
+	uint32 Out = 0;
+	Out |= !!(In & uint32(EHairStrandsFlags::HasTrimmedPoint)) 		? EDataFlags::DataFlags_HasTrimmedPoint		: 0u;
+	Out |= !!(In & uint32(EHairStrandsFlags::HasTrimmedCurve)) 		? EDataFlags::DataFlags_HasTrimmedCurve		: 0u;
+	Out |= !!(In & uint32(EHairStrandsFlags::HasShuffledCurve)) 	? EDataFlags::DataFlags_HasShuffledCurves	: 0u;
+	Out |= !!(In & uint32(EHairStrandsFlags::HasExtraEndingPoint)) 	? EDataFlags::DataFlags_HasExtraEndingPoint	: 0u;
+	return Out;
+}
+
 void FHairStrandsBulkData::SerializeHeader(FArchive& Ar, UObject* Owner)
 {
 	GetResourceVersion(Ar);
@@ -598,6 +618,9 @@ void FHairStrandsBulkData::SerializeHeader(FArchive& Ar, UObject* Owner)
 	Ar << Header.MaxPointPerCurve;
 	Ar << Header.AvgPointPerCurve;
 	Ar << Header.Flags;
+	Ar << Header.Hash.PositionHash;
+	Ar << Header.Hash.CurveCountHash;
+	Ar << Header.Hash.CurveOffsetHash;
 	for (uint8 AttributeIt = 0; AttributeIt < HAIR_CURVE_ATTRIBUTE_COUNT; ++AttributeIt)
 	{
 		Ar << Header.CurveAttributeOffsets[AttributeIt];
@@ -1020,6 +1043,10 @@ void FHairStrandsRootBulkData::SerializeHeader(FArchive& Ar, UObject* Owner)
 		Ar << Header.SampleCount;
 		Ar << Header.UniqueSectionIndices;
 		Ar << Header.MeshSectionCount;
+		Ar << Header.Flags;
+		Ar << Header.Hash.PositionHash;
+		Ar << Header.Hash.CurveCountHash;
+		Ar << Header.Hash.CurveOffsetHash;
 	}
 
 	// Strides
diff --git a/Engine/Plugins/Runtime/HairStrands/Source/HairStrandsCore/Private/HairStrandsInterpolation.cpp b/Engine/Plugins/Runtime/HairStrands/Source/HairStrandsCore/Private/HairStrandsInterpolation.cpp
--- a/Engine/Plugins/Runtime/HairStrands/Source/HairStrandsCore/Private/HairStrandsInterpolation.cpp
+++ b/Engine/Plugins/Runtime/HairStrands/Source/HairStrandsCore/Private/HairStrandsInterpolation.cpp
@@ -64,8 +64,20 @@ static FAutoConsoleVariableRef CVarHairStrands_ScaleHairRadiusByTransform(TEXT("
 static int32 GVarHairStrands_AddDeformerStrandsDisplacement = 0;
 static FAutoConsoleVariableRef CVarHairStrands_AddDeformerStrandsDisplacement(TEXT("r.HairStrands.Strands.AddDeformerStrandsDisplacement"), GVarHairStrands_AddDeformerStrandsDisplacement, TEXT("Add the deformer displacement to the current position. If set to 0 the deformer position will overwrite the deformed one"));
 
+static int32 GVarHairStrands_BindingDataValidation = 1;
+static FAutoConsoleVariableRef CVarVarHairStrands_BindingDataValidation(TEXT("r.HairStrands.Strands.BindingDataValidation"), GVarHairStrands_BindingDataValidation, TEXT("Add binding data validation at runtime (for debugging purpose)"));
+
 ////////////////////////////////////////////////////////////////////////////////////////////////////////////////////////
 
+bool GetHairStrandsBindingDataValidation()
+{
+	#if WITH_EDITORONLY_DATA
+	return GVarHairStrands_BindingDataValidation > 0;
+	#else
+	return false;
+	#endif
+}
+
 uint32 GetHairStrandsForceSingleGuideInterpolation()
 {
 	return GStrandHairInterpolationForceSingleGuide > 0;
@@ -611,6 +623,9 @@ void AddHairStrandsInterpolationPass(
 	const FHairGroupInstance* Instance,
 	const uint32 VertexCount,
 	const uint32 CurveCount,
+	const FHairStrandsRestResource* RenRestResource,
+	const FHairStrandsRestResource* SimRestResource,
+	const bool bSimDataValid,
 	const uint32 MaxPointPerCurve,
 	const int32 MeshLODIndex,
 	const float HairLengthScale,
@@ -694,6 +709,11 @@ void AddHairStrandsInterpolationPass(
 		{
 			Parameters->SimRootDeformedPositionBuffer = RegisterAsSRV(GraphBuilder, Sim_DeformedLODDatas.GetDeformedUniqueTrianglePositionBuffer(DeformedFrame));
 		}
+
+		if (bSimDataValid && GetHairStrandsBindingDataValidation())
+		{
+			check(SimRestResource->BulkData.Header.Hash == Sim_RestLODDatas.GetBulkData().Header.Hash);
+		}
 	}
 
 	// Strands
@@ -709,6 +729,11 @@ void AddHairStrandsInterpolationPass(
 		{
 			Parameters->RenRootDeformedPositionBuffer = RegisterAsSRV(GraphBuilder, Ren_DeformedLODDatas.GetDeformedUniqueTrianglePositionBuffer(DeformedFrame));
 		}
+
+		if (GetHairStrandsBindingDataValidation())
+		{
+			check(RenRestResource->BulkData.Header.Hash == Ren_RestLODDatas.GetBulkData().Header.Hash);
+		}
 	}
 
 	const bool bSupportDeformer = RenDeformerPositionBuffer != nullptr;
diff --git a/Engine/Plugins/Runtime/HairStrands/Source/HairStrandsCore/Public/HairStrandsDatas.h b/Engine/Plugins/Runtime/HairStrands/Source/HairStrandsCore/Public/HairStrandsDatas.h
--- a/Engine/Plugins/Runtime/HairStrands/Source/HairStrandsCore/Public/HairStrandsDatas.h
+++ b/Engine/Plugins/Runtime/HairStrands/Source/HairStrandsCore/Public/HairStrandsDatas.h
@@ -600,6 +600,21 @@ struct FHairStrandsCurves
 /** Hair strands datas that are stored on CPU */
 struct FHairStrandsCommonDatas
 {
+	struct FHash
+	{
+		uint64 PositionHash = 0;
+		uint64 CurveCountHash = 0;
+		uint64 CurveOffsetHash = 0;
+
+		bool operator==(const FHash& In) const
+		{
+			return 
+				In.PositionHash == PositionHash &&
+				In.CurveCountHash == CurveCountHash &&
+				In.CurveOffsetHash == CurveOffsetHash;
+		}
+	};
+
 	/* Get the total number of points */
 	uint32 GetNumPoints() const { return StrandsPoints.Num(); }
 
@@ -614,6 +629,8 @@ struct FHairStrandsCommonDatas
 
 	bool IsValid() const { return StrandsCurves.Num() > 0 && StrandsPoints.Num() > 0; }
 
+	FHash GetHash() const;
+
 	/** List of all the strands points */
 	FHairStrandsPoints StrandsPoints;
 
@@ -663,6 +680,9 @@ struct FHairStrandsBulkData : FHairStrandsBulkCommon
 		DataFlags_HasExtraEndingPoint = 0x100u	// Curve contains an extra ending control point
 	};
 
+	static uint32 ConvertDataFlagsToHairStrandsFlags(uint32 In);
+	static uint32 ConvertDataFlagsFromHairStrandsFlags(uint32 In);
+
 	UE_API virtual void SerializeHeader(FArchive& Ar, UObject* Owner) override;
 	UE_API virtual uint32 GetResourceCount() const override;
 	UE_API virtual void GetResources(FQuery& Out) override;
@@ -720,6 +740,8 @@ struct FHairStrandsBulkData : FHairStrandsBulkCommon
 		// Coverage scale
 		TArray<float> CoverageScales;
 
+		FHairStrandsCommonDatas::FHash Hash;
+
 		// Data transcoding parameters
 		struct FTranscoding
 		{
@@ -954,6 +976,12 @@ struct FHairStrandsRootData
 
 	/* Number of control points */
 	uint32 PointCount = 0;
+
+	/** General flags (shuffled curve, extra ending point, ...) */
+	uint32 Flags = 0;
+
+	/** Source data's hash for validation */
+	FHairStrandsDatas::FHash Hash;
 };
 
 /* Bulk data for root resources (GPU resources are stored into FHairStrandsRootResources) */
@@ -973,6 +1001,7 @@ struct FHairStrandsRootBulkData : FHairStrandsBulkCommon
 #if !WITH_EDITORONLY_DATA
 	bool IsCookedBulkDataLoadedForCPUAccess() const;
 #endif
+	bool HasFlags(FHairStrandsBulkData::EDataFlags In) const { return !!(Header.Flags & In); }
 
 	struct FHeader
 	{		
@@ -982,6 +1011,8 @@ struct FHairStrandsRootBulkData : FHairStrandsBulkCommon
 		uint32 SampleCount = 0; 				// Number of samples used for the mesh interpolation
 		uint32 MeshSectionCount = 0;			// Number of section of the target mesh
 		uint32 UniqueTriangleCount = 0;			
+		uint32 Flags = 0;
+		FHairStrandsDatas::FHash Hash;
 		TArray<uint32> UniqueSectionIndices; 	// Store the mesh section indices which are relevant for this root LOD data
 
 		struct FStrides
diff --git a/Engine/Plugins/Runtime/HairStrands/Source/HairStrandsCore/Public/HairStrandsInterpolation.h b/Engine/Plugins/Runtime/HairStrands/Source/HairStrandsCore/Public/HairStrandsInterpolation.h
--- a/Engine/Plugins/Runtime/HairStrands/Source/HairStrandsCore/Public/HairStrandsInterpolation.h
+++ b/Engine/Plugins/Runtime/HairStrands/Source/HairStrandsCore/Public/HairStrandsInterpolation.h
@@ -119,6 +119,9 @@ void AddHairStrandsInterpolationPass(
 	const FHairGroupInstance* Instance,
 	const uint32 VertexCount,
 	const uint32 CurveCount,
+	const FHairStrandsRestResource* RenRestResource,
+	const FHairStrandsRestResource* SimRestResource,
+	const bool bSimDataValid,
 	const uint32 MaxPointPerCurve,
 	const int32 MeshLODIndex,
 	const float HairLengthScale,
