Remove unecessary check in groom's group index voxelization causing crash.

#Synced-CL 50856398

diff --git a/Engine/Plugins/Runtime/HairStrands/Source/HairStrandsCore/Private/GroomBuilder.cpp b/Engine/Plugins/Runtime/HairStrands/Source/HairStrandsCore/Private/GroomBuilder.cpp
--- a/Engine/Plugins/Runtime/HairStrands/Source/HairStrandsCore/Private/GroomBuilder.cpp
+++ b/Engine/Plugins/Runtime/HairStrands/Source/HairStrandsCore/Private/GroomBuilder.cpp
@@ -3360,7 +3360,6 @@ FHairStrandsVoxelData::FData FHairStrandsVoxelData::GetData(const FVector3f& P)
 	}
 
 	// No valid group index has been found. In this case force the group index to 0
-	check(Out.GroupIndex != FHairStrandsVoxelData::InvalidGroupIndex);
 	if (Out.GroupIndex == FHairStrandsVoxelData::InvalidGroupIndex) { Out.GroupIndex = 0; }
 
 	return Out;
