Fix crash when creating new LOD on an empty groom card.

Within an empty groom cards, the groom asset cache derived data will failed as no hair description is available. When recreating component using this groom (e.g. within the groom editor), the groom component creation will try to access the groom resources, causing an out of bound access.

#Synced-CL 50856398

diff --git a/Engine/Plugins/Runtime/HairStrands/Source/HairStrandsCore/Private/GroomComponent.cpp b/Engine/Plugins/Runtime/HairStrands/Source/HairStrandsCore/Private/GroomComponent.cpp
--- a/Engine/Plugins/Runtime/HairStrands/Source/HairStrandsCore/Private/GroomComponent.cpp
+++ b/Engine/Plugins/Runtime/HairStrands/Source/HairStrandsCore/Private/GroomComponent.cpp
@@ -2700,7 +2700,7 @@ void UGroomComponent::InitResources(bool bIsBindingReloading)
 
 	UpdateHairGroupsDesc();
 
-	if (!GroomAsset || GroomAsset->GetNumHairGroups() == 0)
+	if (!GroomAsset || GroomAsset->GetNumHairGroups() == 0 || !GroomAsset->IsValid())
 	{
 		return;
 	}
