Remove bUseAttachParentBound override on groom component when using groom binding.

#Synced-CL 50856398

diff --git a/Engine/Plugins/Runtime/HairStrands/Source/HairStrandsCore/Private/GroomComponent.cpp b/Engine/Plugins/Runtime/HairStrands/Source/HairStrandsCore/Private/GroomComponent.cpp
--- a/Engine/Plugins/Runtime/HairStrands/Source/HairStrandsCore/Private/GroomComponent.cpp
+++ b/Engine/Plugins/Runtime/HairStrands/Source/HairStrandsCore/Private/GroomComponent.cpp
@@ -2819,7 +2819,7 @@ void UGroomComponent::InitResources(bool bIsBindingReloading)
 	{
 		RegisteredMeshComponent = ValidatedMeshComponent;
 		AddTickPrerequisiteComponent(ValidatedMeshComponent);
-		bUseAttachParentBound = false;
+
 		if (bHasNeedBindingData)
 		{
 			LocalBindingAsset = BindingAsset;
