Fix out of bound access when hairstrands plugin is disabled through CVar.

#Synced-CL 51090839

diff --git a/Engine/Plugins/Runtime/HairStrands/Source/HairStrandsCore/Private/GroomComponent.cpp b/Engine/Plugins/Runtime/HairStrands/Source/HairStrandsCore/Private/GroomComponent.cpp
--- a/Engine/Plugins/Runtime/HairStrands/Source/HairStrandsCore/Private/GroomComponent.cpp
+++ b/Engine/Plugins/Runtime/HairStrands/Source/HairStrandsCore/Private/GroomComponent.cpp
@@ -2308,9 +2308,9 @@ EHairGeometryType UGroomComponent::GetMaterialGeometryType(int32 ElementIndex) c
 	for (uint32 GroupIt = 0, GroupCount = GroomAsset->GetHairGroupsRendering().Num(); GroupIt < GroupCount; ++GroupIt)
 	{
 		// Material - Strands
-		const FHairGroupPlatformData& InGroupData = GroomAsset->GetHairGroupsPlatformData()[GroupIt];
 		if (IsHairStrandsEnabled(EHairStrandsShaderType::Strands, Platform))
 		{
+			const FHairGroupPlatformData& InGroupData = GroomAsset->GetHairGroupsPlatformData()[GroupIt];
 			const int32 SlotIndex = GroomAsset->GetMaterialIndex(GroomAsset->GetHairGroupsRendering()[GroupIt].MaterialSlotName);
 			if (SlotIndex == ElementIndex)
 			{
@@ -2321,6 +2321,7 @@ EHairGeometryType UGroomComponent::GetMaterialGeometryType(int32 ElementIndex) c
 		// Material - Cards
 		if (IsHairStrandsEnabled(EHairStrandsShaderType::Cards, Platform))
 		{
+			const FHairGroupPlatformData& InGroupData = GroomAsset->GetHairGroupsPlatformData()[GroupIt];
 			uint32 CardsLODIndex = 0;
 			for (const FHairGroupPlatformData::FCards::FLOD& LOD : InGroupData.Cards.LODs)
 			{
@@ -2348,6 +2349,7 @@ EHairGeometryType UGroomComponent::GetMaterialGeometryType(int32 ElementIndex) c
 		// Material - Meshes
 		if (IsHairStrandsEnabled(EHairStrandsShaderType::Meshes, Platform))
 		{
+			const FHairGroupPlatformData& InGroupData = GroomAsset->GetHairGroupsPlatformData()[GroupIt];
 			uint32 MeshesLODIndex = 0;
 			for (const FHairGroupPlatformData::FMeshes::FLOD& LOD : InGroupData.Meshes.LODs)
 			{
