Fix shader print being incorrect activated by hair culling.

#Synced-CL 51090839

diff --git a/Engine/Plugins/Runtime/HairStrands/Source/HairStrandsCore/Private/HairStrandsClusterCulling.cpp b/Engine/Plugins/Runtime/HairStrands/Source/HairStrandsCore/Private/HairStrandsClusterCulling.cpp
--- a/Engine/Plugins/Runtime/HairStrands/Source/HairStrandsCore/Private/HairStrandsClusterCulling.cpp
+++ b/Engine/Plugins/Runtime/HairStrands/Source/HairStrandsCore/Private/HairStrandsClusterCulling.cpp
@@ -94,9 +94,13 @@ void AddClusterCullingPass(
 {
 	check(View);
 
-	ShaderPrint::SetEnabled(true);
-	ShaderPrint::RequestSpaceForCharacters(2048);
-	ShaderPrint::RequestSpaceForLines(2048);
+	static bool bDebug = false;
+	if (bDebug)
+	{
+		ShaderPrint::SetEnabled(true);
+		ShaderPrint::RequestSpaceForCharacters(2048);
+		ShaderPrint::RequestSpaceForLines(2048);
+	}
 
 	const uint32 ClusterCount = ClusterDatas.HairGroups.Num();
 
@@ -124,7 +128,7 @@ void AddClusterCullingPass(
 		Parameters->OutIndexBuffer 		 = RegisterAsUAV(GraphBuilder, *ClusterData.GetCulledVertexIdBuffer());
 		Parameters->OutCulledCurveBuffer = RegisterAsUAV(GraphBuilder, *ClusterData.GetCulledCurveBuffer()); // TODO: this could be changed to be transient buffer
 
-		if (ShaderPrintData)
+		if (bDebug && ShaderPrintData)
 		{
 			ShaderPrint::SetParameters(GraphBuilder, *ShaderPrintData, Parameters->ShaderPrintUniformBuffer);
 		}
diff --git a/Engine/Shaders/Private/HairStrands/HairStrandsClusterCulling.usf b/Engine/Shaders/Private/HairStrands/HairStrandsClusterCulling.usf
--- a/Engine/Shaders/Private/HairStrands/HairStrandsClusterCulling.usf
+++ b/Engine/Shaders/Private/HairStrands/HairStrandsClusterCulling.usf
@@ -58,8 +58,6 @@ void MainClearAABBCS(uint2 DispatchThreadId : SV_DispatchThreadID)
 
 #if SHADER_CLUSTER_CULL
 
-#include "../ShaderPrint.ush"  
-
 #define NUM_THREADS_PER_GROUP GROUP_SIZE
 #include "../ThreadGroupPrefixSum.ush"
 
