Dataflow : add average curve node for groom

#Synced-CL 51104518

diff --git a/Engine/Plugins/Runtime/HairStrands/Source/HairStrandsDataflow/Private/AverageCurvesNode.cpp b/Engine/Plugins/Runtime/HairStrands/Source/HairStrandsDataflow/Private/AverageCurvesNode.cpp
new file mode 100644
--- /dev/null
+++ b/Engine/Plugins/Runtime/HairStrands/Source/HairStrandsDataflow/Private/AverageCurvesNode.cpp
@@ -0,0 +1,98 @@
+// Copyright Epic Games, Inc. All Rights Reserved.
+
+#include "AverageCurvesNode.h"
+
+#include "GeometryCollection/Facades/CollectionCurveFacade.h"
+
+#include UE_INLINE_GENERATED_CPP_BY_NAME(AverageCurvesNode)
+
+#define LOCTEXT_NAMESPACE "AverageCurvesNode"
+
+FDataflowAverageCurves::FDataflowAverageCurves(const UE::Dataflow::FNodeParameters& InParam, FGuid InGuid)
+	: FDataflowNode(InParam, InGuid)
+{
+	RegisterInputConnection(&Collection);
+	RegisterInputConnection(&CurveSelection);
+	RegisterInputConnection(&NumSamples);
+	RegisterOutputConnection(&Collection);
+}
+
+void FDataflowAverageCurves::Evaluate(UE::Dataflow::FContext& Context, const FDataflowOutput* Out) const
+{
+	if (Out->IsA(&Collection))
+	{
+		FManagedArrayCollection OutCollection;
+
+		const int32 InNumSamples = GetValue(Context, &NumSamples);
+		if (InNumSamples < 2)
+		{
+			Context.Error(LOCTEXT("InvalidCollection", "Collection does not contain curves, result collection will be empty"), this, Out);
+			SetValue(Context, MoveTemp(OutCollection), &Collection);
+			return;
+		}
+
+		const FManagedArrayCollection& InCollection = GetValue(Context, &Collection);
+		const FDataflowCurveSelection& InCurveSelection = GetValue(Context, &CurveSelection);
+		if (!InCurveSelection.IsValidForCollection(InCollection))
+		{
+			Context.Error(LOCTEXT("InvalidCollection", "Collection does not contain curves, result collection will be empty"), this, Out);
+			SetValue(Context, MoveTemp(OutCollection), &Collection);
+			return;
+		}
+
+		GeometryCollection::Facades::FCollectionCurveGeometryFacade CurvesFacade(InCollection);
+		if (!CurvesFacade.IsValid())
+		{
+			Context.Error(LOCTEXT("InvalidCollection", "Collection does not contain curves, result collection will be empty"), this, Out);
+			SetValue(Context, MoveTemp(OutCollection), &Collection);
+			return;
+		}
+
+		TArray<FVector3f> AveragedCurve;
+		AveragedCurve.Init(FVector3f::ZeroVector, InNumSamples);
+		int32 NumProcessedCurves = 0;
+
+		TArray<FVector3f> ResampledCurve;
+		const int32 NumCurves = CurvesFacade.GetNumCurves();
+		for (int32 CurveIndex = 0; CurveIndex < NumCurves; ++CurveIndex)
+		{
+			if (InCurveSelection.IsSelected(CurveIndex))
+			{
+				CurvesFacade.GetCurveResampledPositions(CurveIndex, InNumSamples, ResampledCurve);
+				const int32 MaxIndex = FMath::Min(ResampledCurve.Num(), InNumSamples);
+				for (int32 Index = 0; Index < MaxIndex; ++Index)
+				{
+					AveragedCurve[Index] += ResampledCurve[Index];
+				}
+				++NumProcessedCurves;
+			}
+		}
+		
+		if (NumProcessedCurves == 0)
+		{
+			Context.Error(LOCTEXT("NoCurveSelected", "Selection seems empty, result collection will be empty"), this, Out);
+			SetValue(Context, MoveTemp(OutCollection), &Collection);
+			return;
+		}
+
+		const float InvNumProcessedCurves = 1.0f / (float)NumProcessedCurves;
+		for (int32 Index = 0; Index < AveragedCurve.Num(); ++Index)
+		{
+			AveragedCurve[Index] *= InvNumProcessedCurves;
+		}
+
+		// now add the resample curve to the collection
+		const TArray<int32> CurvePointOffsets = { AveragedCurve.Num() };
+		const TArray<int32> GeometryCurveOffsets = { 1 };
+		const TArray<FString> GeometryGroupNames = { TEXT("AverageCurve") };
+		const TArray<float> GeometryCurveThickness = { 0.5f };
+		const TArray<int32> CurveSourceIndices = { 0 };
+
+		GeometryCollection::Facades::FCollectionCurveGeometryFacade OutCurvesFacade(OutCollection);
+		OutCurvesFacade.InitCurvesCollection(AveragedCurve, CurvePointOffsets, GeometryCurveOffsets, GeometryGroupNames, GeometryCurveThickness, CurveSourceIndices);
+
+		SetValue(Context, MoveTemp(OutCollection), &Collection);
+	}
+}
+
+#undef LOCTEXT_NAMESPACE
\ No newline at end of file
diff --git a/Engine/Plugins/Runtime/HairStrands/Source/HairStrandsDataflow/Private/AverageCurvesNode.h b/Engine/Plugins/Runtime/HairStrands/Source/HairStrandsDataflow/Private/AverageCurvesNode.h
new file mode 100644
--- /dev/null
+++ b/Engine/Plugins/Runtime/HairStrands/Source/HairStrandsDataflow/Private/AverageCurvesNode.h
@@ -0,0 +1,40 @@
+// Copyright Epic Games, Inc. All Rights Reserved.
+
+#pragma once 
+
+#include "Dataflow/DataflowTerminalNode.h"
+#include "Dataflow/DataflowSelection.h"
+#include "GeometryCollection/ManagedArrayCollection.h"
+
+#include "AverageCurvesNode.generated.h"
+
+/** 
+* Take an collection curve selection and generate a single curve collection that is the average of the input curves
+*/
+USTRUCT(meta = (Experimental, DataflowGroom))
+struct FDataflowAverageCurves : public FDataflowNode
+{
+	GENERATED_BODY()
+
+	DATAFLOW_NODE_DEFINE_INTERNAL(FDataflowAverageCurves, "AverageCurves", "Groom", "Curves Guide Groom")
+
+public:
+	FDataflowAverageCurves(const UE::Dataflow::FNodeParameters& InParam, FGuid InGuid = FGuid::NewGuid());
+
+private:
+	//~ Begin FDataflowNode interface
+	virtual void Evaluate(UE::Dataflow::FContext& Context, const FDataflowOutput* Out) const override;
+	//~ End FDataflowNode interface
+
+	/** Managed array collection source of the curves to be average from */
+	UPROPERTY(meta = (DataflowInput, DataflowOutput, DataflowRenderGroups = "Surface"))
+	FManagedArrayCollection Collection;
+
+	/** Curve selection to use */
+	UPROPERTY(meta = (DataflowInput))
+	FDataflowCurveSelection CurveSelection;
+
+	/** Number of samples to use for averaged result curve */
+	UPROPERTY(EditAnywhere, Category = Settings, meta = (DataflowInput, ClampMin = 2, UIMin = 2))
+	int32 NumSamples = 8;
+};
diff --git a/Engine/Plugins/Runtime/HairStrands/Source/HairStrandsDataflow/Private/HairStrandsDataflowModule.cpp b/Engine/Plugins/Runtime/HairStrands/Source/HairStrandsDataflow/Private/HairStrandsDataflowModule.cpp
--- a/Engine/Plugins/Runtime/HairStrands/Source/HairStrandsDataflow/Private/HairStrandsDataflowModule.cpp
+++ b/Engine/Plugins/Runtime/HairStrands/Source/HairStrandsDataflow/Private/HairStrandsDataflowModule.cpp
@@ -12,6 +12,7 @@
 #include "BuildGroomSplineSkinningNode.h"
 #include "GenerateGuidesCurvesNode.h"
 #include "SmoothGuidesCurvesNode.h"
+#include "AverageCurvesNode.h"
 #include "GroomDataflowVisualization.h"
 #include "Dataflow/DataflowCollectionAddScalarVertexPropertyNode.h"
 #include "Dataflow/DataflowCategoryRegistry.h"
@@ -51,6 +52,7 @@ PRAGMA_ENABLE_DEPRECATION_WARNINGS
 	DATAFLOW_NODE_REGISTER_CREATION_FACTORY(FAttachCurveRootsDataflowNode);
 	DATAFLOW_NODE_REGISTER_CREATION_FACTORY(FBuildCurveWeightsDataflowNode);
 	DATAFLOW_NODE_REGISTER_CREATION_FACTORY(FBuildCurveLODsDataflowNode);
+	DATAFLOW_NODE_REGISTER_CREATION_FACTORY(FDataflowAverageCurves);
 
 	UE::Dataflow::FDataflowSimulationVisualizationRegistry::GetInstance().RegisterVisualization(
 		MakeUnique<UE::Groom::FGroomDataflowSimulationVisualization>());
