Dataflow : add a terminal node to convert guides to a skeleton ( GuidesToJoints )

#Synced-CL 51288786

diff --git a/Engine/Plugins/Runtime/HairStrands/Source/HairStrandsDataflow/HairStrandsDataflow.Build.cs b/Engine/Plugins/Runtime/HairStrands/Source/HairStrandsDataflow/HairStrandsDataflow.Build.cs
--- a/Engine/Plugins/Runtime/HairStrands/Source/HairStrandsDataflow/HairStrandsDataflow.Build.cs
+++ b/Engine/Plugins/Runtime/HairStrands/Source/HairStrandsDataflow/HairStrandsDataflow.Build.cs
@@ -27,6 +27,7 @@ namespace UnrealBuildTool.Rules
 					"MeshDescription",
 					"DynamicMesh",
 					"SkeletalMeshDescription",
+					"SkeletalMeshUtilitiesCommon",
 					"StaticMeshDescription",
 					"Slate",
 					"SlateCore",
diff --git a/Engine/Plugins/Runtime/HairStrands/Source/HairStrandsDataflow/Private/GuidesToJointsNode.cpp b/Engine/Plugins/Runtime/HairStrands/Source/HairStrandsDataflow/Private/GuidesToJointsNode.cpp
new file mode 100644
--- /dev/null
+++ b/Engine/Plugins/Runtime/HairStrands/Source/HairStrandsDataflow/Private/GuidesToJointsNode.cpp
@@ -0,0 +1,376 @@
+// Copyright Epic Games, Inc. All Rights Reserved.
+
+#include "GuidesToJointsNode.h"
+
+#include "GeometryCollection/Facades/CollectionCurveFacade.h"
+#include "GeometryCollection/Facades/CollectionVertexBoneWeightsFacade.h"
+#include "StaticToSkeletalMeshConverter.h"
+#include "SkeletalMeshAttributes.h"
+
+#include UE_INLINE_GENERATED_CPP_BY_NAME(GuidesToJointsNode)
+
+#define LOCTEXT_NAMESPACE "DataflowGuidesToJointsNode"
+
+////////////////////////////////////////////////////////////////////////////////////////////////////////////
+
+#include "AssetRegistry/AssetData.h"
+#include "AssetRegistry/AssetRegistryModule.h"
+#include "Misc/PackageName.h"
+#include "Animation/Skeleton.h"
+
+////////////////////////////////////////////////////////////////////////////////////////////////////////////
+
+FDataflowGuidesToJointsNode::FDataflowGuidesToJointsNode(const UE::Dataflow::FNodeParameters& InParam, FGuid InGuid)
+	: FDataflowTerminalNode(InParam, InGuid)
+{
+	RegisterInputConnection(&Collection);
+	RegisterInputConnection(&CurveSelection);
+	RegisterInputConnection(&RootBone);
+	RegisterInputConnection(&SkeletalMeshAssetPath);
+	RegisterInputConnection(&SkeletonAssetPath);
+	RegisterInputConnection(&BaseSkeletalMesh);
+
+	RegisterOutputConnection(&SkeletalMeshAsset);
+	RegisterOutputConnection(&SkeletonAsset);
+}
+
+void FDataflowGuidesToJointsNode::Evaluate(UE::Dataflow::FContext& Context) const
+{
+	// Asset is created and stored by the SetAssetValue, Evaluate only return an existing one or nullptr
+	const FString& InAssetPath = GetValue(Context, &SkeletalMeshAssetPath);
+	TObjectPtr<USkeletalMesh> OutSkeletalMesh = ::LoadObject<USkeletalMesh>(nullptr, InAssetPath);
+	SetValue(Context, OutSkeletalMesh, &SkeletalMeshAsset);
+
+	TObjectPtr<USkeleton> OutSkeleton = (OutSkeletalMesh)? OutSkeletalMesh->GetSkeleton(): nullptr;
+	SetValue(Context, OutSkeleton, &SkeletonAsset);
+}
+
+void FDataflowGuidesToJointsNode::SetAssetValue(TObjectPtr<UObject> Asset, UE::Dataflow::FContext& Context) const
+{
+	TObjectPtr<USkeletalMesh> NullSkeletalMeshPtr = nullptr;
+	TObjectPtr<USkeleton> NullSkeletonPtr = nullptr;
+
+	auto SetOutputsOnFailure = [this, &Context, &NullSkeletalMeshPtr, &NullSkeletonPtr]()
+		{
+			SetValue(Context, NullSkeletalMeshPtr, &SkeletalMeshAsset);
+			SetValue(Context, NullSkeletonPtr, &SkeletonAsset);
+		};
+
+	//-------------------------------------------------------------------
+	// input validation section
+	//-------------------------------------------------------------------
+
+	// veryfy the collection and selection
+	const FManagedArrayCollection& InCollection = GetValue(Context, &Collection);
+	const FDataflowCurveSelection& InCurveSelection = GetValue(Context, &CurveSelection);
+
+	// consider all curves when selection is not valid 
+	const bool bValidSelection = InCurveSelection.IsValidForCollection(InCollection);;
+	const bool bUseAllCurves = (!bValidSelection);
+
+	GeometryCollection::Facades::FCollectionCurveGeometryFacade CurvesFacade(InCollection);
+	if (!CurvesFacade.IsValid())
+	{
+		Context.Error(LOCTEXT("InvalidCurveFacade", "Invalid collection - do not contains curve information"));
+		SetOutputsOnFailure();
+		return;
+	}
+
+	constexpr int32 MaxNumSelectedCurves = 100;
+	const int32 NumSelectedCurves = bUseAllCurves? CurvesFacade.GetNumCurves(): InCurveSelection.NumSelected();
+	if (NumSelectedCurves > MaxNumSelectedCurves)
+	{
+		Context.Error(FText::Format(LOCTEXT("TooManyCurves", "Too many curve selected - Limit is {0}"), FText::AsNumber(MaxNumSelectedCurves)));
+		SetOutputsOnFailure();
+		return;
+	}
+
+	TObjectPtr<const USkeletalMesh> InBaseSkeletalMesh = GetValue(Context, &BaseSkeletalMesh);
+	if (IsValid(InBaseSkeletalMesh))
+	{
+		const FString SkeletalMeshPathFromInput = GetValue(Context, &SkeletalMeshAssetPath);
+		if (InBaseSkeletalMesh == Asset || InBaseSkeletalMesh->GetPathName() == SkeletalMeshPathFromInput)
+		{
+			Context.Error(FText::Format(LOCTEXT("BaseSkeletalMeshSameAsOuput", "Base skeletal mesh is the same as the output asset : {0}"), FText::FromString(InBaseSkeletalMesh->GetPathName())));
+			SetOutputsOnFailure();
+			return;
+		}
+	}
+
+#if WITH_EDITOR
+
+	//-------------------------------------------------------------------
+	// skeleton section
+	//-------------------------------------------------------------------
+	const FString SkeletonPathFromInput = GetValue(Context, &SkeletonAssetPath);
+	const FString& InSkeletonAssetPath = GetAssetPath(SkeletonPathFromInput, Asset);
+	if (InSkeletonAssetPath.IsEmpty())
+	{
+		Context.Error(FText::Format(LOCTEXT("InvalidSkeletonPathOrBoundAsset", "invalid skeleton path or no skeleton bound asset : {0}"), FText::FromString(SkeletonPathFromInput)));
+		SetOutputsOnFailure();
+		return;
+	}
+
+	TObjectPtr<USkeleton> OutSkeleton = Cast<USkeleton>(GetOrCreateAsset(Context, InSkeletonAssetPath, USkeleton::StaticClass()));
+	if (!OutSkeleton)
+	{
+		// error logging already handled by GetOrCreateAsset
+		SetOutputsOnFailure();
+		return;
+	}
+
+	UpdateSkeleton(Context, OutSkeleton);
+
+	SetValue(Context, OutSkeleton, &SkeletonAsset);
+
+	//-------------------------------------------------------------------
+	// skeletalMesh section (optional)
+	//-------------------------------------------------------------------
+	const FString SkeletalMeshPathFromInput = GetValue(Context, &SkeletalMeshAssetPath);
+	const FString& InSkeletalMeshAssetPath = GetAssetPath(SkeletalMeshPathFromInput, Asset);
+	if (InSkeletalMeshAssetPath.IsEmpty())
+	{
+		Context.Warning(FText::Format(LOCTEXT("InvalidSkeletalMeshPathOrBoundAsset", "invalid skeletalMesh path or no skeletalMesh bound asset - only saving a skeleton : {0}"), FText::FromString(SkeletalMeshPathFromInput)));
+		SetValue(Context, NullSkeletalMeshPtr, &SkeletalMeshAsset);
+		return;
+	}
+
+	TObjectPtr<USkeletalMesh> OutSkeletalMesh = Cast<USkeletalMesh>(GetOrCreateAsset(Context, InSkeletalMeshAssetPath, USkeletalMesh::StaticClass()));
+	if (!OutSkeletalMesh)
+	{
+		// error logging already handled by GetOrCreateAsset
+		SetValue(Context, NullSkeletalMeshPtr, &SkeletalMeshAsset);
+		return;
+	}
+
+	UpdateSkeletalMesh(Context, OutSkeletalMesh, OutSkeleton);
+
+	// finally set the skeletal mesh to the output
+	SetValue(Context, OutSkeletalMesh, &SkeletalMeshAsset);
+
+#else
+	Context.Error(TEXT("Creation of skeletal mesh asset only supported in Editor"), this);
+	SetOutputsOnFailure();
+#endif
+}
+
+int32 FDataflowGuidesToJointsNode::GetRefSkeletonRootIndex(const FReferenceSkeleton& RefSkeleton) const
+{
+	int32 RootIndex = INDEX_NONE;
+	const TArray<FMeshBoneInfo>& BoneInfos = RefSkeleton.GetRefBoneInfo();
+	if (BoneInfos.Num() > 0)
+	{
+		RootIndex = 0;
+		while (BoneInfos[RootIndex].ParentIndex != INDEX_NONE)
+		{
+			RootIndex = BoneInfos[RootIndex].ParentIndex;
+		}
+	}
+	return RootIndex;
+}
+
+void FDataflowGuidesToJointsNode::ClearSkeleton(USkeleton* InOutSkeleton) const
+{
+	// clear the existing skeleton data 
+	const TArray<FMeshBoneInfo>& BoneInfos = InOutSkeleton->GetReferenceSkeleton().GetRefBoneInfo();
+	if (BoneInfos.Num())
+	{
+		// loop so we remove all roots 
+		const int32 RootIndex = GetRefSkeletonRootIndex(InOutSkeleton->GetReferenceSkeleton());
+		if (RootIndex != INDEX_NONE)
+		{
+			{ // scope for modifier 
+				FReferenceSkeletonModifier Edit(InOutSkeleton);
+				Edit.Remove(BoneInfos[RootIndex].Name, /* bRemoveChildren */ true);
+			}
+		}
+	}
+}
+
+namespace DataflowGuidesToJointsNode::Private
+{
+	void AddRefBoneRecursivelyToSkeleton(const FReferenceSkeleton& RefSkeleton, const int32 RefBoneIndex, FReferenceSkeletonModifier& Edit)
+	{
+		if (RefBoneIndex != INDEX_NONE)
+		{
+			const FName BoneName = RefSkeleton.GetBoneName(RefBoneIndex);
+
+			int32 ParentIndex = INDEX_NONE;
+			const int32 RefParentIndex = RefSkeleton.GetParentIndex(RefBoneIndex);
+			if (RefParentIndex != INDEX_NONE)
+			{
+				const FName RefParentName = RefSkeleton.GetBoneName(RefParentIndex);
+				ParentIndex = Edit.FindBoneIndex(RefParentName);
+			}
+		
+			const FTransform BoneTransform = RefSkeleton.GetRefBonePose().IsValidIndex(RefBoneIndex) 
+				? RefSkeleton.GetRefBonePose()[RefBoneIndex]
+				: FTransform::Identity;
+
+			FMeshBoneInfo BoneInfo(BoneName, BoneName.ToString(), ParentIndex);
+			Edit.Add(BoneInfo, BoneTransform, false/*bAllowMultipleRoots*/);
+
+			TArray<int32> Children;
+			RefSkeleton.GetDirectChildBones(RefBoneIndex, Children);
+			for (const int32& Child : Children)
+			{
+				AddRefBoneRecursivelyToSkeleton(RefSkeleton, Child, Edit);
+			}
+		}
+	}
+}
+
+void FDataflowGuidesToJointsNode::InitializeSkeletonFromRefSkeleton(const FReferenceSkeleton& RefSkeleton, USkeleton* InOutSkeleton) const
+{
+	ClearSkeleton(InOutSkeleton);
+	const int32 RefRootIndex = GetRefSkeletonRootIndex(RefSkeleton);
+	if (RefRootIndex != INDEX_NONE)
+	{
+		FReferenceSkeletonModifier Edit(InOutSkeleton);
+		DataflowGuidesToJointsNode::Private::AddRefBoneRecursivelyToSkeleton(RefSkeleton, RefRootIndex, Edit);
+	}
+}
+
+void FDataflowGuidesToJointsNode::UpdateSkeleton(UE::Dataflow::FContext& Context, USkeleton* InOutSkeleton) const
+{
+	ClearSkeleton(InOutSkeleton);
+
+	// Get root node , search the base skeletal mesh, if not found create one 
+	const FString BoneToUseAsRoot = GetValue(Context, &RootBone);
+	const TObjectPtr<const USkeletalMesh> InBaseSkeletalMesh = GetValue(Context, &BaseSkeletalMesh);
+
+	int32 RootBoneIndex = INDEX_NONE;
+	FTransform RootBoneTransform = FTransform::Identity;
+	if (IsValid(InBaseSkeletalMesh))
+	{
+		// initialize the skeleton with the ref skeleton
+		const FReferenceSkeleton& RefSkeleton = InBaseSkeletalMesh->GetRefSkeleton();
+		InitializeSkeletonFromRefSkeleton(RefSkeleton, InOutSkeleton);
+
+		RootBoneIndex = RefSkeleton.FindRawBoneIndex(FName(BoneToUseAsRoot));
+		if (RootBoneIndex == INDEX_NONE)
+		{
+			const FText InfoMessage = FText::Format(LOCTEXT("InvalidRootBoneName", "Bone {0} does not exists in Skeletal mesh {1}, use the actual root bone if present"), FText::FromString(BoneToUseAsRoot), FText::FromString(InBaseSkeletalMesh->GetPathName()));
+			Context.Info(InfoMessage, this);
+		}
+		else
+		{
+			RootBoneTransform = RefSkeleton.GetBoneAbsoluteTransform(RootBoneIndex);
+		}
+	}
+
+	FReferenceSkeletonModifier Edit(InOutSkeleton);
+
+	// Add a root if we haven't found one yet 
+	if (RootBoneIndex == INDEX_NONE)
+	{
+		RootBoneIndex = GetRefSkeletonRootIndex(InOutSkeleton->GetReferenceSkeleton());
+		if (RootBoneIndex == INDEX_NONE)
+		{
+			static const FName DefaultRootBoneName(TEXT("root"));
+			FMeshBoneInfo BoneInfo(DefaultRootBoneName, DefaultRootBoneName.ToString(), INDEX_NONE);
+			Edit.Add(BoneInfo, FTransform::Identity, false/*bAllowMultipleRoots*/);
+			RootBoneIndex = Edit.FindBoneIndex(DefaultRootBoneName);
+		}
+		else
+		{
+			RootBoneTransform = InOutSkeleton->GetReferenceSkeleton().GetBoneAbsoluteTransform(RootBoneIndex);
+		}
+	}
+
+	// add the curves 
+	const FManagedArrayCollection& InCollection = GetValue(Context, &Collection);
+	const FDataflowCurveSelection& InCurveSelection = GetValue(Context, &CurveSelection);
+
+	// consider all curves when selection is not valid 
+	const bool bValidSelection = InCurveSelection.IsValidForCollection(InCollection);;
+	const bool bUseAllCurves = (!bValidSelection);
+
+	GeometryCollection::Facades::FCollectionCurveGeometryFacade CurvesFacade(InCollection);
+	if (CurvesFacade.IsValid())
+	{
+		TArray<FVector3f> CurvePoints;
+		const int32 NumCurves = CurvesFacade.GetNumCurves();
+		for (int32 CurveIndex = 0; CurveIndex < NumCurves; ++CurveIndex)
+		{
+			if (bUseAllCurves || InCurveSelection.IsSelected(CurveIndex))
+			{
+				CurvesFacade.GetCurvePointPositions(CurveIndex, CurvePoints);
+
+				// Now add to the skeleton
+				FTransform BoneParentTransform = RootBoneTransform;
+				int32 BoneParentIndex = RootBoneIndex;
+				const int32 NumPoints = CurvePoints.Num();
+				for (int32 PointIndex = 0; PointIndex < NumPoints; ++PointIndex)
+				{
+					const FString BoneNameStr = FString::Format(TEXT("Curve_{0}_{1}"), { FString::FromInt(CurveIndex), FString::FromInt(PointIndex) });
+					const FVector BonePosition = FVector(CurvePoints[PointIndex]);
+					const FTransform AbsoluteBoneTransform(BonePosition);
+					const FTransform LocalBoneTransform = AbsoluteBoneTransform.GetRelativeTransform(BoneParentTransform);
+					FMeshBoneInfo BoneInfo(FName(BoneNameStr), BoneNameStr, BoneParentIndex);
+					Edit.Add(BoneInfo, LocalBoneTransform, false);
+					BoneParentIndex = Edit.FindBoneIndex(FName(BoneNameStr));
+					BoneParentTransform = AbsoluteBoneTransform;
+				}
+			}
+		}
+	}
+}
+
+void FDataflowGuidesToJointsNode::UpdateSkeletalMesh(UE::Dataflow::FContext& Context, USkeletalMesh* InOutSkeletalMesh, USkeleton* InOutSkeleton) const
+{
+	InOutSkeletalMesh->PreEditChange(nullptr);
+	{
+		InOutSkeletalMesh->Clear();
+		// Make sure that we create non-transient assets
+		InOutSkeletalMesh->ClearFlags(RF_Transient);
+
+		InOutSkeletalMesh->SetSkeleton(InOutSkeleton);
+
+		FMeshDescription MeshDescription;
+		const FVertexID VId0 = MeshDescription.CreateVertex();
+		const FVertexID VId1 = MeshDescription.CreateVertex();
+		const FVertexID VId2 = MeshDescription.CreateVertex();
+		const int32 VInst0 = MeshDescription.CreateVertexInstance(VId0);
+		const int32 VInst1 = MeshDescription.CreateVertexInstance(VId1);
+		const int32 VInst2 = MeshDescription.CreateVertexInstance(VId2);
+		TVertexAttributesRef<FVector3f> Positions = MeshDescription.GetVertexPositions();
+		Positions[VId0] = FVector3f(0.0f, 0.0f, 0.0f);
+		Positions[VId1] = FVector3f(0.1f, 0.0f, 0.0f);
+		Positions[VId2] = FVector3f(0.1f, 0.1f, 0.0f);
+		const FPolygonGroupID PolyGroupID = MeshDescription.CreatePolygonGroup();
+		MeshDescription.CreateTriangle(PolyGroupID, { VInst0, VInst1, VInst2 });
+		
+		// add bone weights information
+		FSkeletalMeshAttributes SkeletalMeshAttributes{ MeshDescription };
+		SkeletalMeshAttributes.Register(/*bKeepExistingAttribute*/true);
+		FSkinWeightsVertexAttributesRef VertexSkinWeights = SkeletalMeshAttributes.GetVertexSkinWeights();
+		UE::AnimationCore::FBoneWeights Weights;
+		Weights.SetBoneWeight(0, 1.0f);
+		VertexSkinWeights.Set(VId0, Weights);
+		VertexSkinWeights.Set(VId1, Weights);
+		VertexSkinWeights.Set(VId2, Weights);
+
+		TArray<const FMeshDescription*> MeshDescriptions = { &MeshDescription };
+		// Update the skeletal mesh from the description
+		FStaticToSkeletalMeshConverter::FInitializationParams Parameters;
+		//Parameters.Materials = SkeletalMaterials;
+		Parameters.bRecomputeNormals = true;
+		Parameters.bRecomputeTangents = false;
+
+		FStaticToSkeletalMeshConverter::InitializeSkeletalMeshFromMeshDescriptions(
+			InOutSkeletalMesh,
+			MeshDescriptions,
+			InOutSkeleton->GetReferenceSkeleton(),
+			Parameters);
+
+		// update the relationship between the asset
+		InOutSkeleton->SetPreviewMesh(InOutSkeletalMesh);
+		InOutSkeleton->RecreateBoneTree(InOutSkeletalMesh);
+	}
+	InOutSkeletalMesh->InvalidateDeriveDataCacheGUID();
+	//InOutSkeletalMesh->InitResources();
+	InOutSkeletalMesh->PostEditChange();
+}
+
+#undef LOCTEXT_NAMESPACE
\ No newline at end of file
diff --git a/Engine/Plugins/Runtime/HairStrands/Source/HairStrandsDataflow/Private/GuidesToJointsNode.h b/Engine/Plugins/Runtime/HairStrands/Source/HairStrandsDataflow/Private/GuidesToJointsNode.h
new file mode 100644
--- /dev/null
+++ b/Engine/Plugins/Runtime/HairStrands/Source/HairStrandsDataflow/Private/GuidesToJointsNode.h
@@ -0,0 +1,74 @@
+// Copyright Epic Games, Inc. All Rights Reserved.
+
+#pragma once 
+
+#include "Dataflow/DataflowAnyType.h"
+#include "Dataflow/DataflowTerminalNode.h"
+#include "Dataflow/DataflowSelection.h"
+//#include "Dataflow/Skeleton.h"
+
+#include "Engine/SkeletalMesh.h"
+#include "GeometryCollection/ManagedArrayCollection.h"
+#include "UObject/ObjectPtr.h"
+
+#include "GuidesToJointsNode.generated.h"
+
+
+/** 
+* Generate joints/bones from a guide selection and save the resulting skeleton asset on disk 
+*/
+USTRUCT(meta = (Experimental, DataflowGroom))
+struct FDataflowGuidesToJointsNode : public FDataflowTerminalNode
+{
+	GENERATED_BODY()
+
+	DATAFLOW_NODE_DEFINE_INTERNAL(FDataflowGuidesToJointsNode, "GuidesToJoints", "Terminal", "Bones Rigged SkeletonMesh")
+
+public:
+	FDataflowGuidesToJointsNode(const UE::Dataflow::FNodeParameters& InParam, FGuid InGuid = FGuid::NewGuid());
+
+private:
+	//~ Begin FDataflowTerminalNode interface
+	virtual void Evaluate(UE::Dataflow::FContext& Context) const override;
+	virtual void SetAssetValue(TObjectPtr<UObject> Asset, UE::Dataflow::FContext& Context) const override;
+	//~ End FDataflowTerminalNode interface
+
+	void ClearSkeleton(USkeleton* InOutSkeleton) const;
+	void UpdateSkeleton(UE::Dataflow::FContext& Context, USkeleton* InOutSkeleton) const;
+	void InitializeSkeletonFromRefSkeleton(const FReferenceSkeleton& RefSkeleton, USkeleton* InOutSkeleton) const;
+	int32 GetRefSkeletonRootIndex(const FReferenceSkeleton& RefSkeleton) const;
+
+	void UpdateSkeletalMesh(UE::Dataflow::FContext& Context, USkeletalMesh* InOutSkeletalMesh, USkeleton* InOutSkeleton) const;
+
+	/** Managed array collection source of the guides */
+	UPROPERTY(meta = (DataflowInput, DataflowRenderGroups = "Surface"))
+	FManagedArrayCollection Collection;
+
+	/** Curve selection to use to generate bones/joints - if none specified all curve will be considered */
+	UPROPERTY(meta = (DataflowInput))
+	FDataflowCurveSelection CurveSelection;
+
+	/** Optional name of root bone from the input skeletal mesh to use as the root of the generate bones */
+	UPROPERTY(EditAnywhere, Category = Source, meta = (DataflowInput))
+	FDataflowStringTypes RootBone;
+
+	/** Path to the skeletal mesh asset to create or update */
+	UPROPERTY(EditAnywhere, Category = Asset, meta = (DataflowInput))
+	FString SkeletalMeshAssetPath;
+
+	/** Path to the skeleton asset to create or update */
+	UPROPERTY(EditAnywhere, Category = Asset, meta = (DataflowInput))
+	FString SkeletonAssetPath;
+
+	/** Skeletal Mesh to use a base*/
+	UPROPERTY(EditAnywhere, Category = Asset, meta = (DataflowInput))
+	TObjectPtr<const USkeletalMesh> BaseSkeletalMesh = nullptr;
+
+	/** Created Skeletal Mesh */
+	UPROPERTY(meta = (DataflowOutput))
+	TObjectPtr<USkeletalMesh> SkeletalMeshAsset = nullptr;
+
+	/** Created Skeleton */
+	UPROPERTY(meta = (DataflowOutput))
+	TObjectPtr<USkeleton> SkeletonAsset = nullptr;
+};
diff --git a/Engine/Plugins/Runtime/HairStrands/Source/HairStrandsDataflow/Private/HairStrandsDataflowModule.cpp b/Engine/Plugins/Runtime/HairStrands/Source/HairStrandsDataflow/Private/HairStrandsDataflowModule.cpp
--- a/Engine/Plugins/Runtime/HairStrands/Source/HairStrandsDataflow/Private/HairStrandsDataflowModule.cpp
+++ b/Engine/Plugins/Runtime/HairStrands/Source/HairStrandsDataflow/Private/HairStrandsDataflowModule.cpp
@@ -13,6 +13,7 @@
 #include "GenerateGuidesCurvesNode.h"
 #include "SmoothGuidesCurvesNode.h"
 #include "AverageCurvesNode.h"
+#include "GuidesToJointsNode.h"
 #include "GroomDataflowVisualization.h"
 #include "Dataflow/DataflowCollectionAddScalarVertexPropertyNode.h"
 #include "Dataflow/DataflowCategoryRegistry.h"
@@ -53,6 +54,7 @@ PRAGMA_ENABLE_DEPRECATION_WARNINGS
 	DATAFLOW_NODE_REGISTER_CREATION_FACTORY(FBuildCurveWeightsDataflowNode);
 	DATAFLOW_NODE_REGISTER_CREATION_FACTORY(FBuildCurveLODsDataflowNode);
 	DATAFLOW_NODE_REGISTER_CREATION_FACTORY(FDataflowAverageCurves);
+	DATAFLOW_NODE_REGISTER_CREATION_FACTORY(FDataflowGuidesToJointsNode);
 
 	UE::Dataflow::FDataflowSimulationVisualizationRegistry::GetInstance().RegisterVisualization(
 		MakeUnique<UE::Groom::FGroomDataflowSimulationVisualization>());
