dataflow groom : fix guides to joints terminal to properly clear the skeleton

#Synced-CL 51362814

diff --git a/Engine/Plugins/Runtime/HairStrands/Source/HairStrandsDataflow/Private/GuidesToJointsNode.cpp b/Engine/Plugins/Runtime/HairStrands/Source/HairStrandsDataflow/Private/GuidesToJointsNode.cpp
--- a/Engine/Plugins/Runtime/HairStrands/Source/HairStrandsDataflow/Private/GuidesToJointsNode.cpp
+++ b/Engine/Plugins/Runtime/HairStrands/Source/HairStrandsDataflow/Private/GuidesToJointsNode.cpp
@@ -179,9 +179,22 @@ void FDataflowGuidesToJointsNode::ClearSkeleton(USkeleton* InOutSkeleton) const
 		const int32 RootIndex = GetRefSkeletonRootIndex(InOutSkeleton->GetReferenceSkeleton());
 		if (RootIndex != INDEX_NONE)
 		{
+			// cannot remove the root so we need to get the children and remove them one by one 
+			TArray<int32> ChildBoneIndices;
+			InOutSkeleton->GetReferenceSkeleton().GetDirectChildBones(RootIndex, ChildBoneIndices);
+			TArray<FName> ChildBoneNames;
+			ChildBoneNames.Reserve(ChildBoneIndices.Num());
+			for (const int32 ChildBoneIndex : ChildBoneIndices)
+			{
+				ChildBoneNames.Add(InOutSkeleton->GetReferenceSkeleton().GetBoneName(ChildBoneIndex));
+			}
+
 			{ // scope for modifier 
 				FReferenceSkeletonModifier Edit(InOutSkeleton);
-				Edit.Remove(BoneInfos[RootIndex].Name, /* bRemoveChildren */ true);
+				for (const FName ChildBoneName : ChildBoneNames)
+				{
+					Edit.Remove(ChildBoneName, /* bRemoveChildren */ true);
+				}
 			}
 		}
 	}
