Fix hair cards/meshes not always outputing velocity vectory when needed.

The velocity output was not taking into account the AlwaysOutputVelocity flags.

#Synced-CL 51334621

diff --git a/Engine/Plugins/Runtime/HairStrands/Source/HairStrandsCore/Private/GroomComponent.cpp b/Engine/Plugins/Runtime/HairStrands/Source/HairStrandsCore/Private/GroomComponent.cpp
--- a/Engine/Plugins/Runtime/HairStrands/Source/HairStrandsCore/Private/GroomComponent.cpp
+++ b/Engine/Plugins/Runtime/HairStrands/Source/HairStrandsCore/Private/GroomComponent.cpp
@@ -1062,11 +1062,7 @@ public:
 
 		FPrimitiveSceneInfo* PrimSceneInfo = GetPrimitiveSceneInfo();
 		GetScene().GetPrimitiveUniformShaderParameters_RenderThread(PrimSceneInfo, bHasPrecomputedVolumetricLightmap, PreviousLocalToWorld, SingleCaptureIndex, bOutputVelocity);
-
-		if (GHairStrands_ForceVelocityOutput > 0)
-		{
-			bOutputVelocity = true;
-		}
+		bOutputVelocity |= AlwaysHasVelocity();
 
 		const bool bUseProxy = UseProxyLocalToWorld(Instance);
 
