Add early return statement to GroomComponent::SetForcedLOD to fix an edge case where an LOD value of -1 would call SwitchSimulationLOD with an invalid LOD when in a non-forced LOD selection mode.

#Synced-CL 51453793

diff --git a/Engine/Plugins/Runtime/HairStrands/Source/HairStrandsCore/Private/GroomComponent.cpp b/Engine/Plugins/Runtime/HairStrands/Source/HairStrandsCore/Private/GroomComponent.cpp
--- a/Engine/Plugins/Runtime/HairStrands/Source/HairStrandsCore/Private/GroomComponent.cpp
+++ b/Engine/Plugins/Runtime/HairStrands/Source/HairStrandsCore/Private/GroomComponent.cpp
@@ -1878,6 +1878,14 @@ void UGroomComponent::SetForcedLOD(int32 CurrLODIndex)
 		CurrLODSelectionType = bNeedSimulationNeed ? EHairLODSelectionType::Predicted : EHairLODSelectionType::Immediate;
 	}
 
+	// If not setting a valid forced LOD, just reset the state and let the auto system handle the rest	
+	if (CurrLODSelectionType != EHairLODSelectionType::Forced)
+	{
+		LODForcedIndex = CurrLODIndex;
+		LODSelectionType = CurrLODSelectionType;
+		return;
+	}
+
 	// Inform simulation about LOD switch.
 	SwitchSimulationLOD(PrevLODIndex, CurrLODIndex, CurrLODSelectionType);
 
