Fix groom allocating RT memory for strands data even when not needed if a project has PathTracing enabled.

#Synced-CL 51536989

diff --git a/Engine/Plugins/Runtime/HairStrands/Source/HairStrandsCore/Private/GroomAsset.cpp b/Engine/Plugins/Runtime/HairStrands/Source/HairStrandsCore/Private/GroomAsset.cpp
--- a/Engine/Plugins/Runtime/HairStrands/Source/HairStrandsCore/Private/GroomAsset.cpp
+++ b/Engine/Plugins/Runtime/HairStrands/Source/HairStrandsCore/Private/GroomAsset.cpp
@@ -3409,6 +3409,18 @@ FHairStrandsRaytracingResource* UGroomAsset::AllocateMeshesRaytracingResources(u
 	return nullptr;
 }
 
+FHairStrandsRaytracingResource** UGroomAsset::GetStrandsRaytracingResourcesPtr(uint32 GroupIndex)
+{
+	if (GetHairGroupsPlatformData().IsValidIndex(GroupIndex))
+	{
+		FHairGroupPlatformData& GroupData = GetHairGroupsPlatformData()[GroupIndex];
+		FHairGroupResources& GroupResource = GetHairGroupsResources()[GroupIndex];
+		check(GroupData.Strands.HasValidData());
+		return &GroupResource.Strands.RaytracingResource;
+	}
+	return nullptr;
+}
+
 FHairStrandsRaytracingResource* UGroomAsset::AllocateStrandsRaytracingResources(uint32 GroupIndex)
 {
 	if (IsRayTracingAllowed() && GetHairGroupsPlatformData().IsValidIndex(GroupIndex))
diff --git a/Engine/Plugins/Runtime/HairStrands/Source/HairStrandsCore/Private/GroomComponent.cpp b/Engine/Plugins/Runtime/HairStrands/Source/HairStrandsCore/Private/GroomComponent.cpp
--- a/Engine/Plugins/Runtime/HairStrands/Source/HairStrandsCore/Private/GroomComponent.cpp
+++ b/Engine/Plugins/Runtime/HairStrands/Source/HairStrandsCore/Private/GroomComponent.cpp
@@ -89,6 +89,10 @@ static FAutoConsoleVariableRef CVarHairStrands_ForceVelocityOutput(TEXT("r.HairS
 bool IsHairManualSkinCacheEnabled();
 bool GetSupportHairStrandsLinearSweptSpheres(EShaderPlatform InShaderPlatform);
 
+#if RHI_RAYTRACING
+void AllocateRaytracingResources(FHairGroupInstance* Instance, bool bNeedDynamicResources, FName OwnerName);
+#endif
+
 ////////////////////////////////////////////////////////////////////////////////////////////////////////////////////////
 
 bool IsGroomBindingValidationEnabled()
@@ -3144,36 +3148,18 @@ void UGroomComponent::InitResources(bool bIsBindingReloading)
 
 			#if RHI_RAYTRACING
 			HairGroupInstance->Strands.ViewRayTracingMask = EHairViewRayTracingMask::None;
+			HairGroupInstance->Strands.RenRaytracingResourceShared = GroomAsset->GetStrandsRaytracingResourcesPtr(GroupIt);
 			if (IsHairRayTracingEnabled() && bVisibleInRayTracing)
 			{
-				// If the project has path tracing, we may need raytracing resources
-				static const auto CVarLocalVariablePathTracing = IConsoleManager::Get().FindTConsoleVariableDataInt(TEXT("r.PathTracing"));
-				const int PathTracingEnabled = CVarLocalVariablePathTracing ? CVarLocalVariablePathTracing->GetValueOnAnyThread() : 0;
-				if (PathTracingEnabled)
-				{
-					// Always visible for path tracing
-					HairGroupInstance->Strands.ViewRayTracingMask |= EHairViewRayTracingMask::PathTracing;
-				}
 				// Visible to ray tracing if the user has opted in
+				// If the project has path tracing, resources for pathtracing will be lazyly allocated on demand
 				if (HairGroupInstance->Strands.Modifier.bUseHairRaytracingGeometry)
 				{
 					HairGroupInstance->Strands.ViewRayTracingMask |= EHairViewRayTracingMask::RayTracing;
 				}
 				if (HairGroupInstance->Strands.ViewRayTracingMask != EHairViewRayTracingMask::None)
 				{
-					if (bNeedDynamicResources)
-					{
-						// Allocate dynamic raytracing resources (owned by the groom component/instance)
-						HairGroupInstance->Strands.RenRaytracingResource = new FHairStrandsRaytracingResource(GroupData.Strands.BulkData, ResourceName, OwnerName);
-						HairGroupInstance->Strands.RenRaytracingResourceOwned = true;
-					}
-					else
-					{
-						// (Lazy) Allocate static raytracing resources (owned by the grooom asset)
-						HairGroupInstance->Strands.RenRaytracingResourceOwned = false;
-						HairGroupInstance->Strands.RenRaytracingResource = GroomAsset->AllocateStrandsRaytracingResources(GroupIt);
-						check(HairGroupInstance->Strands.RenRaytracingResource);
-					}
+					AllocateRaytracingResources(HairGroupInstance, bNeedDynamicResources, OwnerName);
 				}
 			}
 			#endif
diff --git a/Engine/Plugins/Runtime/HairStrands/Source/HairStrandsCore/Private/GroomManager.cpp b/Engine/Plugins/Runtime/HairStrands/Source/HairStrandsCore/Private/GroomManager.cpp
--- a/Engine/Plugins/Runtime/HairStrands/Source/HairStrandsCore/Private/GroomManager.cpp
+++ b/Engine/Plugins/Runtime/HairStrands/Source/HairStrandsCore/Private/GroomManager.cpp
@@ -2403,6 +2403,57 @@ static void AddCopyHairStrandsPositionPass(
 	AddCopyBufferPass(GraphBuilder, Deformed1Buffer, RestBuffer);
 }
 
+#if RHI_RAYTRACING
+void AllocateRaytracingResources(FHairGroupInstance* Instance, bool bNeedDynamicResources, FName OwnerName)
+{
+	check(Instance->Strands.RenRaytracingResource == nullptr);
+
+	if (bNeedDynamicResources)
+	{
+		// Allocate dynamic raytracing resources (owned by the groom component/instance)
+		FHairResourceName ResourceName(FName(Instance->Debug.GroomAssetName), Instance->Debug.GroupIndex);
+		Instance->Strands.RenRaytracingResource = new FHairStrandsRaytracingResource(Instance->Strands.GetData(), ResourceName, OwnerName);
+		Instance->Strands.RenRaytracingResourceOwned = true;
+	}
+	else
+	{
+		if (Instance->Strands.RenRaytracingResourceShared)
+		{
+			// (Lazy) Allocate static raytracing resources (owned by the grooom asset)
+			if (*Instance->Strands.RenRaytracingResourceShared == nullptr)
+			{
+				FName SharedOwnerName(Instance->Debug.GroomAssetName);
+				FHairResourceName ResourceName(SharedOwnerName, Instance->Debug.GroupIndex);
+				*Instance->Strands.RenRaytracingResourceShared = new FHairStrandsRaytracingResource(Instance->Strands.GetData(), ResourceName, SharedOwnerName);
+			}
+			Instance->Strands.RenRaytracingResourceOwned = false;
+			Instance->Strands.RenRaytracingResource = *Instance->Strands.RenRaytracingResourceShared;
+			check(Instance->Strands.RenRaytracingResource);
+		}
+	}
+}
+
+static void AllocateRaytracingResources(FHairGroupInstance* Instance, int32 InMeshLODIndex)
+{
+	const bool bNeedDynamicResources = Instance->Strands.DeformedResource != nullptr;
+	if (IsHairRayTracingEnabled() && !Instance->Strands.RenRaytracingResource )
+	{
+		check(Instance->Strands.IsValid());
+
+		#if RHI_ENABLE_RESOURCE_INFO
+		FName OwnerName(FString::Printf(TEXT("%s [LOD%d]"), *Instance->Debug.MeshComponentName, InMeshLODIndex));
+		#else
+		FName OwnerName = NAME_None;
+		#endif
+		
+		AllocateRaytracingResources(Instance, bNeedDynamicResources, OwnerName);
+	}
+	Instance->Strands.ViewRayTracingMask |= EHairViewRayTracingMask::PathTracing;
+}
+
+#endif
+
+
 void AddHairStreamingRequest(FHairGroupInstance* Instance, int32 InLODIndex)
 {
 	if (Instance && InLODIndex >= 0)
@@ -2710,7 +2761,8 @@ static bool SelectValidLOD(
 	float PrevHairLODIndex, 
 	float HairLODCount, 
 	float MeshLODIndex, 
-	float PrevMeshLODIndex)
+	float PrevMeshLODIndex,
+	bool bHasPathTracingView)
 {
 	const int32 IntHairLODIndex = FMath::Clamp(FMath::FloorToInt(HairLODIndex), 0, HairLODCount - 1);
 	const TArray<bool>& LODVisibilities = Instance->HairGroupPublicData->GetLODVisibilities();
@@ -2863,6 +2915,7 @@ static bool SelectValidLOD(
 		if (Instance->Strands.DeformedRootResource)		{ Instance->Strands.DeformedRootResource->Allocate	(GraphBuilder, EHairResourceLoadingType::Async, ResourceStatus, RequestedCurveCount, RequestedPointCount, false/*bAllowDeallocation*/, MeshLODIndex); }
 		if (Instance->Strands.DeformedResource)			{ Instance->Strands.DeformedResource->Allocate		(GraphBuilder, EHairResourceLoadingType::Async, ResourceStatus, RequestedCurveCount, RequestedPointCount); }
 #if RHI_RAYTRACING
+		if (bHasPathTracingView)						{ AllocateRaytracingResources(Instance, MeshLODIndex); }
 		if (Instance->Strands.RenRaytracingResource)	{ Instance->Strands.RenRaytracingResource->Allocate	(GraphBuilder, EHairResourceLoadingType::Async, ResourceStatus, RequestedCurveCount, RequestedPointCount, Instance->Strands.RenRaytracingResourceOwned/*bAllowDeallocation*/); }
 #endif
 		
@@ -2954,6 +3007,19 @@ static void RunHairLODSelection(
 	{
 		ShaderPlatform = Views[0]->GetShaderPlatform();
 	}
+	
+	// Detect if one of the current view has path tracing enabled for allocating raytracing resources
+	bool bHasPathTracingView = false;
+	#if RHI_RAYTRACING
+	for (const FSceneView* View : Views)
+	{
+		if (View->Family->EngineShowFlags.PathTracing)
+		{
+			bHasPathTracingView = true;
+			break;
+		}
+	}
+	#endif
 
 	struct FInstanceData
 	{
@@ -3159,7 +3225,8 @@ static void RunHairLODSelection(
 			InstanceData.PrevHairLODIndex,
 			InstanceData.HairLOD.HairLODCount,
 			InstanceData.MeshLODIndex,
-			InstanceData.PrevMeshLODIndex);
+			InstanceData.PrevMeshLODIndex,
+			bHasPathTracingView);
 
 		// 2. If no valid LOD are found, try to find a coarser LOD which is:
 		// * loaded 
@@ -3183,7 +3250,8 @@ static void RunHairLODSelection(
 						InstanceData.PrevHairLODIndex,
 						InstanceData.HairLOD.HairLODCount,
 						InstanceData.MeshLODIndex,
-						InstanceData.PrevMeshLODIndex);
+						InstanceData.PrevMeshLODIndex,
+						bHasPathTracingView);
 					break;
 				}
 			}
diff --git a/Engine/Plugins/Runtime/HairStrands/Source/HairStrandsCore/Public/GroomAsset.h b/Engine/Plugins/Runtime/HairStrands/Source/HairStrandsCore/Public/GroomAsset.h
--- a/Engine/Plugins/Runtime/HairStrands/Source/HairStrandsCore/Public/GroomAsset.h
+++ b/Engine/Plugins/Runtime/HairStrands/Source/HairStrandsCore/Public/GroomAsset.h
@@ -861,6 +861,7 @@ private:
 	FHairStrandsRaytracingResource*		AllocateCardsRaytracingResources(uint32 GroupIndex, uint32 LODIndex);
 	FHairStrandsRaytracingResource*		AllocateMeshesRaytracingResources(uint32 GroupIndex, uint32 LODIndex);
 	FHairStrandsRaytracingResource*		AllocateStrandsRaytracingResources(uint32 GroupIndex);
+	FHairStrandsRaytracingResource**	GetStrandsRaytracingResourcesPtr(uint32 GroupIndex);
 #endif // RHI_RAYTRACING
 	friend class UGroomComponent;
 
diff --git a/Engine/Plugins/Runtime/HairStrands/Source/HairStrandsCore/Public/GroomInstance.h b/Engine/Plugins/Runtime/HairStrands/Source/HairStrandsCore/Public/GroomInstance.h
--- a/Engine/Plugins/Runtime/HairStrands/Source/HairStrandsCore/Public/GroomInstance.h
+++ b/Engine/Plugins/Runtime/HairStrands/Source/HairStrandsCore/Public/GroomInstance.h
@@ -138,6 +138,7 @@ struct FHairGroupInstance : public FHairStrandsInstance
 		// Resources - Raytracing data when enabling (expensive) raytracing method
 		#if RHI_RAYTRACING
 		FHairStrandsRaytracingResource* RenRaytracingResource = nullptr;
+		FHairStrandsRaytracingResource** RenRaytracingResourceShared = nullptr;
 		bool RenRaytracingResourceOwned = false;
 		EHairViewRayTracingMask ViewRayTracingMask = EHairViewRayTracingMask::None;
 		float CachedHairScaledRadius = 0;
