Dataflow : update Guides to joints terminal node with more options:
- Added option to copy the geometry from the base mesh
- Added option to define the base name for generated nodes
- Renamed root bone input to Parent bone
- Double click on thenode opens the skeletal mesh editor

#Synced-CL 51616234

diff --git a/Engine/Plugins/Experimental/Dataflow/Source/DataflowEditor/Private/Dataflow/DataflowSNode.cpp b/Engine/Plugins/Experimental/Dataflow/Source/DataflowEditor/Private/Dataflow/DataflowSNode.cpp
--- a/Engine/Plugins/Experimental/Dataflow/Source/DataflowEditor/Private/Dataflow/DataflowSNode.cpp
+++ b/Engine/Plugins/Experimental/Dataflow/Source/DataflowEditor/Private/Dataflow/DataflowSNode.cpp
@@ -453,9 +453,13 @@ FReply SDataflowEdNode::OnMouseButtonDoubleClick(const FGeometry& InMyGeometry,
 				FSourceCodeNavigation::NavigateToStruct(Node->TypedScriptStruct());
 			}
 		}
-		else
+		else 
 		{
-			Node->OnDoubleClicked();
+			if (DataflowInterface)
+			{
+				TSharedPtr<UE::Dataflow::FContext> ContextPtr = DataflowInterface->GetDataflowContext();
+				Node->OnDoubleClicked(ContextPtr.Get());
+			}
 		}
 	}
 	
diff --git a/Engine/Plugins/Experimental/GeometryCollectionPlugin/Source/GeometryCollectionNodes/Private/Dataflow/GeometryCollectionMeshNodes.cpp b/Engine/Plugins/Experimental/GeometryCollectionPlugin/Source/GeometryCollectionNodes/Private/Dataflow/GeometryCollectionMeshNodes.cpp
--- a/Engine/Plugins/Experimental/GeometryCollectionPlugin/Source/GeometryCollectionNodes/Private/Dataflow/GeometryCollectionMeshNodes.cpp
+++ b/Engine/Plugins/Experimental/GeometryCollectionPlugin/Source/GeometryCollectionNodes/Private/Dataflow/GeometryCollectionMeshNodes.cpp
@@ -589,6 +589,8 @@ FSplitDataflowMeshDataflowNode::FSplitDataflowMeshDataflowNode(const UE::Dataflo
 	RegisterOutputConnection(&MaterialArray);
 }
 
+/////////////////////////////////////////////////////////////////////////////////////////////////
+
 void FSplitDataflowMeshDataflowNode::Evaluate(UE::Dataflow::FContext& Context, const FDataflowOutput* Out) const
 {
 	TObjectPtr<UDynamicMesh> NewMesh = NewObject<UDynamicMesh>();
@@ -612,6 +614,29 @@ void FSplitDataflowMeshDataflowNode::Evaluate(UE::Dataflow::FContext& Context, c
 	}
 }
 
+void FSplitDataflowMeshDataflowNode::OnRenderOutput(UE::Dataflow::FContext& Context, const FName OutputName, const FName RenderGroup, const TArray<UPrimitiveComponent*>& RenderComponents) const
+{
+	if (RenderGroup == TEXT("Surface") && OutputName == TEXT("Mesh"))
+	{
+		if (UDataflowMesh* const InDataflowMesh = GetValue(Context, &InMesh))
+		{
+			const TArray<TObjectPtr<UMaterialInterface>>& Materials = InDataflowMesh->GetMaterials();
+			for (UPrimitiveComponent* Component : RenderComponents)
+			{
+				if (Component)
+				{
+					for (int32 MatIndex = 0; MatIndex < Materials.Num(); ++MatIndex)
+					{
+						Component->SetMaterial(MatIndex, Materials[MatIndex]);
+					}
+				}
+			}
+		}
+	}
+}
+
+/////////////////////////////////////////////////////////////////////////////////////////////////
+
 FDuplicateMeshUVChannelNode::FDuplicateMeshUVChannelNode(const UE::Dataflow::FNodeParameters& InParam, FGuid InGuid)
 	: FDataflowNode(InParam, InGuid)
 {
@@ -980,7 +1005,7 @@ void FMeshProcessorDataflowNodeBase::ApplyParametersToBlueprintInstance(UE::Data
 }
 
 #if WITH_EDITOR
-void FMeshProcessorDataflowNodeBase::OnDoubleClicked() const
+void FMeshProcessorDataflowNodeBase::OnDoubleClicked(UE::Dataflow::FContext* Context) const
 {
 	if (!MeshProcessor)
 	{
diff --git a/Engine/Plugins/Experimental/GeometryCollectionPlugin/Source/GeometryCollectionNodes/Public/Dataflow/GeometryCollectionMeshNodes.h b/Engine/Plugins/Experimental/GeometryCollectionPlugin/Source/GeometryCollectionNodes/Public/Dataflow/GeometryCollectionMeshNodes.h
--- a/Engine/Plugins/Experimental/GeometryCollectionPlugin/Source/GeometryCollectionNodes/Public/Dataflow/GeometryCollectionMeshNodes.h
+++ b/Engine/Plugins/Experimental/GeometryCollectionPlugin/Source/GeometryCollectionNodes/Public/Dataflow/GeometryCollectionMeshNodes.h
@@ -117,7 +117,7 @@ private:
 	}
 
 #if WITH_EDITOR
-	virtual void OnDoubleClicked() const override;
+	virtual void OnDoubleClicked(UE::Dataflow::FContext* Context) const override;
 #endif
 
 	// Handling for the selected blueprint being changed under the node (e.g., recompiled)
@@ -454,7 +454,6 @@ struct FStaticMeshToMeshDataflowNode : public FDataflowNode
 {
 	GENERATED_USTRUCT_BODY()
 	DATAFLOW_NODE_DEFINE_INTERNAL(FStaticMeshToMeshDataflowNode, "StaticMeshToMesh", "Mesh|Utilities", "")
-	DATAFLOW_NODE_RENDER_TYPE("SurfaceRender",FName("FDynamicMesh3"), {"Mesh", "MaterialArray"})
 
 public:
 	/** StaticMesh to convert */
@@ -568,7 +567,6 @@ struct FDataflowMeshAppendDataflowNode : public FDataflowNode
 {
 	GENERATED_USTRUCT_BODY()
 	DATAFLOW_NODE_DEFINE_INTERNAL(FDataflowMeshAppendDataflowNode, "DataflowMeshAppend", "Mesh|Utilities", "")
-	DATAFLOW_NODE_RENDER_TYPE("SurfaceRender", FName("UDataflowMesh"), "Mesh" )
 
 public:
 
@@ -598,7 +596,6 @@ struct FMakeDataflowMeshDataflowNode : public FDataflowNode
 {
 	GENERATED_USTRUCT_BODY()
 	DATAFLOW_NODE_DEFINE_INTERNAL(FMakeDataflowMeshDataflowNode, "MakeDataflowMesh", "Mesh|Utilities", "")
-	DATAFLOW_NODE_RENDER_TYPE("SurfaceRender", FName("UDataflowMesh"), "Mesh")
 
 public:
 
@@ -630,7 +627,6 @@ struct FDuplicateMeshUVChannelNode : public FDataflowNode
 {
 	GENERATED_USTRUCT_BODY()
 	DATAFLOW_NODE_DEFINE_INTERNAL(FDuplicateMeshUVChannelNode, "DuplicateMeshUVChannelNode", "Mesh|Utilities", "Mesh UV DataflowMesh")
-	DATAFLOW_NODE_RENDER_TYPE("SurfaceRender", FName("UDataflowMesh"), "Mesh", "UVChannel")
 
 public:
 
@@ -647,7 +643,7 @@ private:
 	int32 SourceUVChannel = -1;
 
 	/** Index of the added UV channel */
-	UPROPERTY(meta = (DataflowOutput))
+	UPROPERTY(meta = (DataflowOutput, DataflowUVChannel))
 	int32 NewUVChannel = -1;
 
 	virtual void Evaluate(UE::Dataflow::FContext& Context, const FDataflowOutput* Out) const override;
@@ -709,7 +705,6 @@ struct FSplitDataflowMeshDataflowNode : public FDataflowNode
 {
 	GENERATED_USTRUCT_BODY()
 	DATAFLOW_NODE_DEFINE_INTERNAL(FSplitDataflowMeshDataflowNode, "SplitDataflowMesh", "Mesh|Utilities", "")
-	DATAFLOW_NODE_RENDER_TYPE("SurfaceRender", FName("FDynamicMesh3"), { "Mesh", "MaterialArray" })
 
 public:
 
@@ -730,6 +725,8 @@ private:
 	TArray<TObjectPtr<UMaterialInterface>> MaterialArray;
 
 	virtual void Evaluate(UE::Dataflow::FContext& Context, const FDataflowOutput* Out) const override;
+
+	virtual void OnRenderOutput(UE::Dataflow::FContext& Context, const FName OutputName, const FName RenderGroup, const TArray<UPrimitiveComponent*>& RenderComponents) const override;
 };
 
 
diff --git a/Engine/Plugins/Runtime/HairStrands/Source/HairStrandsDataflow/Private/GuidesToJointsNode.cpp b/Engine/Plugins/Runtime/HairStrands/Source/HairStrandsDataflow/Private/GuidesToJointsNode.cpp
--- a/Engine/Plugins/Runtime/HairStrands/Source/HairStrandsDataflow/Private/GuidesToJointsNode.cpp
+++ b/Engine/Plugins/Runtime/HairStrands/Source/HairStrandsDataflow/Private/GuidesToJointsNode.cpp
@@ -2,10 +2,22 @@
 
 #include "GuidesToJointsNode.h"
 
+#include "Framework/Notifications/NotificationManager.h"
 #include "GeometryCollection/Facades/CollectionCurveFacade.h"
 #include "GeometryCollection/Facades/CollectionVertexBoneWeightsFacade.h"
+#include "Rendering/SkeletalMeshLODModel.h"
+#include "Rendering/SkeletalMeshModel.h"
 #include "StaticToSkeletalMeshConverter.h"
 #include "SkeletalMeshAttributes.h"
+#include "Serialization/ObjectWriter.h"
+#include "Serialization/ObjectReader.h"
+#include "SkinnedAssetCompiler.h"
+#include "Widgets/Notifications/SNotificationList.h"
+
+#if WITH_EDITOR
+#include "Editor.h"		// For GEditor
+#include "Subsystems/AssetEditorSubsystem.h"
+#endif
 
 #include UE_INLINE_GENERATED_CPP_BY_NAME(GuidesToJointsNode)
 
@@ -25,13 +37,16 @@ FDataflowGuidesToJointsNode::FDataflowGuidesToJointsNode(const UE::Dataflow::FNo
 {
 	RegisterInputConnection(&Collection);
 	RegisterInputConnection(&CurveSelection);
-	RegisterInputConnection(&RootBone);
+	RegisterInputConnection(&ParentBone);
+	RegisterInputConnection(&BoneBaseName);
 	RegisterInputConnection(&SkeletalMeshAssetPath);
 	RegisterInputConnection(&SkeletonAssetPath);
 	RegisterInputConnection(&BaseSkeletalMesh);
 
 	RegisterOutputConnection(&SkeletalMeshAsset);
 	RegisterOutputConnection(&SkeletonAsset);
+
+	BoneBaseName.Value = TEXT("Curve");
 }
 
 void FDataflowGuidesToJointsNode::Evaluate(UE::Dataflow::FContext& Context) const
@@ -171,33 +186,10 @@ int32 FDataflowGuidesToJointsNode::GetRefSkeletonRootIndex(const FReferenceSkele
 
 void FDataflowGuidesToJointsNode::ClearSkeleton(USkeleton* InOutSkeleton) const
 {
-	// clear the existing skeleton data 
-	const TArray<FMeshBoneInfo>& BoneInfos = InOutSkeleton->GetReferenceSkeleton().GetRefBoneInfo();
-	if (BoneInfos.Num())
-	{
-		// loop so we remove all roots 
-		const int32 RootIndex = GetRefSkeletonRootIndex(InOutSkeleton->GetReferenceSkeleton());
-		if (RootIndex != INDEX_NONE)
-		{
-			// cannot remove the root so we need to get the children and remove them one by one 
-			TArray<int32> ChildBoneIndices;
-			InOutSkeleton->GetReferenceSkeleton().GetDirectChildBones(RootIndex, ChildBoneIndices);
-			TArray<FName> ChildBoneNames;
-			ChildBoneNames.Reserve(ChildBoneIndices.Num());
-			for (const int32 ChildBoneIndex : ChildBoneIndices)
-			{
-				ChildBoneNames.Add(InOutSkeleton->GetReferenceSkeleton().GetBoneName(ChildBoneIndex));
-			}
-
-			{ // scope for modifier 
-				FReferenceSkeletonModifier Edit(InOutSkeleton);
-				for (const FName ChildBoneName : ChildBoneNames)
-				{
-					Edit.Remove(ChildBoneName, /* bRemoveChildren */ true);
-				}
-			}
-		}
-	}
+	// to fully clear a skeleton we need to const cast the reference skeleton 
+	FReferenceSkeleton& RefSkeleton = const_cast<FReferenceSkeleton&>(InOutSkeleton->GetReferenceSkeleton());
+	RefSkeleton.Empty();
+	RefSkeleton.RebuildRefSkeleton(InOutSkeleton, true);
 }
 
 namespace DataflowGuidesToJointsNode::Private
@@ -248,46 +240,47 @@ void FDataflowGuidesToJointsNode::UpdateSkeleton(UE::Dataflow::FContext& Context
 {
 	ClearSkeleton(InOutSkeleton);
 
-	// Get root node , search the base skeletal mesh, if not found create one 
-	const FString BoneToUseAsRoot = GetValue(Context, &RootBone);
+	// Get Parent node , search the base skeletal mesh, if not found create one 
+	const FString InParentBone = GetValue(Context, &ParentBone);
+	const FString InBoneBaseName = GetValue(Context, &BoneBaseName);
 	const TObjectPtr<const USkeletalMesh> InBaseSkeletalMesh = GetValue(Context, &BaseSkeletalMesh);
 
-	int32 RootBoneIndex = INDEX_NONE;
-	FTransform RootBoneTransform = FTransform::Identity;
+	int32 InParentBoneIndex = INDEX_NONE;
+	FTransform InParentBoneTransform = FTransform::Identity;
 	if (IsValid(InBaseSkeletalMesh))
 	{
 		// initialize the skeleton with the ref skeleton
 		const FReferenceSkeleton& RefSkeleton = InBaseSkeletalMesh->GetRefSkeleton();
 		InitializeSkeletonFromRefSkeleton(RefSkeleton, InOutSkeleton);
 
-		RootBoneIndex = RefSkeleton.FindRawBoneIndex(FName(BoneToUseAsRoot));
-		if (RootBoneIndex == INDEX_NONE)
+		InParentBoneIndex = RefSkeleton.FindRawBoneIndex(FName(InParentBone));
+		if (InParentBoneIndex == INDEX_NONE)
 		{
-			const FText InfoMessage = FText::Format(LOCTEXT("InvalidRootBoneName", "Bone {0} does not exists in Skeletal mesh {1}, use the actual root bone if present"), FText::FromString(BoneToUseAsRoot), FText::FromString(InBaseSkeletalMesh->GetPathName()));
+			const FText InfoMessage = FText::Format(LOCTEXT("InvalidParentBoneName", "Bone {0} does not exists in Skeletal mesh {1}, use the actual root bone if present"), FText::FromString(InParentBone), FText::FromString(InBaseSkeletalMesh->GetPathName()));
 			Context.Info(InfoMessage, this);
 		}
 		else
 		{
-			RootBoneTransform = RefSkeleton.GetBoneAbsoluteTransform(RootBoneIndex);
+			InParentBoneTransform = RefSkeleton.GetBoneAbsoluteTransform(InParentBoneIndex);
 		}
 	}
 
 	FReferenceSkeletonModifier Edit(InOutSkeleton);
 
 	// Add a root if we haven't found one yet 
-	if (RootBoneIndex == INDEX_NONE)
+	if (InParentBoneIndex == INDEX_NONE)
 	{
-		RootBoneIndex = GetRefSkeletonRootIndex(InOutSkeleton->GetReferenceSkeleton());
-		if (RootBoneIndex == INDEX_NONE)
+		InParentBoneIndex = GetRefSkeletonRootIndex(InOutSkeleton->GetReferenceSkeleton());
+		if (InParentBoneIndex == INDEX_NONE)
 		{
 			static const FName DefaultRootBoneName(TEXT("root"));
 			FMeshBoneInfo BoneInfo(DefaultRootBoneName, DefaultRootBoneName.ToString(), INDEX_NONE);
 			Edit.Add(BoneInfo, FTransform::Identity, false/*bAllowMultipleRoots*/);
-			RootBoneIndex = Edit.FindBoneIndex(DefaultRootBoneName);
+			InParentBoneIndex = Edit.FindBoneIndex(DefaultRootBoneName);
 		}
 		else
 		{
-			RootBoneTransform = InOutSkeleton->GetReferenceSkeleton().GetBoneAbsoluteTransform(RootBoneIndex);
+			InParentBoneTransform = InOutSkeleton->GetReferenceSkeleton().GetBoneAbsoluteTransform(InParentBoneIndex);
 		}
 	}
 
@@ -311,18 +304,19 @@ void FDataflowGuidesToJointsNode::UpdateSkeleton(UE::Dataflow::FContext& Context
 				CurvesFacade.GetCurvePointPositions(CurveIndex, CurvePoints);
 
 				// Now add to the skeleton
-				FTransform BoneParentTransform = RootBoneTransform;
-				int32 BoneParentIndex = RootBoneIndex;
+				FTransform BoneParentTransform = InParentBoneTransform;
+				int32 BoneParentIndex = InParentBoneIndex;
 				const int32 NumPoints = CurvePoints.Num();
 				for (int32 PointIndex = 0; PointIndex < NumPoints; ++PointIndex)
 				{
-					const FString BoneNameStr = FString::Format(TEXT("Curve_{0}_{1}"), { FString::FromInt(CurveIndex), FString::FromInt(PointIndex) });
+					const FString BoneNameStr = FString::Format(TEXT("{0}_{1}_{2}"), { InBoneBaseName, FString::FromInt(CurveIndex), FString::FromInt(PointIndex) });
+					const FName UniqueBoneName = GenerateUniqueBoneName(FName(BoneNameStr), Edit.GetReferenceSkeleton());
 					const FVector BonePosition = FVector(CurvePoints[PointIndex]);
 					const FTransform AbsoluteBoneTransform(BonePosition);
 					const FTransform LocalBoneTransform = AbsoluteBoneTransform.GetRelativeTransform(BoneParentTransform);
-					FMeshBoneInfo BoneInfo(FName(BoneNameStr), BoneNameStr, BoneParentIndex);
+					FMeshBoneInfo BoneInfo(UniqueBoneName, UniqueBoneName.ToString(), BoneParentIndex);
 					Edit.Add(BoneInfo, LocalBoneTransform, false);
-					BoneParentIndex = Edit.FindBoneIndex(FName(BoneNameStr));
+					BoneParentIndex = Edit.FindBoneIndex(UniqueBoneName);
 					BoneParentTransform = AbsoluteBoneTransform;
 				}
 			}
@@ -330,6 +324,18 @@ void FDataflowGuidesToJointsNode::UpdateSkeleton(UE::Dataflow::FContext& Context
 	}
 }
 
+FName FDataflowGuidesToJointsNode::GenerateUniqueBoneName(FName BaseName, const FReferenceSkeleton& RefSkeleton) const
+{
+	FName UniqueName = BaseName;
+	int32 BoneIndex = RefSkeleton.FindRawBoneIndex(UniqueName);
+	while (BoneIndex != INDEX_NONE)
+	{
+		UniqueName.SetNumber(UniqueName.GetNumber() + 1);
+		BoneIndex = RefSkeleton.FindRawBoneIndex(UniqueName);
+	}
+	return UniqueName;
+}
+
 void FDataflowGuidesToJointsNode::UpdateSkeletalMesh(UE::Dataflow::FContext& Context, USkeletalMesh* InOutSkeletalMesh, USkeleton* InOutSkeleton) const
 {
 	InOutSkeletalMesh->PreEditChange(nullptr);
@@ -338,52 +344,137 @@ void FDataflowGuidesToJointsNode::UpdateSkeletalMesh(UE::Dataflow::FContext& Con
 		// Make sure that we create non-transient assets
 		InOutSkeletalMesh->ClearFlags(RF_Transient);
 
+		bool bIsCopied = false;
+		if (bCopyGeometryFromBaseSkeletalMesh)
+		{
+			if (const TObjectPtr<USkeletalMesh> InBaseSkeletalMesh = GetValue(Context, &BaseSkeletalMesh))
+			{
+				TArray<uint8> Buffer;
+				 
+				FObjectWriter Writer(Buffer);
+				InBaseSkeletalMesh->Serialize(Writer);
+
+				FObjectReader Reader(Buffer);
+				InOutSkeletalMesh->Serialize(Reader);
+
+				InOutSkeletalMesh->SetSkeleton(InBaseSkeletalMesh->GetSkeleton());
+				InOutSkeletalMesh->SetRefSkeleton(InBaseSkeletalMesh->GetRefSkeleton());
+				InOutSkeletalMesh->CalculateInvRefMatrices();
+				InOutSkeletalMesh->SetMaterials(InBaseSkeletalMesh->GetMaterials());
+
+				// make sure tyhe render data is up to date 
+				const FSkeletalMeshModel* SrcModel = InBaseSkeletalMesh->GetImportedModel();
+				FSkeletalMeshModel* DstModel = InOutSkeletalMesh->GetImportedModel();
+				if (SrcModel && DstModel)
+				{
+					DstModel->LODModels.Empty();
+					for (int32 LODIdx = 0; LODIdx < SrcModel->LODModels.Num(); LODIdx++)
+					{
+						FSkeletalMeshLODModel* NewLODModel = new FSkeletalMeshLODModel();
+						FSkeletalMeshLODModel::CopyStructure(NewLODModel, &SrcModel->LODModels[LODIdx]);
+						DstModel->LODModels.Add(NewLODModel);
+					}
+				}
+
+				InOutSkeletalMesh->Build();
+#if WITH_EDITOR
+				FSkinnedAssetCompilingManager::Get().FinishCompilation({ InOutSkeletalMesh });
+#endif
+				bIsCopied = true;
+			}
+		}
+
 		InOutSkeletalMesh->SetSkeleton(InOutSkeleton);
+		InOutSkeletalMesh->SetRefSkeleton(InOutSkeleton->GetReferenceSkeleton());
+		InOutSkeletalMesh->CalculateInvRefMatrices();
 
-		FMeshDescription MeshDescription;
-		const FVertexID VId0 = MeshDescription.CreateVertex();
-		const FVertexID VId1 = MeshDescription.CreateVertex();
-		const FVertexID VId2 = MeshDescription.CreateVertex();
-		const int32 VInst0 = MeshDescription.CreateVertexInstance(VId0);
-		const int32 VInst1 = MeshDescription.CreateVertexInstance(VId1);
-		const int32 VInst2 = MeshDescription.CreateVertexInstance(VId2);
-		TVertexAttributesRef<FVector3f> Positions = MeshDescription.GetVertexPositions();
-		Positions[VId0] = FVector3f(0.0f, 0.0f, 0.0f);
-		Positions[VId1] = FVector3f(0.1f, 0.0f, 0.0f);
-		Positions[VId2] = FVector3f(0.1f, 0.1f, 0.0f);
-		const FPolygonGroupID PolyGroupID = MeshDescription.CreatePolygonGroup();
-		MeshDescription.CreateTriangle(PolyGroupID, { VInst0, VInst1, VInst2 });
-		
-		// add bone weights information
-		FSkeletalMeshAttributes SkeletalMeshAttributes{ MeshDescription };
-		SkeletalMeshAttributes.Register(/*bKeepExistingAttribute*/true);
-		FSkinWeightsVertexAttributesRef VertexSkinWeights = SkeletalMeshAttributes.GetVertexSkinWeights();
-		UE::AnimationCore::FBoneWeights Weights;
-		Weights.SetBoneWeight(0, 1.0f);
-		VertexSkinWeights.Set(VId0, Weights);
-		VertexSkinWeights.Set(VId1, Weights);
-		VertexSkinWeights.Set(VId2, Weights);
-
-		TArray<const FMeshDescription*> MeshDescriptions = { &MeshDescription };
-		// Update the skeletal mesh from the description
-		FStaticToSkeletalMeshConverter::FInitializationParams Parameters;
-		//Parameters.Materials = SkeletalMaterials;
-		Parameters.bRecomputeNormals = true;
-		Parameters.bRecomputeTangents = false;
-
-		FStaticToSkeletalMeshConverter::InitializeSkeletalMeshFromMeshDescriptions(
-			InOutSkeletalMesh,
-			MeshDescriptions,
-			InOutSkeleton->GetReferenceSkeleton(),
-			Parameters);
+		if (!bIsCopied)
+		{
+			FMeshDescription MeshDescription;
+			const FVertexID VId0 = MeshDescription.CreateVertex();
+			const FVertexID VId1 = MeshDescription.CreateVertex();
+			const FVertexID VId2 = MeshDescription.CreateVertex();
+			const int32 VInst0 = MeshDescription.CreateVertexInstance(VId0);
+			const int32 VInst1 = MeshDescription.CreateVertexInstance(VId1);
+			const int32 VInst2 = MeshDescription.CreateVertexInstance(VId2);
+			TVertexAttributesRef<FVector3f> Positions = MeshDescription.GetVertexPositions();
+			Positions[VId0] = FVector3f(0.0f, 0.0f, 0.0f);
+			Positions[VId1] = FVector3f(0.1f, 0.0f, 0.0f);
+			Positions[VId2] = FVector3f(0.1f, 0.1f, 0.0f);
+			const FPolygonGroupID PolyGroupID = MeshDescription.CreatePolygonGroup();
+			MeshDescription.CreateTriangle(PolyGroupID, { VInst0, VInst1, VInst2 });
+
+			// add bone weights information
+			FSkeletalMeshAttributes SkeletalMeshAttributes{ MeshDescription };
+			SkeletalMeshAttributes.Register(/*bKeepExistingAttribute*/true);
+			FSkinWeightsVertexAttributesRef VertexSkinWeights = SkeletalMeshAttributes.GetVertexSkinWeights();
+			UE::AnimationCore::FBoneWeights Weights;
+			Weights.SetBoneWeight(0, 1.0f);
+			VertexSkinWeights.Set(VId0, Weights);
+			VertexSkinWeights.Set(VId1, Weights);
+			VertexSkinWeights.Set(VId2, Weights);
+
+			TArray<const FMeshDescription*> MeshDescriptions = { &MeshDescription };
+			// Update the skeletal mesh from the description
+			FStaticToSkeletalMeshConverter::FInitializationParams Parameters;
+			//Parameters.Materials = SkeletalMaterials;
+			Parameters.bRecomputeNormals = true;
+			Parameters.bRecomputeTangents = false;
+
+			FStaticToSkeletalMeshConverter::InitializeSkeletalMeshFromMeshDescriptions(
+				InOutSkeletalMesh,
+				MeshDescriptions,
+				InOutSkeleton->GetReferenceSkeleton(),
+				Parameters);
+
+		}
 
 		// update the relationship between the asset
 		InOutSkeleton->SetPreviewMesh(InOutSkeletalMesh);
 		InOutSkeleton->RecreateBoneTree(InOutSkeletalMesh);
 	}
+
+	// make sure the new bones are properly referenced in the skeletal mesh LODs
+	if (InOutSkeletalMesh->GetImportedModel())
+	{
+		for (FSkeletalMeshLODModel& LODModel : InOutSkeletalMesh->GetImportedModel()->LODModels)
+		{
+			USkeletalMesh::CalculateRequiredBones(LODModel, InOutSkeletalMesh->GetRefSkeleton(), nullptr);
+		}
+	}
+
 	InOutSkeletalMesh->InvalidateDeriveDataCacheGUID();
-	//InOutSkeletalMesh->InitResources();
+	InOutSkeletalMesh->InitResources();
 	InOutSkeletalMesh->PostEditChange();
+
+	// final rebuild for good measure 
+	InOutSkeletalMesh->Build();
+#if WITH_EDITOR
+	FSkinnedAssetCompilingManager::Get().FinishCompilation({ InOutSkeletalMesh });
+#endif
 }
 
+#if WITH_EDITOR
+void FDataflowGuidesToJointsNode::OnDoubleClicked(UE::Dataflow::FContext* Context) const
+{
+	if (Context)
+	{
+		const FString& InAssetPath = GetValue(*Context, &SkeletalMeshAssetPath);
+		if (TObjectPtr<USkeletalMesh> LoadedSkeletalMesh = ::LoadObject<USkeletalMesh>(nullptr, InAssetPath))
+		{
+			if (UAssetEditorSubsystem* AssetEditorSubsystem = GEditor->GetEditorSubsystem<UAssetEditorSubsystem>())
+			{
+				AssetEditorSubsystem->OpenEditorForAssets({ LoadedSkeletalMesh });
+			}
+		}
+		else
+		{
+			FNotificationInfo NotificationInfo(LOCTEXT("SkeletalMeshEditorFailedNoAsset", "Failed to open SkeletalMesh Editor : No Skeletal Mesh Asset"));
+			NotificationInfo.ExpireDuration = 5.0f;
+			FSlateNotificationManager::Get().AddNotification(NotificationInfo);
+		}
+	}
+}
+#endif
+
 #undef LOCTEXT_NAMESPACE
\ No newline at end of file
diff --git a/Engine/Plugins/Runtime/HairStrands/Source/HairStrandsDataflow/Private/GuidesToJointsNode.h b/Engine/Plugins/Runtime/HairStrands/Source/HairStrandsDataflow/Private/GuidesToJointsNode.h
--- a/Engine/Plugins/Runtime/HairStrands/Source/HairStrandsDataflow/Private/GuidesToJointsNode.h
+++ b/Engine/Plugins/Runtime/HairStrands/Source/HairStrandsDataflow/Private/GuidesToJointsNode.h
@@ -28,6 +28,12 @@ public:
 	FDataflowGuidesToJointsNode(const UE::Dataflow::FNodeParameters& InParam, FGuid InGuid = FGuid::NewGuid());
 
 private:
+	//~ Begin FDataflowNode interface
+#if WITH_EDITOR
+	virtual void OnDoubleClicked(UE::Dataflow::FContext* Context) const override;
+#endif
+	//~ End FDataflowNode interface
+
 	//~ Begin FDataflowTerminalNode interface
 	virtual void Evaluate(UE::Dataflow::FContext& Context) const override;
 	virtual void SetAssetValue(TObjectPtr<UObject> Asset, UE::Dataflow::FContext& Context) const override;
@@ -39,19 +45,12 @@ private:
 	int32 GetRefSkeletonRootIndex(const FReferenceSkeleton& RefSkeleton) const;
 
 	void UpdateSkeletalMesh(UE::Dataflow::FContext& Context, USkeletalMesh* InOutSkeletalMesh, USkeleton* InOutSkeleton) const;
+	FName GenerateUniqueBoneName(FName BaseName, const FReferenceSkeleton& RefSkeleton) const;
 
-	/** Managed array collection source of the guides */
-	UPROPERTY(meta = (DataflowInput, DataflowRenderGroups = "Surface"))
+	/** Managed array collection source of the guides - Required */
+	UPROPERTY(meta = (DataflowInput, DataflowIntrinsic, DataflowRenderGroups = "Surface"))
 	FManagedArrayCollection Collection;
 
-	/** Curve selection to use to generate bones/joints - if none specified all curve will be considered */
-	UPROPERTY(meta = (DataflowInput))
-	FDataflowCurveSelection CurveSelection;
-
-	/** Optional name of root bone from the input skeletal mesh to use as the root of the generate bones */
-	UPROPERTY(EditAnywhere, Category = Source, meta = (DataflowInput))
-	FDataflowStringTypes RootBone;
-
 	/** Path to the skeletal mesh asset to create or update */
 	UPROPERTY(EditAnywhere, Category = Asset, meta = (DataflowInput))
 	FString SkeletalMeshAssetPath;
@@ -60,9 +59,25 @@ private:
 	UPROPERTY(EditAnywhere, Category = Asset, meta = (DataflowInput))
 	FString SkeletonAssetPath;
 
-	/** Skeletal Mesh to use a base*/
+	/** Curve selection to use to generate bones/joints - if none specified all curve will be considered */
+	UPROPERTY(meta = (DataflowInput))
+	FDataflowCurveSelection CurveSelection;
+
+	/** Base name for the bones created from the curve */
 	UPROPERTY(EditAnywhere, Category = Asset, meta = (DataflowInput))
-	TObjectPtr<const USkeletalMesh> BaseSkeletalMesh = nullptr;
+	FDataflowStringTypes BoneBaseName;
+
+	/** Optional name of parent bone from the input skeletal mesh to use as the root of the generate bones */
+	UPROPERTY(EditAnywhere, Category = Source, meta = (DataflowInput))
+	FDataflowStringTypes ParentBone;
+
+	/** Skeletal Mesh to use a base*/
+	UPROPERTY(EditAnywhere, Category = Base, meta = (DataflowInput))
+	TObjectPtr<USkeletalMesh> BaseSkeletalMesh = nullptr;
+
+	/** Option to copy the geometry from the base skeletal mesh */
+	UPROPERTY(EditAnywhere, Category = Base)
+	bool bCopyGeometryFromBaseSkeletalMesh = false;
 
 	/** Created Skeletal Mesh */
 	UPROPERTY(meta = (DataflowOutput))
diff --git a/Engine/Source/Runtime/Experimental/Dataflow/Core/Public/Dataflow/DataflowNode.h b/Engine/Source/Runtime/Experimental/Dataflow/Core/Public/Dataflow/DataflowNode.h
--- a/Engine/Source/Runtime/Experimental/Dataflow/Core/Public/Dataflow/DataflowNode.h
+++ b/Engine/Source/Runtime/Experimental/Dataflow/Core/Public/Dataflow/DataflowNode.h
@@ -130,7 +130,7 @@ PRAGMA_ENABLE_DEPRECATION_WARNINGS
 		IDataflowDebugDrawInterface& DataflowRenderingInterface,
 		const FDebugDrawParameters& DebugDrawParameters) const {}
 
-	virtual void OnDoubleClicked() const {};
+	virtual void OnDoubleClicked(UE::Dataflow::FContext* Context) const {};
 #endif
 
 	//
