Change hair scene to store hair instances as ref. counted reference to avoid use after crashes.

#Synced-CL 51662140

diff --git a/Engine/Plugins/Runtime/HairStrands/Source/HairStrandsCore/Private/GroomComponent.cpp b/Engine/Plugins/Runtime/HairStrands/Source/HairStrandsCore/Private/GroomComponent.cpp
--- a/Engine/Plugins/Runtime/HairStrands/Source/HairStrandsCore/Private/GroomComponent.cpp
+++ b/Engine/Plugins/Runtime/HairStrands/Source/HairStrandsCore/Private/GroomComponent.cpp
@@ -627,14 +627,15 @@ public:
 	{
 		FPrimitiveSceneProxy::CreateRenderThreadResources(RHICmdList);
 
-		// Register the data to the scene
+		// Register the data to the scene. AddHairStrands takes a raw pointer 
+		// but internally registers it as TRefCountPtr which keeps the instance 
+		// alive while it is registered in the scene.
 		FSceneInterface& LocalScene = GetScene();
 		for (TRefCountPtr<FHairGroupInstance>& Instance : HairGroupInstances)
 		{
 			if (Instance->IsValid() || Instance->Strands.ClusterResource)
 			{
 				check(Instance->HairGroupPublicData != nullptr);
-				Instance->AddRef();
 				Instance->Debug.Proxy = this;
 				LocalScene.AddHairStrands(Instance);
 			}
@@ -650,7 +651,9 @@ public:
 	{
 		FPrimitiveSceneProxy::DestroyRenderThreadResources();
 
-		// Unregister the data to the scene
+		// Unregister the data from the scene. RemoveHairStrands releases the 
+		// scene's TRefCountPtr, which may destroy the instance if this was 
+		// the last reference.
 		FSceneInterface& LocalScene = GetScene();
 		for (TRefCountPtr<FHairGroupInstance>& Instance : HairGroupInstances)
 		{
@@ -658,7 +661,6 @@ public:
 			{
 				LocalScene.RemoveHairStrands(Instance);
 				Instance->Debug.Proxy = nullptr;
-				Instance->Release();
 			}
 		}
 	}
diff --git a/Engine/Plugins/Runtime/HairStrands/Source/HairStrandsCore/Private/GroomDebug.cpp b/Engine/Plugins/Runtime/HairStrands/Source/HairStrandsCore/Private/GroomDebug.cpp
--- a/Engine/Plugins/Runtime/HairStrands/Source/HairStrandsCore/Private/GroomDebug.cpp
+++ b/Engine/Plugins/Runtime/HairStrands/Source/HairStrandsCore/Private/GroomDebug.cpp
@@ -823,7 +823,7 @@ static void AddHairDebugPrintInstancePass(
 	Infos.Reserve(InstanceCount);
 	for (uint32 InstanceIndex = 0; InstanceIndex < InstanceCount; ++InstanceIndex)
 	{
-		const FHairGroupInstance* Instance = static_cast<const FHairGroupInstance*>(Instances[InstanceIndex]);
+		const FHairGroupInstance* Instance = static_cast<const FHairGroupInstance*>(Instances[InstanceIndex].GetReference());
 
 		// Collect InstanceNames
 		InstanceNames.Add(*Instance->Debug.GroomAssetName, InstanceIndex);
diff --git a/Engine/Plugins/Runtime/HairStrands/Source/HairStrandsCore/Private/GroomManager.cpp b/Engine/Plugins/Runtime/HairStrands/Source/HairStrandsCore/Private/GroomManager.cpp
--- a/Engine/Plugins/Runtime/HairStrands/Source/HairStrandsCore/Private/GroomManager.cpp
+++ b/Engine/Plugins/Runtime/HairStrands/Source/HairStrandsCore/Private/GroomManager.cpp
@@ -3505,9 +3505,9 @@ void ProcessHairStrandsBookmark(
 			DebugInstances.Append(Parameters.CullingResults.VisibleInAnyViews_Strands);
 			DebugInstances.Append(Parameters.CullingResults.VisibleInAnyViews_CardsOrMeshes_Primary);
 			DebugInstances.Append(Parameters.CullingResults.VisibleInAnyViews_CardsOrMeshes_Shadow);
-			DebugInstances.Sort([](const FHairStrandsInstance& A, const FHairStrandsInstance& B)
+			DebugInstances.Sort([](const TRefCountPtr<FHairStrandsInstance>& A, const TRefCountPtr<FHairStrandsInstance>& B)
 			{
-				return A.RegisteredIndex < B.RegisteredIndex;
+				return A->RegisteredIndex < B->RegisteredIndex;
 			});
 			Instances = &DebugInstances;
 		}
diff --git a/Engine/Source/Runtime/Renderer/Private/RendererScene.cpp b/Engine/Source/Runtime/Renderer/Private/RendererScene.cpp
--- a/Engine/Source/Runtime/Renderer/Private/RendererScene.cpp
+++ b/Engine/Source/Runtime/Renderer/Private/RendererScene.cpp
@@ -2732,7 +2732,7 @@ void FScene::AddHairStrands(FHairStrandsInstance* Proxy)
 {
 	if (Proxy)
 	{
-		const int32 PackedIndex = HairStrandsSceneData.RegisteredProxies.Add(Proxy);
+		const int32 PackedIndex = HairStrandsSceneData.RegisteredProxies.Add(TRefCountPtr<FHairStrandsInstance>(Proxy));
 		Proxy->RegisteredIndex = PackedIndex;
 	}
 }
@@ -2749,8 +2749,7 @@ void FScene::RemoveHairStrands(FHairStrandsInstance* Proxy)
 		Proxy->RegisteredIndex = -1;
 		if (HairStrandsSceneData.RegisteredProxies.IsValidIndex(ProxyIndex))
 		{
-			FHairStrandsInstance* Other = HairStrandsSceneData.RegisteredProxies[ProxyIndex];
-			Other->RegisteredIndex = ProxyIndex;
+			HairStrandsSceneData.RegisteredProxies[ProxyIndex]->RegisteredIndex = ProxyIndex;
 		}
 	}
 }
diff --git a/Engine/Source/Runtime/Renderer/Private/ScenePrivate.h b/Engine/Source/Runtime/Renderer/Private/ScenePrivate.h
--- a/Engine/Source/Runtime/Renderer/Private/ScenePrivate.h
+++ b/Engine/Source/Runtime/Renderer/Private/ScenePrivate.h
@@ -1862,7 +1862,10 @@ public:
 class FHairStrandsSceneData
 {
 public:
-	TArray<FHairStrandsInstance*> RegisteredProxies;
+	// Ref-counted to guarantee instances stay alive while registered in the scene,
+	// preventing use-after-free when the game thread destroys groom components
+	// while the render thread is still accessing instances (e.g., during Sequencer scrubbing).
+	TArray<TRefCountPtr<FHairStrandsInstance>> RegisteredProxies;
 	FHairTransientResources* TransientResources = nullptr;
 };
 
diff --git a/Engine/Source/Runtime/Renderer/Public/HairStrandsInterface.h b/Engine/Source/Runtime/Renderer/Public/HairStrandsInterface.h
--- a/Engine/Source/Runtime/Renderer/Public/HairStrandsInterface.h
+++ b/Engine/Source/Runtime/Renderer/Public/HairStrandsInterface.h
@@ -229,7 +229,7 @@ struct FHairStrandsInstance : public FRefCountedObject
 	virtual const FHairGroupPublicData* GetHairData() const { return nullptr; }
 	virtual const EHairGeometryType GetHairGeometry() const { return EHairGeometryType::NoneGeometry; }
 };
-typedef TArray<FHairStrandsInstance*> FHairStrandsInstances;
+typedef TArray<TRefCountPtr<FHairStrandsInstance>> FHairStrandsInstances;
 
 class FHairGroupPublicData
 {
