Fix race condition in FHairStrandsSceneProxy::CreateMeshBatch().

Disable by default parallel GetDynamicMeshElements for groom proxies (bSupportsParallelGDME = false).
When a scene capture and the main view render concurrently, the same FHairStrandsSceneProxy can have
CreateMeshBatch called from two threads, racing on shared FSceneVelocityData accessed via BuildUniformShaderParameters.

This can be reenabled with the CVar r.HairStrands.UseParallelGDME

#Synced-CL 51662140

diff --git a/Engine/Plugins/Runtime/HairStrands/Source/HairStrandsCore/Private/GroomComponent.cpp b/Engine/Plugins/Runtime/HairStrands/Source/HairStrandsCore/Private/GroomComponent.cpp
--- a/Engine/Plugins/Runtime/HairStrands/Source/HairStrandsCore/Private/GroomComponent.cpp
+++ b/Engine/Plugins/Runtime/HairStrands/Source/HairStrandsCore/Private/GroomComponent.cpp
@@ -81,6 +81,11 @@ static FAutoConsoleVariableRef CVarHairStrands_ViewModeClumpIndex(TEXT("r.HairSt
 static int32 GHairStrands_ForceVelocityOutput = 0;
 static FAutoConsoleVariableRef CVarHairStrands_ForceVelocityOutput(TEXT("r.HairStrands.ForceVelocityOutput"), GHairStrands_ForceVelocityOutput, TEXT("When enabled, force the cards/meshes to write velocity vectors."));
 
+// By default GDME is disabled as it causess a (begnin) race condition (in CreateMeshBatch > BuildUniformShaderParameters) when reading velocity data, as the LastFrameIndex is mutable and can be read/writen at the same time.
+// This can happens when view context (e.g., main view + scene capture) are updated at the same time.
+static int32 GHairStrands_UseParallelGDME = 0;
+static FAutoConsoleVariableRef CVarHairStrands_UseParallelGDME(TEXT("r.HairStrands.UseParallelGDME"), GHairStrands_UseParallelGDME, TEXT("When enabled, FHairStrandsSceneProxy uses parallel GDME (using async task gathering) for collecting instances to render."));
+
 #define LOCTEXT_NAMESPACE "GroomComponent"
 
 ////////////////////////////////////////////////////////////////////////////////////////////////////////////////////////
@@ -460,6 +465,7 @@ public:
 		// Forcing primitive uniform as we don't support robustly GPU scene data
 		bVFRequiresPrimitiveUniformBuffer = true;
 		bCastDeepShadow = true;
+		bSupportsParallelGDME = GHairStrands_UseParallelGDME>0;
 
 		HairGroupMaterialProxies.SetNum(Component->HairGroupInstances.Num());
 		HairGroupInstances = Component->HairGroupInstances;
