Fix groom binding builder distance computation with degenerated triangle causing groom binding to fail.

#Synced-CL 51662140

diff --git a/Engine/Plugins/Runtime/HairStrands/Source/HairStrandsCore/Private/GroomBindingBuilder.cpp b/Engine/Plugins/Runtime/HairStrands/Source/HairStrandsCore/Private/GroomBindingBuilder.cpp
--- a/Engine/Plugins/Runtime/HairStrands/Source/HairStrandsCore/Private/GroomBindingBuilder.cpp
+++ b/Engine/Plugins/Runtime/HairStrands/Source/HairStrandsCore/Private/GroomBindingBuilder.cpp
@@ -1298,16 +1298,14 @@ namespace GroomBinding_RootProjection
 			return Out;
 		}
 
+		// Reject degenerate (near-zero area) triangles via cross product to avoid NaN 
+		// from division by zero in ComputeClosestPoint's barycentric computation.
 		bool IsTriangleValid(const FTriangle& T) const
 		{
-			const FVector3f A = T.P0;
-			const FVector3f B = T.P1;
-			const FVector3f C = T.P2;
-
-			const FVector3f AB = B - A;
-			const FVector3f AC = C - A;
-			const FVector3f BC = B - C;
-			return FVector3f::DotProduct(AB, AB) > 0 && FVector3f::DotProduct(AC, AC) > 0 && FVector3f::DotProduct(BC, BC) > 0;
+			const FVector3f AB = T.P1 - T.P0;
+			const FVector3f AC = T.P2 - T.P0;
+			const FVector3f Cross = FVector3f::CrossProduct(AB, AC);
+			return FVector3f::DotProduct(Cross, Cross) > UE_SMALL_NUMBER;
 		}
 
 		bool Insert(const FTriangle& T)
@@ -1449,7 +1447,36 @@ namespace GroomBinding_RootProjection
 		}
 
 		// P must be inside the face region. Compute the closest point through its barycentric coordinates (u,V,W).
-		float Denom = 1.f / (VA + VB + VC);
+		const float DenomSum = VA + VB + VC;
+
+		// Guard against near-zero denominator from degenerate triangles to prevent NaN propagation.
+		// Fallback with picking the closest point
+		if (DenomSum <= UE_SMALL_NUMBER)
+		{
+			const float DistA = FVector3f::DistSquared(P, A);
+			const float DistB = FVector3f::DistSquared(P, B);
+			const float DistC = FVector3f::DistSquared(P, C);
+
+			FTrianglePoint Out;
+			if (DistA <= DistB && DistA <= DistC)
+			{
+				Out.P = A;
+				Out.Barycentric = FVector3f(1, 0, 0);
+			}
+			else if (DistB <= DistC)
+			{
+				Out.P = B;
+				Out.Barycentric = FVector3f(0, 1, 0);
+			}
+			else
+			{
+				Out.P = C;
+				Out.Barycentric = FVector3f(0, 0, 1);
+			}
+			return Out;
+		}
+
+		const float Denom = 1.f / DenomSum;
 		float V = VB * Denom;
 		float W = VC * Denom;
 
@@ -1490,11 +1517,14 @@ namespace GroomBinding_RootProjection
 					}
 				}
 			}
-			check(ClosestDistance < FLT_MAX);
 
-			Out.Triangle = ClosestTriangle;
-			Out.Point = ClosestPointOnTriangle;
-			Out.bValid = true;
+			// Gracefully handle NaN distances to report an error via the existing bValid==false path.
+			if (ClosestDistance < FLT_MAX)
+			{
+				Out.Triangle = ClosestTriangle;
+				Out.Point = ClosestPointOnTriangle;
+				Out.bValid = true;
+			}
 		}
 
 		return Out;
@@ -1757,12 +1787,15 @@ namespace GroomBinding_RootProjection
 					const FVector3f P = LocalMeshLODData->GetVertexPosition(VertexIndex);
 					const FQueryResult Q = QueryClosestPoint(*TriangleGrid, P);
 
-					const float InterpolatedAttribute = 
-						Q.Triangle.Attribute0 * Q.Point.Barycentric.X +
-						Q.Triangle.Attribute1 * Q.Point.Barycentric.Y +
-						Q.Triangle.Attribute2 * Q.Point.Barycentric.Z;
+					if (Q.bValid)
+					{
+						const float InterpolatedAttribute = 
+							Q.Triangle.Attribute0 * Q.Point.Barycentric.X +
+							Q.Triangle.Attribute1 * Q.Point.Barycentric.Y +
+							Q.Triangle.Attribute2 * Q.Point.Barycentric.Z;
 
-					LocalMeshLODData->SetVertexBinding(VertexIndex,  InterpolatedAttribute);
+						LocalMeshLODData->SetVertexBinding(VertexIndex,  InterpolatedAttribute);
+					}
 				});
 			}
 		}
@@ -2251,7 +2284,36 @@ namespace GroomBinding_Transfer
 		}
 
 		// P must be inside the face region. Compute the closest point through its barycentric coordinates (u,V,W).
-		float Denom = 1.f / (VA + VB + VC);
+		const float DenomSum = VA + VB + VC;
+
+		// Guard against near-zero denominator from degenerate triangles to prevent NaN propagation.
+		// Fallback with picking the closest point
+		if (DenomSum <= UE_SMALL_NUMBER)
+		{
+			const float DistA = FVector3f::DistSquared(P, A);
+			const float DistB = FVector3f::DistSquared(P, B);
+			const float DistC = FVector3f::DistSquared(P, C);
+
+			FTrianglePoint Out;
+			if (DistA <= DistB && DistA <= DistC)
+			{
+				Out.P = A;
+				Out.Barycentric = FVector3f(1, 0, 0);
+			}
+			else if (DistB <= DistC)
+			{
+				Out.P = B;
+				Out.Barycentric = FVector3f(0, 1, 0);
+			}
+			else
+			{
+				Out.P = C;
+				Out.Barycentric = FVector3f(0, 0, 1);
+			}
+			return Out;
+		}
+
+		const float Denom = 1.f / DenomSum;
 		float V = VB * Denom;
 		float W = VC * Denom;
 
