Add GetHairDescriptionCustomVersions to make common version tracking (used by FHairDescription and FHairDescriptionBulkData).

#Synced-CL 51755560

diff --git a/Engine/Plugins/Runtime/HairStrands/Source/HairStrandsCore/Private/HairDescription.cpp b/Engine/Plugins/Runtime/HairStrands/Source/HairStrandsCore/Private/HairDescription.cpp
--- a/Engine/Plugins/Runtime/HairStrands/Source/HairStrandsCore/Private/HairDescription.cpp
+++ b/Engine/Plugins/Runtime/HairStrands/Source/HairStrandsCore/Private/HairDescription.cpp
@@ -11,6 +11,13 @@
 const FStrandID FStrandID::Invalid(TNumericLimits<uint32>::Max());
 const FGroomID FGroomID::Invalid(TNumericLimits<uint32>::Max());
 
+static TConstArrayView<FGuid> GetHairDescriptionCustomVersions()
+{
+	// Use FMeshDescriptionBulkData's versions as HairDescription uses structures and functions 
+	// from FMeshDescription(BulkData) which are subject to serizalization changes
+	return FMeshDescriptionBulkData::GetMeshDescriptionCustomVersions();
+}
+
 ///////////////////////////////////////////////////////////////////////////////////////////////////
 // FHairDescription
 
@@ -107,10 +114,10 @@ bool FHairDescription::HasAOAttributes() const { return HasAttribute(EHairAttrib
 
 void FHairDescription::Serialize(FArchive& Ar)
 {
-	Ar.UsingCustomVersion(FEditorObjectVersion::GUID);
-	Ar.UsingCustomVersion(FReleaseObjectVersion::GUID);
-	Ar.UsingCustomVersion(FUE5MainStreamObjectVersion::GUID);
-	Ar.UsingCustomVersion(FFortniteMainBranchObjectVersion::GUID);
+	for (const FGuid& CustomVersionGuid : GetHairDescriptionCustomVersions())
+	{
+		Ar.UsingCustomVersion(CustomVersionGuid);
+	}
 
 	Ar << NumVertices;
 	Ar << NumStrands;
@@ -158,9 +165,10 @@ void FHairDescriptionVersion::CopyVersionsToArchive(FArchive& Ar) const
 #if WITH_EDITORONLY_DATA
 void FHairDescriptionBulkData::Serialize(FArchive& Ar, UObject* Owner)
 {
-	Ar.UsingCustomVersion(FEditorObjectVersion::GUID);
-	Ar.UsingCustomVersion(FReleaseObjectVersion::GUID);
-	Ar.UsingCustomVersion(FUE5MainStreamObjectVersion::GUID);
+	for (const FGuid& CustomVersionGuid : GetHairDescriptionCustomVersions())
+	{
+		Ar.UsingCustomVersion(CustomVersionGuid);
+	}
 
 	if (Ar.IsTransacting())
 	{
