Fix CIS error.

#Synced-CL 51755560

diff --git a/Engine/Plugins/Runtime/HairStrands/Source/HairStrandsCore/Private/HairDescription.cpp b/Engine/Plugins/Runtime/HairStrands/Source/HairStrandsCore/Private/HairDescription.cpp
--- a/Engine/Plugins/Runtime/HairStrands/Source/HairStrandsCore/Private/HairDescription.cpp
+++ b/Engine/Plugins/Runtime/HairStrands/Source/HairStrandsCore/Private/HairDescription.cpp
@@ -11,12 +11,14 @@
 const FStrandID FStrandID::Invalid(TNumericLimits<uint32>::Max());
 const FGroomID FGroomID::Invalid(TNumericLimits<uint32>::Max());
 
+#if WITH_EDITORONLY_DATA
 static TConstArrayView<FGuid> GetHairDescriptionCustomVersions()
 {
 	// Use FMeshDescriptionBulkData's versions as HairDescription uses structures and functions 
 	// from FMeshDescription(BulkData) which are subject to serizalization changes
 	return FMeshDescriptionBulkData::GetMeshDescriptionCustomVersions();
 }
+#endif
 
 ///////////////////////////////////////////////////////////////////////////////////////////////////
 // FHairDescription
@@ -114,10 +116,10 @@ bool FHairDescription::HasAOAttributes() const { return HasAttribute(EHairAttrib
 
 void FHairDescription::Serialize(FArchive& Ar)
 {
-	for (const FGuid& CustomVersionGuid : GetHairDescriptionCustomVersions())
-	{
-		Ar.UsingCustomVersion(CustomVersionGuid);
-	}
+	Ar.UsingCustomVersion(FEditorObjectVersion::GUID);
+	Ar.UsingCustomVersion(FReleaseObjectVersion::GUID);
+	Ar.UsingCustomVersion(FUE5MainStreamObjectVersion::GUID);
+	Ar.UsingCustomVersion(FFortniteMainBranchObjectVersion::GUID);
 
 	Ar << NumVertices;
 	Ar << NumStrands;
