Ensure FHairInterpolationCS's wave ops path is only picked at runtime if WaveSize 32 is supported.

The existing check only looks at the supported wave size range from static platform data when compiling the permutations, which is correct, but that check cannot be used at runtime on PC because the actual wave size support is dependent on the GPU.

#rnx
#Synced-CL 51702310

diff --git a/Engine/Plugins/Runtime/HairStrands/Source/HairStrandsCore/Private/HairStrandsInterpolation.cpp b/Engine/Plugins/Runtime/HairStrands/Source/HairStrandsCore/Private/HairStrandsInterpolation.cpp
--- a/Engine/Plugins/Runtime/HairStrands/Source/HairStrandsCore/Private/HairStrandsInterpolation.cpp
+++ b/Engine/Plugins/Runtime/HairStrands/Source/HairStrandsCore/Private/HairStrandsInterpolation.cpp
@@ -568,7 +568,8 @@ class FHairInterpolationCS : public FGlobalShader
 
 public:
 	static uint32 GetGroupSize() { return 32;}
-	static bool DoesSupportsWaveOps(EShaderPlatform InPlatform, uint32 InPointPerCurve)
+
+	static bool DoesPlatformSupportWaveOps(EShaderPlatform InPlatform, uint32 InPointPerCurve)
 	{
 		// D3D11 / SM5 or preview do not support, or work well with, wave-ops by default (or SM5 preview has issues with wave intrinsics too)
 		if (InPointPerCurve != 32 || InPlatform == SP_PCD3D_SM5 || FDataDrivenShaderPlatformInfo::GetIsPreviewPlatform(InPlatform))
@@ -577,10 +578,17 @@ public:
 		}
 		return FDataDrivenShaderPlatformInfo::GetSupportsWaveOperations(InPlatform) != ERHIFeatureSupport::Unsupported && FDataDrivenShaderPlatformInfo::GetMinimumWaveSize(InPlatform) <= 32 && FDataDrivenShaderPlatformInfo::GetMaximumWaveSize(InPlatform) >= 32;
 	}
+
+	// Runtime check for platforms that have varying wave size support depending on GPU.
+	static bool DoesSupportWaveOps(EShaderPlatform InPlatform, uint32 InPointPerCurve)
+	{
+		return DoesPlatformSupportWaveOps(InPlatform, InPointPerCurve) && GRHIMinimumWaveSize <= 32 && GRHIMaximumWaveSize >= 32;
+	}
+
 	static bool ShouldCompilePermutation(const FGlobalShaderPermutationParameters& Parameters) 
 	{ 
 		FPermutationDomain PermutationVector(Parameters.PermutationId);
-		if (PermutationVector.Get<FWaveOps>() && !DoesSupportsWaveOps(Parameters.Platform, PermutationVector.Get<FPointPerCurve>()))
+		if (PermutationVector.Get<FWaveOps>() && !DoesPlatformSupportWaveOps(Parameters.Platform, PermutationVector.Get<FPointPerCurve>()))
 		{
 			return false;
 		}
@@ -770,7 +778,7 @@ void AddHairStrandsInterpolationPass(
 	Parameters->DispatchCountX = DispatchInfo.DispatchCount.X;
 
 	const bool bSingleGuidePermutation = bUseSingleGuide || GetHairStrandsForceSingleGuideInterpolation();
-	const bool bWaveOps = FHairInterpolationCS::DoesSupportsWaveOps(InPlatform, DispatchInfo.PointPerCurve);
+	const bool bWaveOps = FHairInterpolationCS::DoesSupportWaveOps(InPlatform, DispatchInfo.PointPerCurve);
 
 	// Deformer type : 0 no deformer, 1 add deformer displacement, 2 overwrite deformed position with the deformer one
 	const int32 DeformerType = !bSupportDeformer ? 0 : GVarHairStrands_AddDeformerStrandsDisplacement ? 1 : 2;
