Add missing include for shader compiler flags

#rnx
#Synced-CL 51702310

diff --git a/Engine/Plugins/Runtime/HairStrands/Source/HairStrandsCore/Private/HairStrandsInterpolation.cpp b/Engine/Plugins/Runtime/HairStrands/Source/HairStrandsCore/Private/HairStrandsInterpolation.cpp
--- a/Engine/Plugins/Runtime/HairStrands/Source/HairStrandsCore/Private/HairStrandsInterpolation.cpp
+++ b/Engine/Plugins/Runtime/HairStrands/Source/HairStrandsCore/Private/HairStrandsInterpolation.cpp
@@ -27,6 +27,7 @@
 #include "GroomManager.h"
 #include "GroomInstance.h"
 #include "SystemTextures.h"
+#include "ShaderCompilerCore.h"
 #include "ShaderPrintParameters.h"
 #include "HairStrandsClusterCulling.h"
 #include "DataDrivenShaderPlatformInfo.h"
