Groom dataflow : fix issue wit the GuidesTojoints terminal node not properly creating a valid asset name for the Skeleton when the input was not plugged

#Synced-CL 51755560

diff --git a/Engine/Plugins/Runtime/HairStrands/Source/HairStrandsDataflow/Private/GuidesToJointsNode.cpp b/Engine/Plugins/Runtime/HairStrands/Source/HairStrandsDataflow/Private/GuidesToJointsNode.cpp
--- a/Engine/Plugins/Runtime/HairStrands/Source/HairStrandsDataflow/Private/GuidesToJointsNode.cpp
+++ b/Engine/Plugins/Runtime/HairStrands/Source/HairStrandsDataflow/Private/GuidesToJointsNode.cpp
@@ -117,7 +117,23 @@ void FDataflowGuidesToJointsNode::SetAssetValue(TObjectPtr<UObject> Asset, UE::D
 	//-------------------------------------------------------------------
 	// skeleton section
 	//-------------------------------------------------------------------
-	const FString SkeletonPathFromInput = GetValue(Context, &SkeletonAssetPath);
+	FString SkeletonPathFromInput = GetValue(Context, &SkeletonAssetPath);
+	if (SkeletonPathFromInput.IsEmpty())
+	{
+		// use the skeletal mesh name 
+		SkeletonPathFromInput = GetValue(Context, &SkeletalMeshAssetPath);
+		if (!SkeletonPathFromInput.IsEmpty())
+		{
+			FString PathPart;
+			FString FilenamePart;
+			FString ExtensionPart;
+			FPaths::Split(SkeletonPathFromInput, PathPart, FilenamePart, ExtensionPart);
+			const TCHAR* PostFix  = TEXT("_Skeleton");
+			FilenamePart += PostFix;
+			ExtensionPart += PostFix;
+			SkeletonPathFromInput = FPaths::SetExtension(PathPart / FilenamePart, ExtensionPart);
+		}
+	}
 	const FString& InSkeletonAssetPath = GetAssetPath(SkeletonPathFromInput, Asset);
 	if (InSkeletonAssetPath.IsEmpty())
 	{
