Fix incorrect hair offset in streaming request.

#Synced-CL 51816243

diff --git a/Engine/Plugins/Runtime/HairStrands/Source/HairStrandsCore/Private/HairStrandsDatas.cpp b/Engine/Plugins/Runtime/HairStrands/Source/HairStrandsCore/Private/HairStrandsDatas.cpp
--- a/Engine/Plugins/Runtime/HairStrands/Source/HairStrandsCore/Private/HairStrandsDatas.cpp
+++ b/Engine/Plugins/Runtime/HairStrands/Source/HairStrandsCore/Private/HairStrandsDatas.cpp
@@ -296,7 +296,12 @@ void FHairStrandsBulkCommon::FQuery::Add(FHairBulkContainer& In, const TCHAR* In
 			InOffset = 0;
 			InSize = In.Data.GetBulkDataSize();
 		}
-		check (InSize >= InOffset);
+
+		// If the offset is incorrect (tracking issue?) reload the whole requested size
+		if (InSize < InOffset)
+		{
+			InOffset = 0;
+		}
 
 		check(StreamingRequest);
 
