Expose time dilation to niagara groom DI

#Synced-CL 51879956

diff --git a/Engine/Plugins/Runtime/HairStrands/Source/HairStrandsCore/Private/Niagara/NiagaraDataInterfaceHairStrands.cpp b/Engine/Plugins/Runtime/HairStrands/Source/HairStrandsCore/Private/Niagara/NiagaraDataInterfaceHairStrands.cpp
--- a/Engine/Plugins/Runtime/HairStrands/Source/HairStrandsCore/Private/Niagara/NiagaraDataInterfaceHairStrands.cpp
+++ b/Engine/Plugins/Runtime/HairStrands/Source/HairStrandsCore/Private/Niagara/NiagaraDataInterfaceHairStrands.cpp
@@ -26,6 +26,8 @@
 #include "GroomBindingBuilder.h"
 #include "SystemTextures.h"
 #include "ShaderCompilerCore.h"
+#include "Engine/World.h"
+#include "GameFramework/WorldSettings.h"
 
 #include UE_INLINE_GENERATED_CPP_BY_NAME(NiagaraDataInterfaceHairStrands)
 
@@ -130,6 +132,7 @@ static const FName GetStrandsDensityName("GetStrandsDensity");
 static const FName GetStrandsSmoothingName("GetStrandsSmoothing");
 static const FName GetStrandsThicknessName("GetStrandsThickness");
 static const FName GetThicknessScaleName("GetThicknessScale");
+static const FName GetTimeDilationName("GetTimeDilation");
 
 //------------------------------------------------------------------------------------------------------------
 
@@ -580,9 +583,18 @@ void FNDIHairStrandsData::Update(
 			NumStrands = InData.NumCurves;
 			LocalSimulation = false;
 			BoneTransform = FTransform::Identity;
+			TimeDilation = 1.0f;
 
 			if (Interface->IsComponentValid())
 			{
+				if(UWorld* World = Interface->SourceComponent->GetWorld())
+				{ 
+					if(AWorldSettings* WorldSettings = World->GetWorldSettings())
+					{ 
+						TimeDilation = WorldSettings->GetEffectiveTimeDilation();
+					}
+				}
+
 				const FHairSimulationSettings& SimulationSettings = Interface->SourceComponent->SimulationSettings;
 				LocalSimulation = SimulationSettings.SimulationSetup.bLocalSimulation;
 				Interface->SourceComponent->BuildSimulationTransform(BoneTransform);
@@ -1471,6 +1483,16 @@ void UNiagaraDataInterfaceHairStrands::GetFunctionsInternal(TArray<FNiagaraFunct
 
 		OutFunctions.Add(Sig);
 	}
+	{
+		FNiagaraFunctionSignature Sig;
+		Sig.Name = GetTimeDilationName;
+		Sig.bMemberFunction = true;
+		Sig.bRequiresContext = false;
+		Sig.Inputs.Add(FNiagaraVariable(FNiagaraTypeDefinition(GetClass()), TEXT("Hair Strands")));
+		Sig.Outputs.Add(FNiagaraVariable(FNiagaraTypeDefinition::GetFloatDef(), TEXT("Time Dilation")));
+
+		OutFunctions.Add(Sig);
+	}
 	{
 		FNiagaraFunctionSignature Sig;
 		Sig.Name = GetStrandsSmoothingName;
@@ -2273,6 +2295,7 @@ DEFINE_NDI_DIRECT_FUNC_BINDER(UNiagaraDataInterfaceHairStrands, GetCollisionRadi
 DEFINE_NDI_DIRECT_FUNC_BINDER(UNiagaraDataInterfaceHairStrands, GetRadiusScale);
 
 DEFINE_NDI_DIRECT_FUNC_BINDER(UNiagaraDataInterfaceHairStrands, GetStrandsDensity);
+DEFINE_NDI_DIRECT_FUNC_BINDER(UNiagaraDataInterfaceHairStrands, GetTimeDilation);
 DEFINE_NDI_DIRECT_FUNC_BINDER(UNiagaraDataInterfaceHairStrands, GetStrandsSmoothing);
 DEFINE_NDI_DIRECT_FUNC_BINDER(UNiagaraDataInterfaceHairStrands, GetStrandsThickness);
 DEFINE_NDI_DIRECT_FUNC_BINDER(UNiagaraDataInterfaceHairStrands, GetThicknessScale);
@@ -2479,6 +2502,11 @@ void UNiagaraDataInterfaceHairStrands::GetVMExternalFunction(const FVMExternalFu
 		check(BindingInfo.GetNumInputs() == 1 && BindingInfo.GetNumOutputs() == 1);
 		NDI_FUNC_BINDER(UNiagaraDataInterfaceHairStrands, GetStrandsDensity)::Bind(this, OutFunc);
 	}
+	else if (BindingInfo.Name == GetTimeDilationName)
+	{
+		check(BindingInfo.GetNumInputs() == 1 && BindingInfo.GetNumOutputs() == 1);
+		NDI_FUNC_BINDER(UNiagaraDataInterfaceHairStrands, GetTimeDilation)::Bind(this, OutFunc);
+	}
 	else if (BindingInfo.Name == GetStrandsSmoothingName)
 	{
 		check(BindingInfo.GetNumInputs() == 1 && BindingInfo.GetNumOutputs() == 1);
@@ -3089,6 +3117,17 @@ void UNiagaraDataInterfaceHairStrands::GetStrandsDensity(FVectorVMExternalFuncti
 	}
 }
 
+void UNiagaraDataInterfaceHairStrands::GetTimeDilation(FVectorVMExternalFunctionContext& Context)
+{
+	VectorVM::FUserPtrHandler<FNDIHairStrandsData> InstData(Context);
+	VectorVM::FExternalFuncRegisterHandler<float> OutTimeDilation(Context);
+
+	for (int32 InstanceIdx = 0; InstanceIdx < Context.GetNumInstances(); ++InstanceIdx)
+	{
+		*OutTimeDilation.GetDestAndAdvance() = InstData->TimeDilation;
+	}
+}
+
 void UNiagaraDataInterfaceHairStrands::GetStrandsThickness(FVectorVMExternalFunctionContext& Context)
 {
 	VectorVM::FUserPtrHandler<FNDIHairStrandsData> InstData(Context);
diff --git a/Engine/Plugins/Runtime/HairStrands/Source/HairStrandsCore/Public/Niagara/NiagaraDataInterfaceHairStrands.h b/Engine/Plugins/Runtime/HairStrands/Source/HairStrandsCore/Public/Niagara/NiagaraDataInterfaceHairStrands.h
--- a/Engine/Plugins/Runtime/HairStrands/Source/HairStrandsCore/Public/Niagara/NiagaraDataInterfaceHairStrands.h
+++ b/Engine/Plugins/Runtime/HairStrands/Source/HairStrandsCore/Public/Niagara/NiagaraDataInterfaceHairStrands.h
@@ -159,6 +159,7 @@ struct FNDIHairStrandsData
 		}
 		SkeletalMeshes = 0;
 		LocalSimulation = false;
+		TimeDilation = 1.0f;
 	}
 
 	inline void CopyDatas(const FNDIHairStrandsData* OtherDatas)
@@ -232,6 +233,7 @@ struct FNDIHairStrandsData
 
 			TickingGroup = OtherDatas->TickingGroup;
 			LocalSimulation = OtherDatas->LocalSimulation;
+			TimeDilation = OtherDatas->TimeDilation;
 		}
 	}
 
@@ -375,6 +377,9 @@ struct FNDIHairStrandsData
 
 	/** Check if the simulation is running in local coordinate */
 	bool LocalSimulation;
+
+	/** Get the world time dilation */
+	float TimeDilation;
 };
 
 /** Data Interface for the strand base */
@@ -497,6 +502,8 @@ public:
 
 	UE_API void GetStrandsDensity(FVectorVMExternalFunctionContext& Context);
 
+	UE_API void GetTimeDilation(FVectorVMExternalFunctionContext& Context);
+
 	UE_API void GetStrandsThickness(FVectorVMExternalFunctionContext& Context);
 
 	UE_API void GetThicknessScale(FVectorVMExternalFunctionContext& Context);
