Fix hair strands RT buffer being partially updated causing RT geometry to have stale data.

#Synced-CL 51943914

diff --git a/Engine/Plugins/Runtime/HairStrands/Source/HairStrandsCore/Private/GroomManager.cpp b/Engine/Plugins/Runtime/HairStrands/Source/HairStrandsCore/Private/GroomManager.cpp
--- a/Engine/Plugins/Runtime/HairStrands/Source/HairStrandsCore/Private/GroomManager.cpp
+++ b/Engine/Plugins/Runtime/HairStrands/Source/HairStrandsCore/Private/GroomManager.cpp
@@ -1820,12 +1820,18 @@ static void RunHairStrandsInterpolation_Strands(
 				InstanceData.bNeedRaytracingBuild = !InstanceData.Instance->Strands.RenRaytracingResource->bIsRTGeometryInitialized;
 				if (InstanceData.bNeedRaytracingBuild || InstanceData.bNeedRaytracingUpdate)
 				{
+					// The entire RT buffer needs to be updated for ensuring invalid segments marked. 
+					// When building the RT geometry, the entire buffer is processed, and having partially updated 
+					// buffer would cause the groom to have stale curve/points
+					const uint32 ActivePointCount = InstanceData.ActivePointCount;
+					const uint32 TotalPointCount = InstanceData.Instance->Strands.RestResource->GetPointCount();
 					AddGenerateRaytracingGeometryPass(
 						GraphBuilder,
 						ShaderMap,
 						ShaderPrintData,
 						InstanceData.Instance->RegisteredIndex,
-						InstanceData.ActivePointCount,
+						ActivePointCount,
+						TotalPointCount,
 						InstanceData.Instance->Strands.RenRaytracingResource->RayTracingPrimitiveType,
 						ProceduralSplits,
 						HairRadiusRT,
diff --git a/Engine/Plugins/Runtime/HairStrands/Source/HairStrandsCore/Private/HairStrandsInterpolation.cpp b/Engine/Plugins/Runtime/HairStrands/Source/HairStrandsCore/Private/HairStrandsInterpolation.cpp
--- a/Engine/Plugins/Runtime/HairStrands/Source/HairStrandsCore/Private/HairStrandsInterpolation.cpp
+++ b/Engine/Plugins/Runtime/HairStrands/Source/HairStrandsCore/Private/HairStrandsInterpolation.cpp
@@ -1333,7 +1333,8 @@ void AddGenerateRaytracingGeometryPass(
 	FGlobalShaderMap* ShaderMap,
 	const FShaderPrintData* ShaderPrintData,
 	uint32 InstanceRegisteredIndex,
-	uint32 PointCount,
+	uint32 ActivePointCount,
+	uint32 TotalPointCount,
 	EHairStrandsRayTracingPrimitiveType RayTracingPrimitiveType,
 	int ProceduralSplits,
 	float HairRadius,
@@ -1347,7 +1348,7 @@ void AddGenerateRaytracingGeometryPass(
 	const FRDGBufferUAVRef& OutIndexBuffer)
 {
 	FHairRaytracingGeometryCS::FParameters* Parameters = GraphBuilder.AllocParameters<FHairRaytracingGeometryCS::FParameters>();
-	Parameters->PointCount = PointCount;
+	Parameters->PointCount = ActivePointCount;
 	Parameters->PositionOffsetBuffer = HairWorldOffsetBuffer;
 	Parameters->HairStrandsVF_RegisteredIndex = InstanceRegisteredIndex;
 	Parameters->HairStrandsVF_HairRadius = HairRadius;
@@ -1378,7 +1379,7 @@ void AddGenerateRaytracingGeometryPass(
 	PermutationVector.Set<FHairRaytracingGeometryCS::FPrimitiveType>(RayTracingPrimitiveType);
 	TShaderMapRef<FHairRaytracingGeometryCS> ComputeShader(ShaderMap, PermutationVector);
 
-	const FIntVector DispatchCount(FMath::DivideAndRoundUp(PointCount, FHairRaytracingGeometryCS::GetGroupSize()), 1, 1);
+	const FIntVector DispatchCount(FMath::DivideAndRoundUp(TotalPointCount, FHairRaytracingGeometryCS::GetGroupSize()), 1, 1);
 	FComputeShaderUtils::AddPass(
 		GraphBuilder,
 		RDG_EVENT_NAME("HairStrands::GenerateRaytracingGeometry(bCulling:%d, PrimitiveType:%d)", bCullingEnable, (int32)RayTracingPrimitiveType),
diff --git a/Engine/Plugins/Runtime/HairStrands/Source/HairStrandsCore/Public/HairStrandsInterpolation.h b/Engine/Plugins/Runtime/HairStrands/Source/HairStrandsCore/Public/HairStrandsInterpolation.h
--- a/Engine/Plugins/Runtime/HairStrands/Source/HairStrandsCore/Public/HairStrandsInterpolation.h
+++ b/Engine/Plugins/Runtime/HairStrands/Source/HairStrandsCore/Public/HairStrandsInterpolation.h
@@ -195,7 +195,8 @@ void AddGenerateRaytracingGeometryPass(
 	FGlobalShaderMap* ShaderMap,
 	const FShaderPrintData* ShaderPrintData,
 	uint32 InstanceRegisteredIndex,
-	uint32 PointCount,
+	uint32 ActivePointCount,
+	uint32 TotalPointCount,
 	EHairStrandsRayTracingPrimitiveType RayTracingPrimitiveType,
 	int ProceduralSplits,
 	float HairRadius,
