Make GroomRBFDeformer optionally safe to call from a background thread

#Synced-CL 52002724

diff --git a/Engine/Plugins/Runtime/HairStrands/Source/HairStrandsCore/Private/GroomRBFDeformer.cpp b/Engine/Plugins/Runtime/HairStrands/Source/HairStrandsCore/Private/GroomRBFDeformer.cpp
--- a/Engine/Plugins/Runtime/HairStrands/Source/HairStrandsCore/Private/GroomRBFDeformer.cpp
+++ b/Engine/Plugins/Runtime/HairStrands/Source/HairStrandsCore/Private/GroomRBFDeformer.cpp
@@ -15,6 +15,7 @@
 #include "Engine/Texture.h"
 #include "GroomBindingBuilder.h"
 #include "GroomBindingCommon.h"
+#include "StaticMeshCompiler.h"
 
 #define ALLOW_CURVE_SHUFFLING 0
 
@@ -546,8 +547,7 @@ void DeformStaticMeshPositions(
 	const TArray<FVector3f>& MeshSampleWeightsBuffer = OutSampleDeformationsBuffer;
 
 	// Raw deformation with RBF
-	const uint32 MeshLODCount = OutMesh->GetNumLODs();
-	TArray<const FMeshDescription*> MeshDescriptions;
+	const uint32 MeshLODCount = OutMesh->GetNumSourceModels();
 	for (uint32 MeshLODIt = 0; MeshLODIt < MeshLODCount; ++MeshLODIt)
 	{
 		FMeshDescription* MeshDescription = OutMesh->GetMeshDescription(MeshLODIt);
@@ -560,9 +560,14 @@ void DeformStaticMeshPositions(
 			FVector3f& Point = VertexPositions[VertexID];
 			Point = DisplacePosition(Point, MaxSampleCount, SampleRestPositionsBuffer, SampleDeformedPositionsBuffer, MeshSampleWeightsBuffer);
 		});
-		MeshDescriptions.Add(MeshDescription);
+
+		// Disable bMarkPackageDirty to make this thread safe.
+		//
+		// Package will be marked dirty by the caller.
+		UStaticMesh::FCommitMeshDescriptionParams CommitParams;
+		CommitParams.bMarkPackageDirty = false;
+		OutMesh->CommitMeshDescription(MeshLODIt, CommitParams);
 	}
-	OutMesh->BuildFromMeshDescriptions(MeshDescriptions);
 }
 
 namespace GroomDerivedDataCacheUtils
@@ -624,8 +629,18 @@ static float SampleMaskTexture(const FVector2f& InUV, const FUintPoint& InResolu
 	return 1.0f;
 }
 
-void FGroomRBFDeformer::GetRBFDeformedGroomAsset(const UGroomAsset* InGroomAsset, const UGroomBindingAsset* BindingAsset, FTextureSource* MaskTextureSource, const float MaskScale, UGroomAsset* OutGroomAsset, const ITargetPlatform* TargetPlatform)
+void FGroomRBFDeformer::GetRBFDeformedGroomAsset(
+	const UGroomAsset* InGroomAsset, 
+	const UGroomBindingAsset* BindingAsset, 
+	FTextureSource* MaskTextureSource, 
+	const float MaskScale, 
+	UGroomAsset* OutGroomAsset, 
+	const ITargetPlatform* TargetPlatform,
+	bool bShouldBuildStaticMeshes)
 {
+	// If bShouldBuildStaticMeshes is true, this function may only be called from the Game Thread
+	check(!bShouldBuildStaticMeshes || IsInGameThread());
+
 #if WITH_EDITORONLY_DATA
 	if (InGroomAsset && BindingAsset && BindingAsset->GetTargetSkeletalMesh() && BindingAsset->GetSourceSkeletalMesh())
 	{
@@ -727,7 +742,7 @@ void FGroomRBFDeformer::GetRBFDeformedGroomAsset(const UGroomAsset* InGroomAsset
 			check(RootDatas[GroupIndex].RenRootDatas.IsValidIndex(MeshLODIndex));
 
 			// Get the deformed positions (== target mesh positions), i.e. the non-transfered position.
-			TArray<FVector3f> MeshVertexPositionsBuffer_Target = RootDatas[GroupIndex].MeshPositions[MeshLODIndex];
+			const TArray<FVector3f>& MeshVertexPositionsBuffer_Target = RootDatas[GroupIndex].MeshPositions[MeshLODIndex];
 
 			// Get deformed guides
 			// If the groom override the value, we output dummy value for the guides, since they won't be used
@@ -738,7 +753,7 @@ void FGroomRBFDeformer::GetRBFDeformedGroomAsset(const UGroomAsset* InGroomAsset
 			}
 			else
 			{
-				FHairStrandsRootData SimRootData = RootDatas[GroupIndex].SimRootDatas[MeshLODIndex];
+				const FHairStrandsRootData& SimRootData = RootDatas[GroupIndex].SimRootDatas[MeshLODIndex];
 
 				const bool bHasRBF = InGroomAsset->IsGlobalInterpolationEnable(GroupIndex, 0/*HairLODIndex*/);
 
@@ -758,9 +773,9 @@ void FGroomRBFDeformer::GetRBFDeformedGroomAsset(const UGroomAsset* InGroomAsset
 				// Transfer RBF weights from SimRootData to RenRootData, as RenRootData does not hold RBF sample data
 				// This is transient data.
 				{
-					FHairStrandsRootData SimRootData = RootDatas[GroupIndex].SimRootDatas[MeshLODIndex];
+					const FHairStrandsRootData& SimRootData = RootDatas[GroupIndex].SimRootDatas[MeshLODIndex];
 
-					FHairStrandsRootData& SimLODData = SimRootData;
+					const FHairStrandsRootData& SimLODData = SimRootData;
 					FHairStrandsRootData& RenLODData = RenRootData;
 					RenLODData.SampleCount 						= SimLODData.SampleCount;
 					RenLODData.MeshInterpolationWeightsBuffer 	= SimLODData.MeshInterpolationWeightsBuffer;
@@ -875,12 +890,16 @@ void FGroomRBFDeformer::GetRBFDeformedGroomAsset(const UGroomAsset* InGroomAsset
 			{
 				if (Desc.GroupIndex >= 0)
 				{
-					Mesh->ConditionalPostLoad();
+					if (bShouldBuildStaticMeshes)
+					{
+						Mesh->ConditionalPostLoad();
+						Mesh->PreEditChange(nullptr);
+					}
 	
 					check(RootDatas[Desc.GroupIndex].SimRootDatas.IsValidIndex(MeshLODIndex));
 
 					// Load Sim root data, as they contains RBF weights
-					FHairStrandsRootData SimRootData = RootDatas[Desc.GroupIndex].SimRootDatas[MeshLODIndex];
+					const FHairStrandsRootData& SimRootData = RootDatas[Desc.GroupIndex].SimRootDatas[MeshLODIndex];
 	
 					// Get the deformed positions (== target mesh positions), i.e. the non-transfered position.
 					const TArray<FVector3f>& MeshVertexPositionsBuffer_Target = RootDatas[Desc.GroupIndex].MeshPositions[MeshLODIndex];
@@ -894,17 +913,21 @@ void FGroomRBFDeformer::GetRBFDeformedGroomAsset(const UGroomAsset* InGroomAsset
 		{
 			if (UStaticMesh* Mesh = Desc.ImportedMesh)
 			{
-				if (Mesh->GetNumLODs() == 0 || Desc.GroupIndex < 0 || Desc.GroupIndex >= InGroomAsset->GetNumHairGroups() || Desc.LODIndex == -1)
+				if (Mesh->GetNumSourceModels() == 0 || Desc.GroupIndex < 0 || Desc.GroupIndex >= InGroomAsset->GetNumHairGroups() || Desc.LODIndex == -1)
 				{
 					continue;
 				}
 
-				Mesh->ConditionalPostLoad();
+				if (bShouldBuildStaticMeshes)
+				{
+					Mesh->ConditionalPostLoad();
+					Mesh->PreEditChange(nullptr);
+				}
 
 				check(RootDatas[Desc.GroupIndex].SimRootDatas.IsValidIndex(MeshLODIndex));
 
 				// Load Sim root data, as they contains RBF weights
-				FHairStrandsRootData SimRootData = RootDatas[Desc.GroupIndex].SimRootDatas[MeshLODIndex];
+				const FHairStrandsRootData& SimRootData = RootDatas[Desc.GroupIndex].SimRootDatas[MeshLODIndex];
 
 				// Get the deformed positions (== target mesh positions), i.e. the non-transfered position.
 				const TArray<FVector3f>& MeshVertexPositionsBuffer_Target = RootDatas[Desc.GroupIndex].MeshPositions[MeshLODIndex];
@@ -912,6 +935,34 @@ void FGroomRBFDeformer::GetRBFDeformedGroomAsset(const UGroomAsset* InGroomAsset
 			}
 		}
 
+		// Build all static meshes in one batch
+		if (bShouldBuildStaticMeshes)
+		{
+			TArray<UStaticMesh*> MeshesToBuild;
+			MeshesToBuild.Reserve(OutGroomAsset->GetHairGroupsCards().Num() + OutGroomAsset->GetHairGroupsMeshes().Num());
+
+			for (FHairGroupsCardsSourceDescription& Desc : OutGroomAsset->GetHairGroupsCards())
+			{
+				if (UStaticMesh* Mesh = Desc.ImportedMesh)
+				{
+					MeshesToBuild.Add(Mesh);
+					Mesh->MarkPackageDirty();
+				}
+			}
+			for (FHairGroupsMeshesSourceDescription& Desc : OutGroomAsset->GetHairGroupsMeshes())
+			{
+				if (UStaticMesh* Mesh = Desc.ImportedMesh)
+				{
+					MeshesToBuild.Add(Mesh);
+					Mesh->MarkPackageDirty();
+				}
+			}
+
+			UStaticMesh::BatchBuild(MeshesToBuild, /*bSilent*/ true);
+
+			FStaticMeshCompilingManager::Get().FinishAllCompilation();
+		}
+
 		// Finally, the deformed guides and strands are applied to the GroomAsset
 		ApplyDeformationToGroom(DeformedPositions, OutGroomAsset);
 	}
diff --git a/Engine/Plugins/Runtime/HairStrands/Source/HairStrandsCore/Public/GroomRBFDeformer.h b/Engine/Plugins/Runtime/HairStrands/Source/HairStrandsCore/Public/GroomRBFDeformer.h
--- a/Engine/Plugins/Runtime/HairStrands/Source/HairStrandsCore/Public/GroomRBFDeformer.h
+++ b/Engine/Plugins/Runtime/HairStrands/Source/HairStrandsCore/Public/GroomRBFDeformer.h
@@ -13,7 +13,30 @@ struct FTextureSource;
 struct FGroomRBFDeformer
 {
 	// Return a new GroomAsset with the RBF deformation from the BindingAsset baked into it
-	UE_API void GetRBFDeformedGroomAsset(const UGroomAsset* InGroomAsset, const UGroomBindingAsset* BindingAsset, FTextureSource* MaskTextureSource, const float MaskScale, UGroomAsset* OutGroomAsset, const ITargetPlatform* TargetPlatform=nullptr);
+	//
+	// If bShouldBuildStaticMeshes is true, this function will ensure the static meshes owned by 
+	// the groom are built after being modified. In this case, GetRBFDeformedGroomAsset can only be
+	// called from the Game Thread.
+	//
+	// If bShouldBuildStaticMeshes is false, GetRBFDeformedGroomAsset may be called from a 
+	// background thread, but the caller must:
+	//
+	//	1.	Call ConditionalPostLoad and PreEditChange on each ImportedMesh returned by both 
+	//		GetHairGroupsCards and GetHairGroupsMeshes *before* calling GetRBFDeformedGroomAsset.
+	//
+	//	2.	Call UStaticMesh::BatchBuild and MarkPackageDirty on the same ImportedMeshes *after* 
+	//		calling GetRBFDeformedGroomAsset to begin async building of the meshes.
+	//		
+	//		To block until the async build is finished (if necessary), call 
+	//		FStaticMeshCompilingManager::Get().FinishAllCompilation().
+	UE_API void GetRBFDeformedGroomAsset(
+		const UGroomAsset* InGroomAsset, 
+		const UGroomBindingAsset* BindingAsset, 
+		FTextureSource* MaskTextureSource, 
+		const float MaskScale, 
+		UGroomAsset* OutGroomAsset, 
+		const ITargetPlatform* TargetPlatform = nullptr,
+		bool bShouldBuildStaticMeshes = true);
 
 	static UE_API uint32 GetEntryCount(uint32 InSampleCount);
 	static UE_API uint32 GetWeightCount(uint32 InSampleCount);
