Groom dataflow resample curve node: fix crash when resampling only some curves
- the sample array was using the wrong offset

#Synced-CL 52046631

diff --git a/Engine/Plugins/Runtime/HairStrands/Source/HairStrandsDataflow/Private/ResampleGuidesPointsNode.cpp b/Engine/Plugins/Runtime/HairStrands/Source/HairStrandsDataflow/Private/ResampleGuidesPointsNode.cpp
--- a/Engine/Plugins/Runtime/HairStrands/Source/HairStrandsDataflow/Private/ResampleGuidesPointsNode.cpp
+++ b/Engine/Plugins/Runtime/HairStrands/Source/HairStrandsDataflow/Private/ResampleGuidesPointsNode.cpp
@@ -8,7 +8,7 @@
 
 namespace UE::Groom::Private
 {
-	FORCEINLINE int32 ComputeCurvesPoints(const GeometryCollection::Facades::FCollectionCurveGeometryFacade& CurvesFacade, const FDataflowCurveSelection& CurveSelection, const int32 GuidePoints) 
+	int32 ComputeCurvesPoints(const GeometryCollection::Facades::FCollectionCurveGeometryFacade& CurvesFacade, const FDataflowCurveSelection& CurveSelection, const int32 GuidePoints) 
 	{
 		const bool bValidSelection = CurveSelection.IsValidForCollection(CurvesFacade.GetManagedArrayCollection());
 
@@ -28,7 +28,7 @@ namespace UE::Groom::Private
 		return NumPoints;
 	}
 
-	FORCEINLINE void BuildResampledPoints(const TArray<FVector3f>& PointRestPositions, const TArray<int32>& CurveObjectIndices, const TArray<int32>& CurvePointOffsets,
+	void BuildResampledPoints(const TArray<FVector3f>& PointRestPositions, const TArray<int32>& CurveObjectIndices, const TArray<int32>& CurvePointOffsets,
 		const TArray<int32>& ObjectPointSamples, const int32 CurveIndex, const int32 PointOffset, const int32 SampleOffset,
 					TArray<FVector3f>& SamplePositions, const int32 GuidePoints) 
 	{
@@ -71,7 +71,7 @@ namespace UE::Groom::Private
 		SamplePositions[SampleOffset] = 2 * SamplePositions[SampleOffset+1] - SamplePositions[SampleOffset+2];
 	}
 
-	FORCEINLINE void ResampleCurvesPoints(GeometryCollection::Facades::FCollectionCurveGeometryFacade& CurvesFacade, const FDataflowCurveSelection& CurveSelection, const int32 GuidePoints)
+	void ResampleCurvesPoints(GeometryCollection::Facades::FCollectionCurveGeometryFacade& CurvesFacade, const FDataflowCurveSelection& CurveSelection, const int32 GuidePoints)
 	{
 		if(CurvesFacade.IsValid())
 		{
@@ -101,7 +101,7 @@ namespace UE::Groom::Private
 				{
 					for(int32 PointIndex = 0; PointIndex < NumPoints; ++PointIndex)
 					{
-						PointSampledPositions[PointIndex+PointOffset] = CurvesFacade.GetPointRestPositions()[PointIndex+PointOffset];
+						PointSampledPositions[PointIndex + SampleOffset] = CurvesFacade.GetPointRestPositions()[PointIndex + PointOffset];
 					}
 					SampleOffset += NumPoints;
 				}
@@ -171,4 +171,3 @@ void FResampleCurvePointsDataflowNode::Evaluate(UE::Dataflow::FContext& Context,
 		SetValue(Context, MoveTemp(GuidesCollection), &Collection);
 	}
 }
-
