[HWRT] Hair BLAS management migrated to FRayTracingDynamicGeometryUpdateManager
* Strands, cards, and meshes now register with Dynamic Geometry Manager instead of building BLAS in per-instance RDG passes
* Resource lifecycle uses proper InitResource/ReleaseResource pattern
* Added ExternalDependencies to FRayTracingDynamicGeometryUpdateParams to track hair CS output buffers as RDG dependencies
* Strands interpolation runs early (before DispatchRayTracingWorldUpdates) to ensure BLAS registration before manager dispatch

#tests replayrun
#rnx
#Synced-CL 52201145

diff --git a/Engine/Plugins/Runtime/HairStrands/Source/HairStrandsCore/Private/GroomManager.cpp b/Engine/Plugins/Runtime/HairStrands/Source/HairStrandsCore/Private/GroomManager.cpp
--- a/Engine/Plugins/Runtime/HairStrands/Source/HairStrandsCore/Private/GroomManager.cpp
+++ b/Engine/Plugins/Runtime/HairStrands/Source/HairStrandsCore/Private/GroomManager.cpp
@@ -83,7 +83,6 @@ float GetHairStrandsAutoLODBias()
 ///////////////////////////////////////////////////////////////////////////////////////////////////
 // Foward declaration
 
-bool IsHairStrandsForceRebuildBVH();
 bool IsHairStrandsTransferPositionOnLODChange();
 uint32 GetHairRaytracingProceduralSplits();
 void CreateHairStrandsDebugAttributeBuffer(FRDGBuilder& GraphBuilder, FRDGExternalBuffer* DebugAttributeBuffer, uint32 VertexCount, const FName& OwnerName);
@@ -1208,14 +1207,15 @@ void AddDrawDebugClusterPass(
 static void RunHairStrandsInterpolation_Strands(
 	FRDGBuilder& GraphBuilder,
 	FSceneInterface* Scene,
-	const TArray<const FSceneView*>& Views, 
+	const TArray<const FSceneView*>& Views,
 	const FSceneView* View,
 	const uint32 ViewUniqueID,
 	const uint32 TotalRegisteredInstanceCount,
 	const FHairStrandsInstances& Instances,
 	const FShaderPrintData* ShaderPrintData,
 	FHairTransientResources& TransientResources,
-	FGlobalShaderMap* ShaderMap)
+	FGlobalShaderMap* ShaderMap,
+	FRayTracingDynamicGeometryUpdateManager* DynamicGeometryUpdateManager)
 {
 	if (Instances.IsEmpty()) { return; }
 	if (!IsHairStrandsEnabled(EHairStrandsShaderType::Strands, Scene->GetShaderPlatform())) { return;  }
@@ -1824,8 +1824,10 @@ static void RunHairStrandsInterpolation_Strands(
 					InstanceData.Instance->Strands.CachedHairRootScale != HairRootScaleRT ||
 					InstanceData.Instance->Strands.CachedHairTipScale != HairTipScaleRT ||
 					InstanceData.Instance->Strands.CachedProceduralSplits != ProceduralSplits;
-				InstanceData.bNeedRaytracingBuild = !InstanceData.Instance->Strands.RenRaytracingResource->bIsRTGeometryInitialized;
-				if (InstanceData.bNeedRaytracingBuild || InstanceData.bNeedRaytracingUpdate)
+
+				const bool bFirstTime = InstanceData.Instance->Strands.RenRaytracingResource->RayTracingGeometry.GetRequiresBuild();
+				InstanceData.bNeedRaytracingBuild = bFirstTime;
+				if (bFirstTime || InstanceData.bNeedRaytracingUpdate)
 				{
 					// The entire RT buffer needs to be updated for ensuring invalid segments marked. 
 					// When building the RT geometry, the entire buffer is processed, and having partially updated 
@@ -1859,14 +1861,20 @@ static void RunHairStrandsInterpolation_Strands(
 		}
 
 		// 2. Build/update acceleration structure
+		FScene* RenderScene = static_cast<FScene*>(Scene->GetRenderScene());
+
 		for (FInstanceData& InstanceData : InstanceDatas)
 		{
-			if (InstanceData.bNeedRaytracing && (InstanceData.bNeedRaytracingBuild || InstanceData.bNeedRaytracingUpdate))
+			if (InstanceData.bNeedRaytracing)
 			{
-				AddBuildStrandsAccelerationStructurePass(
+				RegisterHairStrandsWithRayTracingManager(
 					GraphBuilder,
+					DynamicGeometryUpdateManager,
+					RenderScene,
+					View,
 					InstanceData.Instance,
 					ProceduralSplits,
+					InstanceData.bNeedRaytracingBuild,
 					InstanceData.bNeedRaytracingUpdate,
 					InstanceData.RDGResources.Raytracing_PositionBuffer.Buffer,
 					InstanceData.RDGResources.Raytracing_IndexBuffer.Buffer);
@@ -1924,13 +1932,14 @@ uint32 GetHairCardsInterpolationType();
 static void RunHairStrandsInterpolation_Cards(
 	FRDGBuilder& GraphBuilder,
 	FSceneInterface* Scene,
-	const TArray<const FSceneView*>& Views, 
+	const TArray<const FSceneView*>& Views,
 	const FSceneView* View,
 	const uint32 ViewUniqueID,
 	const FHairStrandsInstances& Instances,
 	const FShaderPrintData* ShaderPrintData,
 	FHairTransientResources& TransientResources,
-	FGlobalShaderMap* ShaderMap)
+	FGlobalShaderMap* ShaderMap,
+	FRayTracingDynamicGeometryUpdateManager* DynamicGeometryUpdateManager)
 {
 	const bool bCardSupported = IsHairStrandsEnabled(EHairStrandsShaderType::Cards, Scene->GetShaderPlatform());
 	const bool bMeshSupported = IsHairStrandsEnabled(EHairStrandsShaderType::Meshes, Scene->GetShaderPlatform());
@@ -2256,13 +2265,18 @@ static void RunHairStrandsInterpolation_Cards(
 	#if RHI_RAYTRACING
 	if (IsHairRayTracingEnabled() && IsRayTracingEnabled() && !bIsReflectionCapture)
 	{
+		FScene* RenderScene = static_cast<FScene*>(Scene->GetRenderScene());
+
 		for (uint32 InstanceIndex : CardInstances)
 		{
 			FInstanceData& InstanceData = InstanceDatas[InstanceIndex];
 			if (InstanceData.CardInstance->RaytracingResource)
 			{
-				AddBuildHairCardAccelerationStructurePass(
+				RegisterHairCardsWithRayTracingManager(
 					GraphBuilder,
+					DynamicGeometryUpdateManager,
+					RenderScene,
+					View,
 					InstanceData.Instance,
 					InstanceData.HairLODIndex,
 					InstanceData.bNeedDeformation);
@@ -2294,13 +2308,18 @@ static void RunHairStrandsInterpolation_Cards(
 	#if RHI_RAYTRACING
 	if (IsHairRayTracingEnabled() && IsRayTracingEnabled() && !bIsReflectionCapture)
 	{
+		FScene* RenderScene = static_cast<FScene*>(Scene->GetRenderScene());
+
 		for (uint32 InstanceIndex : MeshInstances)
 		{
 			FInstanceData& InstanceData = InstanceDatas[InstanceIndex];
 			if (InstanceData.MeshInstance->RaytracingResource)
 			{
-				AddBuildHairMeshAccelerationStructurePass(
+				RegisterHairMeshesWithRayTracingManager(
 					GraphBuilder,
+					DynamicGeometryUpdateManager,
+					RenderScene,
+					View,
 					InstanceData.Instance,
 					InstanceData.HairLODIndex,
 					InstanceData.bNeedDeformation);
@@ -3495,7 +3514,8 @@ void ProcessHairStrandsBookmark(
 			*Instances,
 			Parameters.ShaderPrintData,
 			*Parameters.TransientResources,
-			Parameters.ShaderMap);
+			Parameters.ShaderMap,
+			Parameters.DynamicGeometryUpdateManager);
 	}
 	else if (Bookmark == EHairStrandsBookmark::ProcessStrandsInterpolation)
 	{
@@ -3511,7 +3531,8 @@ void ProcessHairStrandsBookmark(
 			*Instances,
 			Parameters.ShaderPrintData,
 			*Parameters.TransientResources,
-			Parameters.ShaderMap);
+			Parameters.ShaderMap,
+			Parameters.DynamicGeometryUpdateManager);
 	}
 	else if (Bookmark == EHairStrandsBookmark::ProcessDebug)
 	{
diff --git a/Engine/Plugins/Runtime/HairStrands/Source/HairStrandsCore/Private/GroomResources.cpp b/Engine/Plugins/Runtime/HairStrands/Source/HairStrandsCore/Private/GroomResources.cpp
--- a/Engine/Plugins/Runtime/HairStrands/Source/HairStrandsCore/Private/GroomResources.cpp
+++ b/Engine/Plugins/Runtime/HairStrands/Source/HairStrandsCore/Private/GroomResources.cpp
@@ -2052,10 +2052,109 @@ void FHairStrandsRaytracingResource::InternalAllocate(FRDGBuilder& GraphBuilder,
 		InternalCreateVertexBufferRDG<FHairStrandsRaytracingFormat>(GraphBuilder, MaxVertexCount, PositionBuffer, ToHairResourceDebugName(TEXT("Hair.StrandsRaytracing_PositionBuffer"), ResourceName), OwnerName, EHairResourceUsageType::Dynamic);
 		InternalCreateStructuredBufferRDG<FHairStrandsIndexFormat>(GraphBuilder, MaxIndexCount, IndexBuffer, ToHairResourceDebugName(TEXT("Hair.StrandsRaytracing_IndexBuffer"), ResourceName), OwnerName, EHairResourceUsageType::Dynamic);
 
-		// Force reallocating and rebuilding the RT geometry
-		RayTracingGeometry.ReleaseRHI();
-		RayTracingGeometry.SetInitializer(FRayTracingGeometryInitializer{});
-		bIsRTGeometryInitialized = false;
+		FRayTracingGeometryInitializer Initializer;
+
+		#if HAIR_RESOURCE_DEBUG_NAME
+		FString DebugName = FString::Printf(TEXT("AS_Hair_LOD%d_%s"), ResourceName.LODIndex, *ResourceName.AssetName.ToString());
+		Initializer.DebugName = FName(DebugName);
+		#else
+		static const FName DebugName("AS_Hair");
+		static int32 DebugNumber = 0;
+		Initializer.DebugName = FDebugName(DebugName, DebugNumber++);
+		#endif
+
+		#if RHI_ENABLE_RESOURCE_INFO
+		Initializer.OwnerName = OwnerName;
+		#endif
+
+		Initializer.bAllowUpdate = false; // Strands always rebuild, never refit
+		Initializer.bAllowCompaction = false;
+		Initializer.bFastBuild = true;
+		Initializer.bTemplate = true;
+
+		// Setup segment with buffer pointers
+		FRayTracingGeometrySegment Segment;
+		Segment.VertexBuffer = PositionBuffer.Buffer->GetRHI();
+		Segment.VertexBufferOffset = 0;
+		Segment.VertexBufferElementType = FHairStrandsRaytracingFormat::VertexElementType;
+		Segment.MaxVertices = MaxVertexCount;
+
+		if (RayTracingPrimitiveType == EHairStrandsRayTracingPrimitiveType::LinearSweptSphere)
+		{
+			Initializer.GeometryType = RTGT_LinearSweptSpheres;
+			Initializer.IndexBuffer = IndexBuffer.Buffer ? IndexBuffer.Buffer->GetRHI() : nullptr;
+			Initializer.IndexBufferOffset = 0;
+			Initializer.TotalPrimitiveCount = MaxIndexCount / 2;
+
+			Segment.VertexBufferStride = FHairStrandsRaytracingFormat::SizeInByte;
+			Segment.NumPrimitives = MaxIndexCount / 2;
+		}
+		else if (RayTracingPrimitiveType == EHairStrandsRayTracingPrimitiveType::Procedural)
+		{
+			Initializer.GeometryType = RTGT_Procedural;
+			Initializer.IndexBuffer = nullptr;
+			Initializer.TotalPrimitiveCount = MaxVertexCount / 2;
+
+			Segment.VertexBufferStride = FHairStrandsRaytracingFormat::SizeInByte * 2;
+			Segment.NumPrimitives = MaxVertexCount / 2;
+		}
+		else 
+		{
+			Initializer.GeometryType = RTGT_Triangles;
+			Initializer.IndexBuffer = IndexBuffer.Buffer ? IndexBuffer.Buffer->GetRHI() : nullptr;
+			Initializer.IndexBufferOffset = 0;
+			Initializer.TotalPrimitiveCount = MaxIndexCount / 3;
+
+			Segment.VertexBufferStride = FHairStrandsRaytracingFormat::SizeInByte;
+			Segment.NumPrimitives = MaxIndexCount / 3;
+		}
+
+		Initializer.Segments.Add(Segment);
+
+		// Release existing RT geometry before re-creating. InternalAllocate can be called multiple
+		// times without InternalRelease (e.g., async streaming scale-up). InitResource() is a no-op
+		// on an already-initialized resource, which would leave the RHI geometry sized for the old
+		// primitive count. ReleaseResource also clears the initializer, so SetInitializer must come after.
+		if (RayTracingGeometry.IsInitialized())
+		{
+			RayTracingGeometry.ReleaseResource();
+		}
+
+		RayTracingGeometry.SetInitializer(Initializer);
+
+		// InitRHI will call CreateRayTracingGeometry with Skip priority
+		// bTemplate=true ensures no immediate BLAS build, manager will handle it later
+		RayTracingGeometry.InitResource(GraphBuilder.RHICmdList);
+
+		RayTracingGeometry.LODIndex = 0;
+		RayTracingGeometry.SetRequiresBuild(true);
+	}
+	else
+	{
+		if (!RayTracingGeometry.IsInitialized())
+		{
+			FRayTracingGeometryInitializer Initializer;
+
+			#if HAIR_RESOURCE_DEBUG_NAME
+			FString DebugName = FString::Printf(TEXT("AS_Hair_LOD%d_%s"), ResourceName.LODIndex, *ResourceName.AssetName.ToString());
+			Initializer.DebugName = FName(DebugName);
+			#else
+			static const FName DebugName("AS_Hair");
+			static int32 DebugNumber = 0;
+			Initializer.DebugName = FDebugName(DebugName, DebugNumber++);
+			#endif
+
+			#if RHI_ENABLE_RESOURCE_INFO
+			Initializer.OwnerName = OwnerName;
+			#endif
+
+			Initializer.bAllowUpdate = true; // Cards/Meshes can refit
+			Initializer.bFastBuild = true;
+			Initializer.bTemplate = true;
+
+			RayTracingGeometry.SetInitializer(Initializer);
+			RayTracingGeometry.InitResource(GraphBuilder.RHICmdList);
+		}
 	}
 }
 
@@ -2063,9 +2162,8 @@ void FHairStrandsRaytracingResource::InternalRelease()
 {
 	PositionBuffer.Release();
 	IndexBuffer.Release();
-	RayTracingGeometry.ReleaseRHI();
+	RayTracingGeometry.ReleaseResource();
 	RayTracingGeometry.SetInitializer(FRayTracingGeometryInitializer{});
-	bIsRTGeometryInitialized = false;
 }
 #endif
 
diff --git a/Engine/Plugins/Runtime/HairStrands/Source/HairStrandsCore/Private/HairStrandsInterpolation.cpp b/Engine/Plugins/Runtime/HairStrands/Source/HairStrandsCore/Private/HairStrandsInterpolation.cpp
--- a/Engine/Plugins/Runtime/HairStrands/Source/HairStrandsCore/Private/HairStrandsInterpolation.cpp
+++ b/Engine/Plugins/Runtime/HairStrands/Source/HairStrandsCore/Private/HairStrandsInterpolation.cpp
@@ -33,6 +33,7 @@
 #include "DataDrivenShaderPlatformInfo.h"
 #include "ProfilingDebugging/RealtimeGPUProfiler.h"
 #include "GroomVisualizationData.h"
+#include "RayTracingDynamicGeometryUpdateManager.h"
 
 static float GHairRaytracingRadiusScale = 0;
 static FAutoConsoleVariableRef CVarHairRaytracingRadiusScale(TEXT("r.HairStrands.RaytracingRadiusScale"), GHairRaytracingRadiusScale, TEXT("Override the per instance scale factor for raytracing hair strands geometry (0: disabled, >0:enabled)"));
@@ -52,9 +53,6 @@ static FAutoConsoleVariableRef CVarHairCardsInterpolationType(TEXT("r.HairStrand
 static int32 GHairStrandsTransferPositionOnLODChange = 0;
 static FAutoConsoleVariableRef CVarHairStrandsTransferPositionOnLODChange(TEXT("r.HairStrands.Strands.TransferPrevPos"), GHairStrandsTransferPositionOnLODChange, TEXT("Transfer strands prev. position to current position on LOD switching to avoid large discrepancy causing large motion vector"));
 
-static int32 GHairStrands_Raytracing_ForceRebuildBVH = 1;
-static FAutoConsoleVariableRef CVarHairRTGeomForceRebuild(TEXT("r.HairStrands.Strands.Raytracing.ForceRebuildBVH"), GHairStrands_Raytracing_ForceRebuildBVH, TEXT("Force BVH rebuild instead of doing a BVH refit when hair positions changed"));
-
 static int32 GHairStrands_MotionVector_CheckViewID = 1;
 static FAutoConsoleVariableRef CVarHairStrands_MotionVector_CheckViewID(TEXT("r.HairStrands.Strands.MotionVectorCheckViewID"), GHairStrands_MotionVector_CheckViewID, TEXT("Issue motion vector on hair strands only if updates happens on the same view"));
 
@@ -99,11 +97,6 @@ bool IsHairStrandsTransferPositionOnLODChange()
 	return GHairStrandsTransferPositionOnLODChange > 0;
 }
 
-bool IsHairStrandsForceRebuildBVH()
-{
-	return GHairStrands_Raytracing_ForceRebuildBVH > 0;
-}
-
 uint32 GetHairRaytracingProceduralSplits()
 {
 	return FMath::Clamp(GHairRaytracingProceduralSplits, 1, STRANDS_PROCEDURAL_INTERSECTOR_MAX_SPLITS);
@@ -1317,15 +1310,6 @@ public:
 	}
 };
 
-BEGIN_SHADER_PARAMETER_STRUCT(FStrandsResourceBLASParameters, )
-	RDG_BUFFER_ACCESS(PositionBuffer, ERHIAccess::SRVCompute)
-	RDG_BUFFER_ACCESS(IndexBuffer, ERHIAccess::SRVCompute)
-END_SHADER_PARAMETER_STRUCT()
-
-BEGIN_SHADER_PARAMETER_STRUCT(FCardsOrMeshesResourceBLASParameters, )
-	RDG_BUFFER_ACCESS(PositionBuffer, ERHIAccess::SRVCompute)
-END_SHADER_PARAMETER_STRUCT()
-
 IMPLEMENT_GLOBAL_SHADER(FHairRaytracingGeometryCS, "/Engine/Private/HairStrands/HairStrandsRaytracingGeometry.usf", "MainCS", SF_Compute);
 
 void AddGenerateRaytracingGeometryPass(
@@ -1435,23 +1419,11 @@ void AddClearAABBPass(
 ///////////////////////////////////////////////////////////////////////////////////////////////////
 
 #if RHI_RAYTRACING
-static void UpdateHairAccelerationStructure(FRHICommandList& RHICmdList, FRayTracingGeometry* RayTracingGeometry, EAccelerationStructureBuildMode InMode)
-{
-	SCOPED_DRAW_EVENT(RHICmdList, CommitHairRayTracingGeometryUpdates);
-
-	FRayTracingGeometryBuildParams Params;
-	Params.BuildMode = InMode;
-	Params.Geometry = RayTracingGeometry->GetRHI();
-	Params.Segments = RayTracingGeometry->GetInitializer().Segments;
-
-	RHICmdList.BuildAccelerationStructures(MakeArrayView(&Params, 1));
-}
-
 static void BuildHairAccelerationStructure_Strands(
-	FRHICommandList& RHICmdList, 
-	uint32 RaytracingVertexCount, 
+	FRHICommandList& RHICmdList,
+	uint32 RaytracingVertexCount,
 	uint32 RayTracingIndexCount,
-	FBufferRHIRef& PositionBuffer, 
+	FBufferRHIRef& PositionBuffer,
 	FBufferRHIRef& IndexBuffer,
 	FRayTracingGeometry* OutRayTracingGeometry,
 	const FString& AssetDebugName,
@@ -1460,54 +1432,42 @@ static void BuildHairAccelerationStructure_Strands(
 	EHairStrandsRayTracingPrimitiveType RayTracingPrimitiveType,
 	int ProceduralSplits)
 {
-	FRayTracingGeometryInitializer Initializer;
+	check(OutRayTracingGeometry);
+
+	FRayTracingGeometryInitializer Initializer = OutRayTracingGeometry->GetInitializer();
+
 #if HAIR_RESOURCE_DEBUG_NAME
 	FString DebugName(TEXT("AS_HairStrands_LOD") + FString::FromInt(LODIndex) + TEXT("_") + AssetDebugName);
 	Initializer.DebugName = FName(DebugName);
-#else
-	static const FName DebugName("AS_HairStrands");	
-	static int32 DebugNumber = 0;
-	Initializer.DebugName = FDebugName(DebugName, DebugNumber++);
 #endif
+
 #if RHI_ENABLE_RESOURCE_INFO
 	const FName OwnerName(FString::Printf(TEXT("%s [LOD%d]"), *MeshComponentName, LODIndex));
 	Initializer.OwnerName = OwnerName;
 #endif
-	
+
+	check(Initializer.Segments.Num() == 1);
+	FRayTracingGeometrySegment& Segment = Initializer.Segments[0];
+	Segment.VertexBuffer = PositionBuffer;
+	Segment.VertexBufferOffset = 0;
+
 	if (RayTracingPrimitiveType == EHairStrandsRayTracingPrimitiveType::LinearSweptSphere)
 	{
 		Initializer.IndexBuffer = IndexBuffer;
 		Initializer.IndexBufferOffset = 0;
 		Initializer.GeometryType = RTGT_LinearSweptSpheres;
 		Initializer.TotalPrimitiveCount = RayTracingIndexCount / 2;
-	}
-	else if (RayTracingPrimitiveType == EHairStrandsRayTracingPrimitiveType::Procedural)
-	{
-		Initializer.GeometryType = RTGT_Procedural;
-		Initializer.TotalPrimitiveCount = (RaytracingVertexCount / (2 * STRANDS_PROCEDURAL_INTERSECTOR_MAX_SPLITS)) * ProceduralSplits;
-	}
-	else
-	{
-		Initializer.IndexBuffer = IndexBuffer;
-		Initializer.IndexBufferOffset = 0;
-		Initializer.GeometryType = RTGT_Triangles;
-		Initializer.TotalPrimitiveCount = RayTracingIndexCount / 3;
-	}
-	Initializer.bFastBuild = true;
-	Initializer.bAllowUpdate = false; // we won't be using update, so specify this to allow degenerate primitives to be discarded
-
-	FRayTracingGeometrySegment Segment;
-	Segment.VertexBuffer = PositionBuffer;
-	Segment.VertexBufferElementType = FHairStrandsRaytracingFormat::VertexElementType;
 
-	if (RayTracingPrimitiveType == EHairStrandsRayTracingPrimitiveType::LinearSweptSphere)
-	{
 		Segment.VertexBufferStride = FHairStrandsRaytracingFormat::SizeInByte;
 		Segment.MaxVertices = RaytracingVertexCount;
 		Segment.NumPrimitives = Initializer.TotalPrimitiveCount;
 	}
 	else if (RayTracingPrimitiveType == EHairStrandsRayTracingPrimitiveType::Procedural)
 	{
+		Initializer.IndexBuffer = nullptr;
+		Initializer.GeometryType = RTGT_Procedural;
+		Initializer.TotalPrimitiveCount = (RaytracingVertexCount / (2 * STRANDS_PROCEDURAL_INTERSECTOR_MAX_SPLITS)) * ProceduralSplits;
+
 		// stride covers one AABB (which is made up of two vertices)
 		Segment.VertexBufferStride = FHairStrandsRaytracingFormat::SizeInByte * 2;
 		Segment.MaxVertices = RaytracingVertexCount;
@@ -1515,14 +1475,22 @@ static void BuildHairAccelerationStructure_Strands(
 	}
 	else
 	{
+		Initializer.IndexBuffer = IndexBuffer;
+		Initializer.IndexBufferOffset = 0;
+		Initializer.GeometryType = RTGT_Triangles;
+		Initializer.TotalPrimitiveCount = RayTracingIndexCount / 3;
+
 		Segment.VertexBufferStride = FHairStrandsRaytracingFormat::SizeInByte;
 		Segment.MaxVertices = RaytracingVertexCount;
 		Segment.NumPrimitives = RayTracingIndexCount / 3;
 	}
-	Initializer.Segments.Add(Segment);
+
+	Initializer.bFastBuild = true;
+	Initializer.bAllowUpdate = false; // we won't be using update, so specify this to allow degenerate primitives to be discarded
 
 	OutRayTracingGeometry->SetInitializer(Initializer);
-	OutRayTracingGeometry->CreateRayTracingGeometry(RHICmdList, ERTAccelerationStructureBuildPriority::Immediate);
+	OutRayTracingGeometry->LODIndex = LODIndex;
+	OutRayTracingGeometry->SetRequiresBuild(true);
 }
 
 static void BuildHairAccelerationStructure_Cards(FRHICommandList& RHICmdList,
@@ -1533,7 +1501,10 @@ static void BuildHairAccelerationStructure_Cards(FRHICommandList& RHICmdList,
 	const FString& MeshComponentName,
 	uint32 LODIndex)
 {
-	FRayTracingGeometryInitializer Initializer;
+	check(OutRayTracingGeometry);
+
+	FRayTracingGeometryInitializer Initializer = OutRayTracingGeometry->GetInitializer();
+
 #if HAIR_RESOURCE_DEBUG_NAME
 	FString DebugName(TEXT("AS_HairCards_LOD") + FString::FromInt(LODIndex) + TEXT("_") + AssetDebugName);
 	Initializer.DebugName = FName(DebugName);
@@ -1542,10 +1513,12 @@ static void BuildHairAccelerationStructure_Cards(FRHICommandList& RHICmdList,
 	static int32 DebugNumber = 0;
 	Initializer.DebugName = FDebugName(DebugName, DebugNumber++);
 #endif
+
 #if RHI_ENABLE_RESOURCE_INFO
 	const FName OwnerName(FString::Printf(TEXT("%s [LOD%d]"), *MeshComponentName, LODIndex));
 	Initializer.OwnerName = OwnerName;
 #endif
+
 	Initializer.IndexBuffer = RestResource->RestIndexBuffer.IndexBufferRHI;
 	Initializer.IndexBufferOffset = 0;
 	Initializer.GeometryType = RTGT_Triangles;
@@ -1556,7 +1529,7 @@ static void BuildHairAccelerationStructure_Cards(FRHICommandList& RHICmdList,
 	FBufferRHIRef PositionBuffer = RestResource->RestPositionBuffer.VertexBufferRHI;
 	if (DeformedResource)
 	{
-		PositionBuffer = DeformedResource->GetBuffer(FHairCardsDeformedResource::Current).Buffer->GetRHI(); // This will likely flicker result in half speed motion as everyother frame will use the wrong buffer
+		PositionBuffer = DeformedResource->GetBuffer(FHairCardsDeformedResource::Current).Buffer->GetRHI();
 	}
 
 	FRayTracingGeometrySegment Segment;
@@ -1565,10 +1538,14 @@ static void BuildHairAccelerationStructure_Cards(FRHICommandList& RHICmdList,
 	Segment.VertexBufferElementType = FHairCardsPositionFormat::VertexElementType;
 	Segment.NumPrimitives = RestResource->GetPrimitiveCount();
 	Segment.MaxVertices = RestResource->GetVertexCount();
+
+	// Replace segment array (cards don't have pre-existing segments from InternalAllocate)
+	Initializer.Segments.Empty();
 	Initializer.Segments.Add(Segment);
 
 	OutRayTracingGeometry->SetInitializer(Initializer);
-	OutRayTracingGeometry->CreateRayTracingGeometry(RHICmdList, ERTAccelerationStructureBuildPriority::Immediate);
+	OutRayTracingGeometry->LODIndex = LODIndex;
+	OutRayTracingGeometry->SetRequiresBuild(true);
 }
 
 static void BuildHairAccelerationStructure_Meshes(FRHICommandList& RHICmdList,
@@ -1578,7 +1555,10 @@ static void BuildHairAccelerationStructure_Meshes(FRHICommandList& RHICmdList,
 	const FString& AssetDebugName,
 	uint32 LODIndex)
 {
-	FRayTracingGeometryInitializer Initializer;
+	check(OutRayTracingGeometry);
+
+	FRayTracingGeometryInitializer Initializer = OutRayTracingGeometry->GetInitializer();
+
 #if HAIR_RESOURCE_DEBUG_NAME
 	FString DebugName(TEXT("AS_HairMeshes_LOD") + FString::FromInt(LODIndex) + TEXT("_") + AssetDebugName);
 	Initializer.DebugName = FName(DebugName);
@@ -1587,6 +1567,7 @@ static void BuildHairAccelerationStructure_Meshes(FRHICommandList& RHICmdList,
 	static int32 DebugNumber = 0;
 	Initializer.DebugName = FDebugName(DebugName, DebugNumber++);
 #endif
+
 	Initializer.IndexBuffer = RestResource->IndexBuffer.IndexBufferRHI;
 	Initializer.IndexBufferOffset = 0;
 	Initializer.GeometryType = RTGT_Triangles;
@@ -1596,8 +1577,8 @@ static void BuildHairAccelerationStructure_Meshes(FRHICommandList& RHICmdList,
 
 	FBufferRHIRef PositionBuffer = RestResource->RestPositionBuffer.VertexBufferRHI;
 	if (DeformedResource)
-	{	
-		PositionBuffer = DeformedResource->GetBuffer(FHairMeshesDeformedResource::Current).Buffer->GetRHI(); // This will likely flicker result in half speed motion as everyother frame will use the wrong buffer
+	{
+		PositionBuffer = DeformedResource->GetBuffer(FHairMeshesDeformedResource::Current).Buffer->GetRHI();
 	}
 
 	FRayTracingGeometrySegment Segment;
@@ -1606,156 +1587,258 @@ static void BuildHairAccelerationStructure_Meshes(FRHICommandList& RHICmdList,
 	Segment.VertexBufferElementType = FHairCardsPositionFormat::VertexElementType;
 	Segment.NumPrimitives = RestResource->GetPrimitiveCount();
 	Segment.MaxVertices = RestResource->GetVertexCount();
+
+	Initializer.Segments.Empty();
 	Initializer.Segments.Add(Segment);
 
 	OutRayTracingGeometry->SetInitializer(Initializer);
-	OutRayTracingGeometry->CreateRayTracingGeometry(RHICmdList, ERTAccelerationStructureBuildPriority::Immediate);
+	OutRayTracingGeometry->LODIndex = LODIndex;
+	OutRayTracingGeometry->SetRequiresBuild(true);
 }
 
-void AddBuildStrandsAccelerationStructurePass(
+void RegisterHairStrandsWithRayTracingManager(
 	FRDGBuilder& GraphBuilder,
+	FRayTracingDynamicGeometryUpdateManager* DynamicGeometryUpdateManager,
+	FScene* RenderScene,
+	const FSceneView* View,
 	FHairGroupInstance* Instance,
 	uint32 ProceduralSplits,
-	bool bNeedUpdate,
+	bool bNeedRaytracingBuild,
+	bool bNeedRaytracingUpdate,
 	FRDGBufferRef Raytracing_PositionBuffer,
 	FRDGBufferRef Raytracing_IndexBuffer)
 {
-	const uint32 HairLODIndex = Instance->HairGroupPublicData->GetIntLODIndex();
-	const EHairStrandsRayTracingPrimitiveType RayTracingPrimitiveType = Instance->Strands.RenRaytracingResource->RayTracingPrimitiveType;
+	if (!DynamicGeometryUpdateManager)
+		return;
 
-	FStrandsResourceBLASParameters* Parameters = GraphBuilder.AllocParameters < FStrandsResourceBLASParameters > ();
-	Parameters->PositionBuffer = Raytracing_PositionBuffer;
-	Parameters->IndexBuffer = Raytracing_IndexBuffer;
-	
-	GraphBuilder.AddPass(
-		RDG_EVENT_NAME("HairStrands::UpdateBLAS(Strands)"),
-		Parameters,
-		ERDGPassFlags::NeverCull | ERDGPassFlags::Compute,
-		[Instance, HairLODIndex, bNeedUpdate, RayTracingPrimitiveType, ProceduralSplits](FRHICommandListImmediate& RHICmdList)
+	FHairStrandsRaytracingResource* RT = Instance->Strands.RenRaytracingResource;
+	const EHairStrandsRayTracingPrimitiveType RayTracingPrimitiveType = RT->RayTracingPrimitiveType;
+
+	bool bNeedInitializerRebuild = false;
+
+	// Check 1: Geometry not valid (first frame or after reallocation)
+	if (!RT->RayTracingGeometry.IsValid())
+	{
+		bNeedInitializerRebuild = true;
+	}
+
+	// Check 2: ProceduralSplits changed
+	if (Instance->Strands.CachedProceduralSplits != ProceduralSplits)
+	{
+		bNeedInitializerRebuild = true;
+		Instance->Strands.CachedProceduralSplits = ProceduralSplits;
+	}
+
+	// Check 3: LOD changed
+	const int32 CurrentLOD = Instance->HairGroupPublicData ? Instance->HairGroupPublicData->GetIntLODIndex() : 0;
+	if (RT->RayTracingGeometry.LODIndex != CurrentLOD)
+	{
+		bNeedInitializerRebuild = true;
+	}
+
+	// Check 4: Buffer pointer changed (reallocation detected)
+	FBufferRHIRef CurrentPosBuffer = RT->PositionBuffer.Buffer->GetRHI();
+	if (RT->RayTracingGeometry.IsInitialized())
+	{
+		FRayTracingGeometryInitializer ExistingInitializer = RT->RayTracingGeometry.GetInitializer();
+		if (ExistingInitializer.Segments.Num() > 0 &&
+			ExistingInitializer.Segments[0].VertexBuffer != CurrentPosBuffer)
 		{
-			SCOPED_GPU_MASK(RHICmdList, FRHIGPUMask::All());
-	
-			const bool bLocalNeedBuild = !Instance->Strands.RenRaytracingResource->bIsRTGeometryInitialized || (Instance->Strands.CachedProceduralSplits != ProceduralSplits);
-			if (bLocalNeedBuild)
-			{
-				FBufferRHIRef PositionBuffer(Instance->Strands.RenRaytracingResource->PositionBuffer.Buffer->GetRHI());
-	
-				// no index buffer needed for procedural primitive
-				FBufferRHIRef IndexBuffer = nullptr;
-				if (RayTracingPrimitiveType == EHairStrandsRayTracingPrimitiveType::Triangle ||
-				    RayTracingPrimitiveType == EHairStrandsRayTracingPrimitiveType::LinearSweptSphere)
-				{
-					IndexBuffer = Instance->Strands.RenRaytracingResource->IndexBuffer.Buffer->GetRHI();
-				}
+			bNeedInitializerRebuild = true;
+		}
+	}
 
-				Instance->Strands.CachedProceduralSplits = ProceduralSplits;
-				BuildHairAccelerationStructure_Strands(
-					RHICmdList,
-					Instance->Strands.RenRaytracingResource->MaxVertexCount,
-					Instance->Strands.RenRaytracingResource->MaxIndexCount,
-					PositionBuffer,
-					IndexBuffer,
-					&Instance->Strands.RenRaytracingResource->RayTracingGeometry,
-					Instance->Debug.GroomAssetName,
-					Instance->Debug.MeshComponentName,
-					HairLODIndex,
-					RayTracingPrimitiveType,
-					ProceduralSplits
-				);
-			}
-			else if (bNeedUpdate)
-			{
-				// hair strands can move chaotically during simulation which can really tank ray tracing performance
-				// even though a rebuild is more expensive, we can still come out ahead overall, even with a single ray cast per pixel
-				UpdateHairAccelerationStructure(RHICmdList, &Instance->Strands.RenRaytracingResource->RayTracingGeometry, IsHairStrandsForceRebuildBVH() ? EAccelerationStructureBuildMode::Build : EAccelerationStructureBuildMode::Update);
-			}
-			Instance->Strands.RenRaytracingResource->bIsRTGeometryInitialized = true;
-		});
+	if (bNeedInitializerRebuild || bNeedRaytracingBuild)
+	{
+		FBufferRHIRef PosBuffer = CurrentPosBuffer;
+		FBufferRHIRef IdxBuffer = (RayTracingPrimitiveType != EHairStrandsRayTracingPrimitiveType::Procedural)
+										? RT->IndexBuffer.Buffer->GetRHI() : nullptr;
+
+		BuildHairAccelerationStructure_Strands(
+			GraphBuilder.RHICmdList,
+			RT->MaxVertexCount,
+			RT->MaxIndexCount,
+			PosBuffer,
+			IdxBuffer,
+			&RT->RayTracingGeometry,
+			Instance->Debug.GroomAssetName,
+			Instance->Debug.MeshComponentName,
+			CurrentLOD,
+			RayTracingPrimitiveType,
+			ProceduralSplits);
+	}
+
+	if (bNeedInitializerRebuild || bNeedRaytracingBuild || bNeedRaytracingUpdate)
+	{
+		FRayTracingDynamicGeometryUpdateParams Params;
+		Params.Geometry = &RT->RayTracingGeometry;
+		Params.ExternalDependencies.Add(Raytracing_PositionBuffer);
+
+		if (RayTracingPrimitiveType != EHairStrandsRayTracingPrimitiveType::Procedural)
+		{
+			Params.ExternalDependencies.Add(Raytracing_IndexBuffer);
+		}
+
+		DynamicGeometryUpdateManager->AddDynamicGeometryToUpdate(
+			GraphBuilder.RHICmdList, RenderScene, View, Instance->Debug.Proxy, Params, 0);
+	}
 }
 
-void AddBuildHairCardAccelerationStructurePass(
+void RegisterHairCardsWithRayTracingManager(
 	FRDGBuilder& GraphBuilder,
+	FRayTracingDynamicGeometryUpdateManager* DynamicGeometryUpdateManager,
+	FScene* RenderScene,
+	const FSceneView* View,
 	FHairGroupInstance* Instance,
 	int32 HairLODIndex,
-	bool bNeedUpdate)
+	bool bNeedDeformation)
 {
+	if (!DynamicGeometryUpdateManager)
+		return;
+
 	check(Instance);
 	check(Instance->Cards.LODs.IsValidIndex(HairLODIndex));
 
 	FHairGroupInstance::FCards::FLOD& LOD = Instance->Cards.LODs[HairLODIndex];
-	FCardsOrMeshesResourceBLASParameters* Parameters = GraphBuilder.AllocParameters<FCardsOrMeshesResourceBLASParameters>();
-	Parameters->PositionBuffer = LOD.DeformedResource ? Register(GraphBuilder, LOD.DeformedResource->GetBuffer(FHairCardsDeformedResource::Current), ERDGImportedBufferFlags::None).Buffer : nullptr;
 
-	// * When cards are dynamic, RT geometry is built once and then update.
-	// * When cards are static (i.e., no sim, not attached to  skinning()), in which case RT geometry needs only to be built once. This static geometry is shared between all the instances, 
-	//   and owned by the asset, rathter than the component. In this case the RT geometry will be built only once for all the instances
-	const bool bNeedBuild = !LOD.RaytracingResource->bIsRTGeometryInitialized;
-	if (bNeedBuild || bNeedUpdate)
+	// Detect if we need to rebuild initializer
+	bool bNeedInitializerRebuild = false;
+
+	// Check 1: Geometry not valid (first frame or after reallocation)
+	if (!LOD.RaytracingResource->RayTracingGeometry.IsValid())
 	{
-		GraphBuilder.AddPass(
-			RDG_EVENT_NAME("HairStrands::UpdateBLAS(Cards)"),
-			Parameters,
-			ERDGPassFlags::NeverCull | ERDGPassFlags::Compute,
-			[Instance, HairLODIndex, bNeedUpdate](FRHICommandListImmediate& RHICmdList)
+		bNeedInitializerRebuild = true;
+	}
+
+	// Check 2: LOD changed
+	if (LOD.RaytracingResource->RayTracingGeometry.LODIndex != HairLODIndex)
+	{
+		bNeedInitializerRebuild = true;
+	}
+
+	// Check 3: Buffer pointer changed (reallocation detected)
+	FRHIBuffer* CurrentPosBuffer = LOD.DeformedResource
+		? LOD.DeformedResource->GetBuffer(FHairCardsDeformedResource::Current).Buffer->GetRHI()
+		: LOD.RestResource->RestPositionBuffer.VertexBufferRHI.GetReference();
+
+	if (LOD.RaytracingResource->RayTracingGeometry.IsInitialized())
+	{
+		FRayTracingGeometryInitializer ExistingInitializer = LOD.RaytracingResource->RayTracingGeometry.GetInitializer();
+		if (ExistingInitializer.Segments.Num() > 0 &&
+			ExistingInitializer.Segments[0].VertexBuffer.GetReference() != CurrentPosBuffer)
 		{
-			SCOPED_GPU_MASK(RHICmdList, FRHIGPUMask::All());
+			bNeedInitializerRebuild = true;
+		}
+	}
+
+	// Rebuild initializer if needed
+	if (bNeedInitializerRebuild)
+	{
+		BuildHairAccelerationStructure_Cards(
+			GraphBuilder.RHICmdList,
+			LOD.RestResource,
+			LOD.DeformedResource,
+			&LOD.RaytracingResource->RayTracingGeometry,
+			Instance->Debug.GroomAssetName,
+			Instance->Debug.MeshComponentName,
+			HairLODIndex);
+	}
 
-			FHairGroupInstance::FCards::FLOD& LocalLOD = Instance->Cards.LODs[HairLODIndex];
+	// Always register with manager (even if initializer unchanged, position data changes each frame)
+	if (bNeedInitializerRebuild || bNeedDeformation)
+	{
+		FRayTracingDynamicGeometryUpdateParams Params;
+		Params.Geometry = &LOD.RaytracingResource->RayTracingGeometry;
 
-			const bool bLocalNeedBuild = !LocalLOD.RaytracingResource->bIsRTGeometryInitialized;
-			if (bLocalNeedBuild)
-			{
-				BuildHairAccelerationStructure_Cards(RHICmdList, LocalLOD.RestResource, LocalLOD.DeformedResource, &LocalLOD.RaytracingResource->RayTracingGeometry, Instance->Debug.GroomAssetName, Instance->Debug.MeshComponentName, HairLODIndex);
-				LocalLOD.RaytracingResource->bIsRTGeometryInitialized = true;
-			}
-			else if (bNeedUpdate)
-			{
-				// TODO: evaluate perf tradeoff of rebuild vs refit
-				UpdateHairAccelerationStructure(RHICmdList, &LocalLOD.RaytracingResource->RayTracingGeometry, EAccelerationStructureBuildMode::Update);
-			}
-		});
+		// Add deformed position buffer (index buffers for cards are static)
+		if (LOD.DeformedResource)
+		{
+			FRDGImportedBuffer DeformedBuffer = Register(GraphBuilder, LOD.DeformedResource->GetBuffer(FHairCardsDeformedResource::Current), ERDGImportedBufferFlags::None);
+			Params.ExternalDependencies.Add(DeformedBuffer.Buffer);
+		}
+
+		DynamicGeometryUpdateManager->AddDynamicGeometryToUpdate(
+			GraphBuilder.RHICmdList, RenderScene, View, Instance->Debug.Proxy, Params, 0);
 	}
 }
 
-void AddBuildHairMeshAccelerationStructurePass(
+void RegisterHairMeshesWithRayTracingManager(
 	FRDGBuilder& GraphBuilder,
+	FRayTracingDynamicGeometryUpdateManager* DynamicGeometryUpdateManager,
+	FScene* RenderScene,
+	const FSceneView* View,
 	FHairGroupInstance* Instance,
 	int32 HairLODIndex,
-	bool bNeedUpdate)
+	bool bNeedDeformation)
 {
+	if (!DynamicGeometryUpdateManager)
+		return;
+
 	check(Instance);
 	check(Instance->Meshes.LODs.IsValidIndex(HairLODIndex));
 
 	FHairGroupInstance::FMeshes::FLOD& LOD = Instance->Meshes.LODs[HairLODIndex];
-	FCardsOrMeshesResourceBLASParameters* Parameters = GraphBuilder.AllocParameters<FCardsOrMeshesResourceBLASParameters>();
-	Parameters->PositionBuffer = LOD.DeformedResource ? Register(GraphBuilder, LOD.DeformedResource->GetBuffer(FHairMeshesDeformedResource::Current), ERDGImportedBufferFlags::None).Buffer : nullptr;
 
-	const bool bNeedBuid = !LOD.RaytracingResource->bIsRTGeometryInitialized;
-	if (bNeedBuid || bNeedUpdate)
+	// Detect if we need to rebuild initializer
+	bool bNeedInitializerRebuild = false;
+
+	// Check 1: Geometry not valid (first frame or after reallocation)
+	if (!LOD.RaytracingResource->RayTracingGeometry.IsValid())
 	{
-		GraphBuilder.AddPass(
-			RDG_EVENT_NAME("HairStrands::UpdateBLAS(Meshes)"),
-			Parameters,
-			ERDGPassFlags::NeverCull | ERDGPassFlags::Compute,
-			[Instance, HairLODIndex, bNeedUpdate](FRHICommandListImmediate& RHICmdList)
-			{
-				SCOPED_GPU_MASK(RHICmdList, FRHIGPUMask::All());
+		bNeedInitializerRebuild = true;
+	}
 
-				FHairGroupInstance::FMeshes::FLOD& LocalLOD = Instance->Meshes.LODs[HairLODIndex];	
-				const bool bLocalNeedBuild = !LocalLOD.RaytracingResource->bIsRTGeometryInitialized;
-				if (bLocalNeedBuild)
-				{
-					BuildHairAccelerationStructure_Meshes(RHICmdList, LocalLOD.RestResource, LocalLOD.DeformedResource,  &LocalLOD.RaytracingResource->RayTracingGeometry, Instance->Debug.GroomAssetName, HairLODIndex);
-					LocalLOD.RaytracingResource->bIsRTGeometryInitialized = true;
-				}
-				else if (bNeedUpdate)
-				{
-					// TODO: evaluate perf tradeoff of rebuild vs refit
-					UpdateHairAccelerationStructure(RHICmdList, &LocalLOD.RaytracingResource->RayTracingGeometry, EAccelerationStructureBuildMode::Update);
-				}
-			});
+	// Check 2: LOD changed
+	if (LOD.RaytracingResource->RayTracingGeometry.LODIndex != HairLODIndex)
+	{
+		bNeedInitializerRebuild = true;
+	}
+
+	// Check 3: Buffer pointer changed (reallocation detected)
+	FRHIBuffer* CurrentPosBuffer = LOD.DeformedResource
+		? LOD.DeformedResource->GetBuffer(FHairMeshesDeformedResource::Current).Buffer->GetRHI()
+		: LOD.RestResource->RestPositionBuffer.VertexBufferRHI.GetReference();
+
+	if (LOD.RaytracingResource->RayTracingGeometry.IsInitialized())
+	{
+		FRayTracingGeometryInitializer ExistingInitializer = LOD.RaytracingResource->RayTracingGeometry.GetInitializer();
+		if (ExistingInitializer.Segments.Num() > 0 &&
+			ExistingInitializer.Segments[0].VertexBuffer.GetReference() != CurrentPosBuffer)
+		{
+			bNeedInitializerRebuild = true;
+		}
+	}
+
+	// Rebuild initializer if needed
+	if (bNeedInitializerRebuild)
+	{
+		BuildHairAccelerationStructure_Meshes(
+			GraphBuilder.RHICmdList,
+			LOD.RestResource,
+			LOD.DeformedResource,
+			&LOD.RaytracingResource->RayTracingGeometry,
+			Instance->Debug.GroomAssetName,
+			HairLODIndex);
+	}
+
+	// Always register with manager (even if initializer unchanged, position data changes each frame)
+	if (bNeedInitializerRebuild || bNeedDeformation)
+	{
+		FRayTracingDynamicGeometryUpdateParams Params;
+		Params.Geometry = &LOD.RaytracingResource->RayTracingGeometry;
+
+		// Add deformed position buffer (index buffers for meshes are static)
+		if (LOD.DeformedResource)
+		{
+			FRDGImportedBuffer DeformedBuffer = Register(GraphBuilder, LOD.DeformedResource->GetBuffer(FHairMeshesDeformedResource::Current), ERDGImportedBufferFlags::None);
+			Params.ExternalDependencies.Add(DeformedBuffer.Buffer);
+		}
+
+		DynamicGeometryUpdateManager->AddDynamicGeometryToUpdate(
+			GraphBuilder.RHICmdList, RenderScene, View, Instance->Debug.Proxy, Params, 0);
 	}
 }
+
 #endif // RHI_RAYTRACING
 
 static void ConvertHairStrandsVFParameters(
diff --git a/Engine/Plugins/Runtime/HairStrands/Source/HairStrandsCore/Public/GroomResources.h b/Engine/Plugins/Runtime/HairStrands/Source/HairStrandsCore/Public/GroomResources.h
--- a/Engine/Plugins/Runtime/HairStrands/Source/HairStrandsCore/Public/GroomResources.h
+++ b/Engine/Plugins/Runtime/HairStrands/Source/HairStrandsCore/Public/GroomResources.h
@@ -688,7 +688,6 @@ struct FHairStrandsRaytracingResource : public FHairCommonResource
 	uint32 MaxVertexCount = 0;
 	uint32 MaxIndexCount = 0;
 	bool bOwnBuffers = false;
-	bool bIsRTGeometryInitialized = false;
 	EHairStrandsRayTracingPrimitiveType RayTracingPrimitiveType = EHairStrandsRayTracingPrimitiveType::Triangle;
 };
 #endif
diff --git a/Engine/Plugins/Runtime/HairStrands/Source/HairStrandsCore/Public/HairStrandsInterpolation.h b/Engine/Plugins/Runtime/HairStrands/Source/HairStrandsCore/Public/HairStrandsInterpolation.h
--- a/Engine/Plugins/Runtime/HairStrands/Source/HairStrandsCore/Public/HairStrandsInterpolation.h
+++ b/Engine/Plugins/Runtime/HairStrands/Source/HairStrandsCore/Public/HairStrandsInterpolation.h
@@ -16,6 +16,7 @@ class FRDGBuilder;
 class FSkeletalMeshRenderData;
 class UGroomBindingAsset;
 class FRHIShaderResourceView;
+class FRayTracingDynamicGeometryUpdateManager;
 
 struct FHairGroupInstance;
 struct FHairStrandClusterData;
@@ -209,25 +210,35 @@ void AddGenerateRaytracingGeometryPass(
 	const FRDGBufferUAVRef& OutPositionBuffer,
 	const FRDGBufferUAVRef& OutIndexBuffer);
 
-void AddBuildStrandsAccelerationStructurePass(
+void RegisterHairStrandsWithRayTracingManager(
 	FRDGBuilder& GraphBuilder,
+	FRayTracingDynamicGeometryUpdateManager* DynamicGeometryUpdateManager,
+	FScene* RenderScene,
+	const FSceneView* View,
 	FHairGroupInstance* Instance,
 	uint32 ProceduralSplits,
-	bool bNeedUpdate,
+	bool bNeedRaytracingBuild,
+	bool bNeedRaytracingUpdate,
 	FRDGBufferRef Raytracing_PositionBuffer,
 	FRDGBufferRef Raytracing_IndexBuffer);
 
-void AddBuildHairCardAccelerationStructurePass(
+void RegisterHairCardsWithRayTracingManager(
 	FRDGBuilder& GraphBuilder,
+	FRayTracingDynamicGeometryUpdateManager* DynamicGeometryUpdateManager,
+	FScene* RenderScene,
+	const FSceneView* View,
 	FHairGroupInstance* Instance,
 	int32 HairLODIndex,
-	bool bNeedUpdate);
+	bool bNeedDeformation);
 
-void AddBuildHairMeshAccelerationStructurePass(
+void RegisterHairMeshesWithRayTracingManager(
 	FRDGBuilder& GraphBuilder,
+	FRayTracingDynamicGeometryUpdateManager* DynamicGeometryUpdateManager,
+	FScene* RenderScene,
+	const FSceneView* View,
 	FHairGroupInstance* Instance,
 	int32 HairLODIndex,
-	bool bNeedUpdate);
+	bool bNeedDeformation);
 #endif
 
 void AddHairCardsDeformationPass(
diff --git a/Engine/Source/Runtime/Engine/Public/SceneManagement.h b/Engine/Source/Runtime/Engine/Public/SceneManagement.h
--- a/Engine/Source/Runtime/Engine/Public/SceneManagement.h
+++ b/Engine/Source/Runtime/Engine/Public/SceneManagement.h
@@ -1494,6 +1494,8 @@ struct FRayTracingDynamicGeometryUpdateParams
 
 	bool bAlphaMasked = false;
 	FBufferRHIRef IndexBuffer = nullptr;
+
+	TArray<FRDGBufferRef> ExternalDependencies;
 };
 
 namespace RayTracing
diff --git a/Engine/Source/Runtime/Renderer/Private/DeferredShadingRenderer.cpp b/Engine/Source/Runtime/Renderer/Private/DeferredShadingRenderer.cpp
--- a/Engine/Source/Runtime/Renderer/Private/DeferredShadingRenderer.cpp
+++ b/Engine/Source/Runtime/Renderer/Private/DeferredShadingRenderer.cpp
@@ -2394,6 +2394,15 @@ void FDeferredShadingSceneRenderer::Render(FRDGBuilder& GraphBuilder, const FSce
 	if (IsHairStrandsEnabled(EHairStrandsShaderType::All, Scene->GetShaderPlatform()) && RendererOutput == ERendererOutput::FinalSceneColor)
 	{
 		CreateHairStrandsBookmarkParameters(Scene, Views, AllViews, HairStrandsBookmarkParameters);
+#if RHI_RAYTRACING
+		// Only set DynamicGeometryUpdateManager when an RT pass is active this frame.
+		// The manager's BeginUpdate/Clear only runs under the bRayTracing guard, so registering
+		// without an active RT pass would accumulate stale state and assert when RT activates.
+		if (FamilyPipelineState[&FFamilyPipelineState::bRayTracing])
+		{
+			HairStrandsBookmarkParameters.DynamicGeometryUpdateManager = Scene->GetRayTracingDynamicGeometryUpdateManager();
+		}
+#endif
 		check(Scene->HairStrandsSceneData.TransientResources);
 		HairStrandsBookmarkParameters.TransientResources = Scene->HairStrandsSceneData.TransientResources;
 		RunHairStrandsBookmark(GraphBuilder, EHairStrandsBookmark::ProcessTasks, HairStrandsBookmarkParameters);
@@ -2956,9 +2965,16 @@ void FDeferredShadingSceneRenderer::Render(FRDGBuilder& GraphBuilder, const FSce
 			TranslucencyLightingVolumeTextures.Init(GraphBuilder, Views, ERDGPassFlags::AsyncCompute);
 		}
 
+		// Run strands interpolation early to ensure it completes before DispatchRayTracingWorldUpdates and base pass
+		if (bHairStrandsEnable && !IsForwardShadingEnabled(ShaderPlatform))
+		{
+			RDG_EVENT_SCOPE(GraphBuilder, "Hair");
+			RunHairStrandsBookmark(GraphBuilder, EHairStrandsBookmark::ProcessStrandsInterpolation, HairStrandsBookmarkParameters);
+		}
+
 		FRDGBufferRef DynamicGeometryScratchBuffer = nullptr;
 #if RHI_RAYTRACING
-		
+
 		ERHIPipeline DynamicRTResourceAccessPipelines = Lumen::UseAsyncCompute(ViewFamily) ? ERHIPipeline::All : ERHIPipeline::Graphics;
 
 		// Async AS builds can potentially overlap with BasePass.
@@ -3133,7 +3149,7 @@ void FDeferredShadingSceneRenderer::Render(FRDGBuilder& GraphBuilder, const FSce
 		{
 			RDG_EVENT_SCOPE(GraphBuilder, "Hair");
 
-			RunHairStrandsBookmark(GraphBuilder, EHairStrandsBookmark::ProcessStrandsInterpolation, HairStrandsBookmarkParameters);
+			// Interpolation already ran earlier (before DispatchRayTracingWorldUpdates / base pass)
 			if (!bHasRayTracedOverlay)
 			{
 				RenderHairPrePass(GraphBuilder, Scene, SceneTextures, Views, InstanceCullingManager, HairStrandsBookmarkParameters.CullingResults);
diff --git a/Engine/Source/Runtime/Renderer/Private/RayTracing/RayTracingDynamicGeometryUpdateManager.cpp b/Engine/Source/Runtime/Renderer/Private/RayTracing/RayTracingDynamicGeometryUpdateManager.cpp
--- a/Engine/Source/Runtime/Renderer/Private/RayTracing/RayTracingDynamicGeometryUpdateManager.cpp
+++ b/Engine/Source/Runtime/Renderer/Private/RayTracing/RayTracingDynamicGeometryUpdateManager.cpp
@@ -429,23 +429,18 @@ void FRayTracingDynamicGeometryUpdateManager::AddDynamicGeometryToUpdate(
 		AddDispatchCommands(RHICmdList, Scene, View, PrimitiveSceneProxy, UpdateParams, PrimitiveId, RWBuffer, NumVertices, VertexBufferOffset, VertexBufferSize, GeometryBuildParams);
 	}
 
-	bool bRefit = true;
+	bool bRecreateGeometry = false;
 
 	// Optionally resize the buffer when not shared (could also be lazy allocated and still empty)
 	if (!bUseSharedVertexBuffer && RWBuffer && RWBuffer->NumBytes != VertexBufferSize)
 	{
 		RWBuffer->Initialize(RHICmdList, TEXT("FRayTracingDynamicGeometryUpdateManager::RayTracingDynamicVertexBuffer"), sizeof(float), VertexBufferSize / sizeof(float), PF_R32_FLOAT, BUF_UnorderedAccess | BUF_ShaderResource);
-		bRefit = false;
+		bRecreateGeometry = true;
 	}
 
 	if (!Geometry.IsValid() || Geometry.IsEvicted())
 	{
-		bRefit = false;
-	}
-
-	if (!Geometry.Initializer.bAllowUpdate)
-	{
-		bRefit = false;
+		bRecreateGeometry = true;
 	}
 
 	check(Geometry.IsInitialized());
@@ -462,7 +457,7 @@ void FRayTracingDynamicGeometryUpdateManager::AddDynamicGeometryToUpdate(
 		Segment.NumPrimitives = UpdateParams.NumTriangles;
 		Segment.MaxVertices = NumVertices;
 		Geometry.Initializer.Segments.Add(Segment);
-		bRefit = false;
+		bRecreateGeometry = true;
 	}
 
 	if (UpdateParams.bAlphaMasked)
@@ -488,17 +483,18 @@ void FRayTracingDynamicGeometryUpdateManager::AddDynamicGeometryToUpdate(
 	}
 #endif
 
-	if (!bRefit)
+	if (bRecreateGeometry)
 	{
 		checkf(Geometry.RawData.IsEmpty() && Geometry.Initializer.OfflineData == nullptr, TEXT("Dynamic ray tracing geometry '%s' is not expected to have offline acceleration structure data."), *Geometry.Initializer.DebugName.ToString());
 		Geometry.CreateRayTracingGeometry(RHICmdList, ERTAccelerationStructureBuildPriority::Skip);
 	}
 
-	EAccelerationStructureBuildMode BuildMode = Geometry.GetRequiresBuild()
+	EAccelerationStructureBuildMode BuildMode = (Geometry.GetRequiresBuild() || !Geometry.Initializer.bAllowUpdate)
 		? EAccelerationStructureBuildMode::Build
 		: EAccelerationStructureBuildMode::Update;
 
 	GeometryBuildParams.Geometry = UpdateParams.Geometry;
+	GeometryBuildParams.ExternalDependencies = UpdateParams.ExternalDependencies;
 
 	if (bUseSharedVertexBuffer)
 	{
@@ -721,6 +717,7 @@ void FRayTracingDynamicGeometryUpdateManager::AddDispatchCommands(
 
 BEGIN_SHADER_PARAMETER_STRUCT(FRayTracingDynamicGeometryUpdatePassParams, )
 	RDG_BUFFER_ACCESS(DynamicGeometryScratchBuffer, ERHIAccess::UAVCompute)
+	RDG_BUFFER_ACCESS_ARRAY(ExternalDependencies)
 
 	SHADER_PARAMETER_STRUCT_INCLUDE(FViewShaderParameters, View)
 	SHADER_PARAMETER_RDG_UNIFORM_BUFFER(FSceneUniformParameters, Scene)
@@ -953,6 +950,23 @@ void FRayTracingDynamicGeometryUpdateManager::AddDynamicGeometryUpdatePass(FRDGB
 		PassParams->Scene = nullptr;
 		PassParams->DynamicGeometryScratchBuffer = OutDynamicGeometryScratchBuffer;
 
+		// Collect all external dependencies from registered geometries (e.g., hair RT buffers)
+		TSet<FRDGBufferRef> UniqueExternalDependencies;
+		for (const FRayTracingDynamicGeometryBuildParams& BuildParam : DynamicGeometryBuilds)
+		{
+			UniqueExternalDependencies.Append(BuildParam.ExternalDependencies);
+		}
+		for (const FRayTracingDynamicGeometryBuildParams& UpdateParam : DynamicGeometryUpdates)
+		{
+			UniqueExternalDependencies.Append(UpdateParam.ExternalDependencies);
+		}
+
+		PassParams->ExternalDependencies.Reserve(UniqueExternalDependencies.Num());
+		for (FRDGBufferRef Buffer : UniqueExternalDependencies)
+		{
+			PassParams->ExternalDependencies.Add(FRDGBufferAccess(Buffer, ERHIAccess::SRVCompute));
+		}
+
 		GraphBuilder.AddPass(RDG_EVENT_NAME("RayTracingDynamicUpdateBuild"), PassParams, ComputePassFlags | ERDGPassFlags::NeverCull,
 			[this, PassParams, bSlowPath](FRHICommandList& RHICmdList)
 			{
diff --git a/Engine/Source/Runtime/Renderer/Public/HairStrandsInterface.h b/Engine/Source/Runtime/Renderer/Public/HairStrandsInterface.h
--- a/Engine/Source/Runtime/Renderer/Public/HairStrandsInterface.h
+++ b/Engine/Source/Runtime/Renderer/Public/HairStrandsInterface.h
@@ -31,6 +31,7 @@
 class UTexture2D;
 class FSceneInterface;
 class FGlobalShaderMap;
+class FRayTracingDynamicGeometryUpdateManager;
 
 class FViewInfo;
 struct FShaderPrintData;
@@ -604,7 +605,9 @@ struct FHairStrandsBookmarkParameters
 	FSceneInterface* Scene = nullptr;
 	TArray<const FSceneView*> AllViews;
 	FRDGTextureRef SceneColorTexture = nullptr;
-	FRDGTextureRef SceneDepthTexture = nullptr; 
+	FRDGTextureRef SceneDepthTexture = nullptr;
+
+	FRayTracingDynamicGeometryUpdateManager* DynamicGeometryUpdateManager = nullptr;
 
 	inline bool HasInstances() const { return Instances != nullptr && Instances->Num() > 0; }
 };
diff --git a/Engine/Source/Runtime/Renderer/Public/RayTracingDynamicGeometryUpdateManager.h b/Engine/Source/Runtime/Renderer/Public/RayTracingDynamicGeometryUpdateManager.h
--- a/Engine/Source/Runtime/Renderer/Public/RayTracingDynamicGeometryUpdateManager.h
+++ b/Engine/Source/Runtime/Renderer/Public/RayTracingDynamicGeometryUpdateManager.h
@@ -48,11 +48,11 @@ public:
 	/** Add dynamic geometry to update including CS shader to deform the vertices */
 	RENDERER_API void AddDynamicGeometryToUpdate(
 		FRHICommandListBase& RHICmdList,
-		const FScene* Scene, 
-		const FSceneView* View, 
-		const FPrimitiveSceneProxy* PrimitiveSceneProxy, 
+		const FScene* Scene,
+		const FSceneView* View,
+		const FPrimitiveSceneProxy* PrimitiveSceneProxy,
 		const FRayTracingDynamicGeometryUpdateParams& Params,
-		uint32 PrimitiveId 
+		uint32 PrimitiveId
 	);
 
 	/** Starts an update batch and returns the current shared buffer generation ID which is used for validation. */
@@ -85,6 +85,7 @@ private:
 		FRayTracingGeometry* Geometry = nullptr;
 		int32 SegmentOffset = -1;
 		float Priority = 0.0;
+		TArray<FRDGBufferRef> ExternalDependencies;
 	};
 
 	void AddDispatchCommands(FRHICommandListBase& RHICmdList,
