Fix compilation error

#rnx
#Synced-CL 52201145

diff --git a/Engine/Plugins/Runtime/HairStrands/Source/HairStrandsCore/Public/HairStrandsInterpolation.h b/Engine/Plugins/Runtime/HairStrands/Source/HairStrandsCore/Public/HairStrandsInterpolation.h
--- a/Engine/Plugins/Runtime/HairStrands/Source/HairStrandsCore/Public/HairStrandsInterpolation.h
+++ b/Engine/Plugins/Runtime/HairStrands/Source/HairStrandsCore/Public/HairStrandsInterpolation.h
@@ -10,7 +10,7 @@
 #include "HairStrandsInterface.h"
 #include "GroomResources.h"
 
-
+class FScene;
 class FGlobalShaderMap;
 class FRDGBuilder;
 class FSkeletalMeshRenderData;
