Disable groom binding hash validation.

#Synced-CL 52245098

diff --git a/Engine/Plugins/Runtime/HairStrands/Source/HairStrandsCore/Private/HairStrandsInterpolation.cpp b/Engine/Plugins/Runtime/HairStrands/Source/HairStrandsCore/Private/HairStrandsInterpolation.cpp
--- a/Engine/Plugins/Runtime/HairStrands/Source/HairStrandsCore/Private/HairStrandsInterpolation.cpp
+++ b/Engine/Plugins/Runtime/HairStrands/Source/HairStrandsCore/Private/HairStrandsInterpolation.cpp
@@ -65,7 +65,7 @@ static FAutoConsoleVariableRef CVarHairStrands_ScaleHairRadiusByTransform(TEXT("
 static int32 GVarHairStrands_AddDeformerStrandsDisplacement = 0;
 static FAutoConsoleVariableRef CVarHairStrands_AddDeformerStrandsDisplacement(TEXT("r.HairStrands.Strands.AddDeformerStrandsDisplacement"), GVarHairStrands_AddDeformerStrandsDisplacement, TEXT("Add the deformer displacement to the current position. If set to 0 the deformer position will overwrite the deformed one"));
 
-static int32 GVarHairStrands_BindingDataValidation = 1;
+static int32 GVarHairStrands_BindingDataValidation = 0;
 static FAutoConsoleVariableRef CVarVarHairStrands_BindingDataValidation(TEXT("r.HairStrands.Strands.BindingDataValidation"), GVarHairStrands_BindingDataValidation, TEXT("Add binding data validation at runtime (for debugging purpose)"));
 
 ////////////////////////////////////////////////////////////////////////////////////////////////////////////////////////
