Set RBF groom deformeroutput to use the Source

#Synced-CL 52257317

diff --git a/Engine/Plugins/Runtime/HairStrands/Source/HairStrandsCore/Private/GroomRBFDeformer.cpp b/Engine/Plugins/Runtime/HairStrands/Source/HairStrandsCore/Private/GroomRBFDeformer.cpp
--- a/Engine/Plugins/Runtime/HairStrands/Source/HairStrandsCore/Private/GroomRBFDeformer.cpp
+++ b/Engine/Plugins/Runtime/HairStrands/Source/HairStrandsCore/Private/GroomRBFDeformer.cpp
@@ -501,7 +501,7 @@ static void ApplyDeformationToGroom(const TArray<FRBFDeformedPositions>& Deforme
 		}
 	}
 
-	GroomAsset->CommitHairDescription(MoveTemp(HairDescription), EHairDescriptionType::Edit);
+	GroomAsset->CommitHairDescription(MoveTemp(HairDescription), EHairDescriptionType::Source);
 	GroomAsset->UpdateHairGroupsInfo();
 
 	// Update/reimport the cards/meshes geometry which have been deformed prior to call this function
