Deprecate ChangeFeatureLevel/ChangePlatformLevel(ERHIFeatureLevel::Type) on UGroomBindingAsset and UGroomAsset, and HandlePlatformPreviewChanged/HandleFeatureLevelChanged(ERHIFeatureLevel::Type) on UGroomComponent.
Replace with ChangeShaderPlatform(EShaderPlatform) and HandlePlatformPreviewChanged(EShaderPlatform).
Replace the two separate cached feature/platform level members into a single EShaderPlatform CachedResourceShaderPlatform

#Synced-CL 52291475

diff --git a/Engine/Plugins/Runtime/HairStrands/Source/HairStrandsCore/Private/GroomAsset.cpp b/Engine/Plugins/Runtime/HairStrands/Source/HairStrandsCore/Private/GroomAsset.cpp
--- a/Engine/Plugins/Runtime/HairStrands/Source/HairStrandsCore/Private/GroomAsset.cpp
+++ b/Engine/Plugins/Runtime/HairStrands/Source/HairStrandsCore/Private/GroomAsset.cpp
@@ -4090,21 +4090,21 @@ void UGroomAsset::RecreateResources()
 
 void UGroomAsset::ChangePlatformLevel(ERHIFeatureLevel::Type In)
 {
-	// When changing platform preview level, recreate resources to the correct platform settings (e.g., r.hairstrands.strands=0/1)
-	if (CachedResourcesPlatformLevel != In)
-	{
-		RecreateResources();
-		CachedResourcesPlatformLevel = In;
-	}
+	ChangeShaderPlatform(GShaderPlatformForFeatureLevel[In]);
 }
 
 void UGroomAsset::ChangeFeatureLevel(ERHIFeatureLevel::Type In)
 {
-	// When changing feature level, recreate resources to the correct feature level
-	if (CachedResourcesFeatureLevel != In)
+	ChangeShaderPlatform(GShaderPlatformForFeatureLevel[In]);
+}
+
+void UGroomAsset::ChangeShaderPlatform(EShaderPlatform InShaderPlatform)
+{
+	// When changing platform preview, recreate resources to the correct platform settings (e.g., r.hairstrands.strands=0/1)
+	if (CachedResourceShaderPlatform != InShaderPlatform)
 	{
 		RecreateResources();
-		CachedResourcesFeatureLevel = In;
+		CachedResourceShaderPlatform = InShaderPlatform;
 	}
 }
 #endif
diff --git a/Engine/Plugins/Runtime/HairStrands/Source/HairStrandsCore/Private/GroomBindingAsset.cpp b/Engine/Plugins/Runtime/HairStrands/Source/HairStrandsCore/Private/GroomBindingAsset.cpp
--- a/Engine/Plugins/Runtime/HairStrands/Source/HairStrandsCore/Private/GroomBindingAsset.cpp
+++ b/Engine/Plugins/Runtime/HairStrands/Source/HairStrandsCore/Private/GroomBindingAsset.cpp
@@ -1981,21 +1981,21 @@ void UGroomBindingAsset::RecreateResources()
 
 void UGroomBindingAsset::ChangeFeatureLevel(ERHIFeatureLevel::Type In)
 {
-	// When changing feature level, recreate resources to the correct feature level
-	if (CachedResourcesFeatureLevel != In)
-	{
-		RecreateResources();
-		CachedResourcesFeatureLevel = In;
-	}
+	ChangeShaderPlatform(GShaderPlatformForFeatureLevel[In]);
 }
 
 void UGroomBindingAsset::ChangePlatformLevel(ERHIFeatureLevel::Type In)
 {
-	// When changing platform preview level, recreate resources to the correct platform settings (e.g., r.hairstrands.strands=0/1)
-	if (CachedResourcesPlatformLevel != In)
+	ChangeShaderPlatform(GShaderPlatformForFeatureLevel[In]);
+}
+
+void UGroomBindingAsset::ChangeShaderPlatform(EShaderPlatform InShaderPlatform)
+{
+	// When changing platform preview, recreate resources to the correct platform settings (e.g., r.hairstrands.strands=0/1)
+	if (CachedResourceShaderPlatform != InShaderPlatform)
 	{
 		RecreateResources();
-		CachedResourcesPlatformLevel = In;
+		CachedResourceShaderPlatform = InShaderPlatform;
 	}
 }
 #endif // WITH_EDITOR
diff --git a/Engine/Plugins/Runtime/HairStrands/Source/HairStrandsCore/Private/GroomComponent.cpp b/Engine/Plugins/Runtime/HairStrands/Source/HairStrandsCore/Private/GroomComponent.cpp
--- a/Engine/Plugins/Runtime/HairStrands/Source/HairStrandsCore/Private/GroomComponent.cpp
+++ b/Engine/Plugins/Runtime/HairStrands/Source/HairStrandsCore/Private/GroomComponent.cpp
@@ -4314,26 +4314,23 @@ void UGroomComponent::PreFeatureLevelChange(ERHIFeatureLevel::Type PendingFeatur
 
 void UGroomComponent::HandlePlatformPreviewChanged(ERHIFeatureLevel::Type InFeatureLevel)
 {
-	if (BindingAsset)
-	{
-		BindingAsset->ChangePlatformLevel(InFeatureLevel);
-	}
-	if (GroomAsset)
-	{
-		GroomAsset->ChangePlatformLevel(InFeatureLevel);
-	}
-	InvalidateAndRecreate();
+	HandlePlatformPreviewChanged(GShaderPlatformForFeatureLevel[InFeatureLevel]);
 }
 
 void UGroomComponent::HandleFeatureLevelChanged(ERHIFeatureLevel::Type InFeatureLevel)
+{
+	HandlePlatformPreviewChanged(GShaderPlatformForFeatureLevel[InFeatureLevel]);
+}
+
+void UGroomComponent::HandlePlatformPreviewChanged(EShaderPlatform InShaderPlatform)
 {
 	if (BindingAsset)
 	{
-		BindingAsset->ChangeFeatureLevel(InFeatureLevel);
+		BindingAsset->ChangeShaderPlatform(InShaderPlatform);
 	}
 	if (GroomAsset)
 	{
-		GroomAsset->ChangeFeatureLevel(InFeatureLevel);
+		GroomAsset->ChangeShaderPlatform(InShaderPlatform);
 	}
 	InvalidateAndRecreate();
 }
diff --git a/Engine/Plugins/Runtime/HairStrands/Source/HairStrandsCore/Public/GroomAsset.h b/Engine/Plugins/Runtime/HairStrands/Source/HairStrandsCore/Public/GroomAsset.h
--- a/Engine/Plugins/Runtime/HairStrands/Source/HairStrandsCore/Public/GroomAsset.h
+++ b/Engine/Plugins/Runtime/HairStrands/Source/HairStrandsCore/Public/GroomAsset.h
@@ -817,8 +817,11 @@ public:
 	UE_API void MarkMaterialsHasChanged();
 
 	UE_API void RecreateResources();
+	UE_DEPRECATED(5.8, "Use ChangeShaderPlatform(EShaderPlatform) instead.")
 	UE_API void ChangeFeatureLevel(ERHIFeatureLevel::Type PendingFeatureLevel);
+	UE_DEPRECATED(5.8, "Use ChangeShaderPlatform(EShaderPlatform) instead.")
 	UE_API void ChangePlatformLevel(ERHIFeatureLevel::Type PendingFeatureLevel);
+	UE_API void ChangeShaderPlatform(EShaderPlatform InShaderPlatform);
 #endif
 
 	/** Array of user data stored with the asset */
@@ -954,8 +957,7 @@ private:
 	TArray<FHairGroupsLOD>						CachedHairGroupsLOD;
 	TArray<FHairGroupsCardsSourceDescription>	CachedHairGroupsCards;
 	TArray<FHairGroupsMeshesSourceDescription>	CachedHairGroupsMeshes;
-	ERHIFeatureLevel::Type 						CachedResourcesPlatformLevel= ERHIFeatureLevel::Num;
-	ERHIFeatureLevel::Type 						CachedResourcesFeatureLevel = ERHIFeatureLevel::Num;
+	EShaderPlatform								CachedResourceShaderPlatform = SP_NumPlatforms;
 
 	// Queue of procedural assets which needs to be saved
 	TQueue<UStaticMesh*> AssetToSave_Meshes;
diff --git a/Engine/Plugins/Runtime/HairStrands/Source/HairStrandsCore/Public/GroomBindingAsset.h b/Engine/Plugins/Runtime/HairStrands/Source/HairStrandsCore/Public/GroomBindingAsset.h
--- a/Engine/Plugins/Runtime/HairStrands/Source/HairStrandsCore/Public/GroomBindingAsset.h
+++ b/Engine/Plugins/Runtime/HairStrands/Source/HairStrandsCore/Public/GroomBindingAsset.h
@@ -486,8 +486,11 @@ public:
 	FOnGroomBindingAssetChanged OnGroomBindingAssetChanged;
 
 	UE_API void RecreateResources();
+	UE_DEPRECATED(5.8, "Use ChangeShaderPlatform(EShaderPlatform) instead.")
 	UE_API void ChangeFeatureLevel(ERHIFeatureLevel::Type PendingFeatureLevel);
+	UE_DEPRECATED(5.8, "Use ChangeShaderPlatform(EShaderPlatform) instead.")
 	UE_API void ChangePlatformLevel(ERHIFeatureLevel::Type PendingFeatureLevel);
+	UE_API void ChangeShaderPlatform(EShaderPlatform InShaderPlatform);
 #endif
 private:
 	UE_API void WaitUntilAsyncPropertyReleased(EGroomBindingAsyncProperties AsyncProperties, EGroomBindingAsyncPropertyLockType LockType) const;
@@ -580,8 +583,7 @@ public:
 	TArray<FString> CachedDerivedDataKey;
 #endif
 #if WITH_EDITOR
-	ERHIFeatureLevel::Type CachedResourcesFeatureLevel = ERHIFeatureLevel::Num;
-	ERHIFeatureLevel::Type CachedResourcesPlatformLevel = ERHIFeatureLevel::Num;
+	EShaderPlatform CachedResourceShaderPlatform = SP_NumPlatforms;
 #endif
 	bool bIsValid = false;
 	uint32 AssetNameHash = 0;
diff --git a/Engine/Plugins/Runtime/HairStrands/Source/HairStrandsCore/Public/GroomComponent.h b/Engine/Plugins/Runtime/HairStrands/Source/HairStrandsCore/Public/GroomComponent.h
--- a/Engine/Plugins/Runtime/HairStrands/Source/HairStrandsCore/Public/GroomComponent.h
+++ b/Engine/Plugins/Runtime/HairStrands/Source/HairStrandsCore/Public/GroomComponent.h
@@ -182,8 +182,11 @@ public:
 	UE_API void InvalidateAndRecreate();
 
 	UE_API virtual void PreFeatureLevelChange(ERHIFeatureLevel::Type PendingFeatureLevel) override;
+	UE_DEPRECATED(5.8, "Use HandlePlatformPreviewChanged(EShaderPlatform) instead.")
 	UE_API void HandlePlatformPreviewChanged(ERHIFeatureLevel::Type InFeatureLevel);
+	UE_DEPRECATED(5.8, "Use HandlePlatformPreviewChanged(EShaderPlatform) instead.")
 	UE_API void HandleFeatureLevelChanged(ERHIFeatureLevel::Type InFeatureLevel);
+	UE_API void HandlePlatformPreviewChanged(EShaderPlatform InShaderPlatform);
 #endif
 
 	UE_API void PostCompilation();
diff --git a/Engine/Plugins/Runtime/HairStrands/Source/HairStrandsEditor/Private/HairStrandsEditor.cpp b/Engine/Plugins/Runtime/HairStrands/Source/HairStrandsEditor/Private/HairStrandsEditor.cpp
--- a/Engine/Plugins/Runtime/HairStrands/Source/HairStrandsEditor/Private/HairStrandsEditor.cpp
+++ b/Engine/Plugins/Runtime/HairStrands/Source/HairStrandsEditor/Private/HairStrandsEditor.cpp
@@ -246,14 +246,11 @@ void FGroomEditor::OnPreExit()
 void FGroomEditor::OnPreviewPlatformChanged(EShaderPlatform NewShaderPlatform)
 {
 #if WITH_EDITOR
-	ERHIFeatureLevel::Type ActiveFeatureLevel = GetMaxSupportedFeatureLevel(NewShaderPlatform);
-
 	for (TObjectIterator<UGroomComponent> It; It; ++It)
 	{
 		if (UGroomComponent* Component = *It)
 		{
-			Component->HandlePlatformPreviewChanged(ActiveFeatureLevel);
-			Component->HandleFeatureLevelChanged(ActiveFeatureLevel);
+			Component->HandlePlatformPreviewChanged(NewShaderPlatform);
 		}
 	}
 #endif
