Ensure we don't access out of bound section data when running AddHairStrandInitMeshSamplesPass.

This seems to happen when groom binding is being rebuilt under certain conditions.

#Synced-CL 52375688

diff --git a/Engine/Plugins/Runtime/HairStrands/Source/HairStrandsCore/Private/HairStrandsMeshProjection.cpp b/Engine/Plugins/Runtime/HairStrands/Source/HairStrandsCore/Private/HairStrandsMeshProjection.cpp
--- a/Engine/Plugins/Runtime/HairStrands/Source/HairStrandsCore/Private/HairStrandsMeshProjection.cpp
+++ b/Engine/Plugins/Runtime/HairStrands/Source/HairStrandsCore/Private/HairStrandsMeshProjection.cpp
@@ -616,7 +616,7 @@ void AddHairStrandInitMeshSamplesPass(
 		uint32 bIsSwapped = 0;
 	};
 	TArray<FSectionData> SectionDatas;
-	SectionDatas.Init(FSectionData(), TotalSectionCount);
+	SectionDatas.Init(FSectionData(), FMath::Max(uint32(MeshData.Sections.Num()), TotalSectionCount));
 	for (const FCachedGeometry::Section& Section : MeshData.Sections)
 	{
 		const uint32 SectionIndex = Section.SectionIndex;
