Fix another crash with Groom binding.

#Synced-CL 52375688

diff --git a/Engine/Plugins/Runtime/HairStrands/Source/HairStrandsCore/Private/GroomManager.cpp b/Engine/Plugins/Runtime/HairStrands/Source/HairStrandsCore/Private/GroomManager.cpp
--- a/Engine/Plugins/Runtime/HairStrands/Source/HairStrandsCore/Private/GroomManager.cpp
+++ b/Engine/Plugins/Runtime/HairStrands/Source/HairStrandsCore/Private/GroomManager.cpp
@@ -320,7 +320,18 @@ struct FHairGeometryCache
 	{
 		check(InLODData);
 		check(InMeshObject);
-		check(UniqueSections.Num() > 0);
+		
+		if (UniqueSections.Num() == 0)
+		{
+			// Binding data is valid but has no section indices
+			static bool bWarningIssued = false;
+			if (!bWarningIssued)
+			{
+				UE_LOG(LogHairStrands, Warning, TEXT("GetOrAllocateCachedGeometry: RootResource LOD %d has no UniqueSectionIndices. Skipping."), InLODIndex);
+				bWarningIssued = true;
+			}
+			return;
+		}
 
 		// Find existing entry
 		const bool bNeedPreviousPosition = IsHairStrandContinuousDecimationReorderingEnabled();
@@ -620,7 +631,10 @@ static void GetOrAllocateCachedGeometry(
 	}
 
 	check(RootResource);
-	check(RootResource->IsDataValid(LODIndex));
+	if (!RootResource->IsDataValid(LODIndex))
+	{
+		return;
+	}
 	const TArray<uint32>& UniqueSections = RootResource->GetLOD(LODIndex)->GetBulkData().Header.UniqueSectionIndices;
 
 	// Create deformed position buffer (output)
diff --git a/Engine/Plugins/Runtime/HairStrands/Source/HairStrandsCore/Private/HairStrandsMeshProjection.cpp b/Engine/Plugins/Runtime/HairStrands/Source/HairStrandsCore/Private/HairStrandsMeshProjection.cpp
--- a/Engine/Plugins/Runtime/HairStrands/Source/HairStrandsCore/Private/HairStrandsMeshProjection.cpp
+++ b/Engine/Plugins/Runtime/HairStrands/Source/HairStrandsCore/Private/HairStrandsMeshProjection.cpp
@@ -589,8 +589,9 @@ void AddHairStrandInitMeshSamplesPass(
 	FHairStrandsLODRestRootResource& RestLODData = *RestResources->GetLOD(LODIndex);
 
 	const uint32 TotalSectionCount = RestLODData.GetBulkData().Header.MeshSectionCount;
+	const uint32 UniqueSectionIndicesCount = RestLODData.GetBulkData().Header.UniqueSectionIndices.Num();
 	const uint32 EffectiveSectionCount = MeshData.Sections.Num();
-	if (EffectiveSectionCount == 0 || RestLODData.SampleCount == 0)
+	if (EffectiveSectionCount == 0 || RestLODData.SampleCount == 0 || UniqueSectionIndicesCount == 0)
 	{
 		return;
 	}
