Dataflow : make cloth , groom and geometry collection use the new template selector

#Synced-CL 52408587

diff --git a/Engine/Plugins/Runtime/HairStrands/Source/HairStrandsEditor/Private/HairStrandsEditor.cpp b/Engine/Plugins/Runtime/HairStrands/Source/HairStrandsEditor/Private/HairStrandsEditor.cpp
--- a/Engine/Plugins/Runtime/HairStrands/Source/HairStrandsEditor/Private/HairStrandsEditor.cpp
+++ b/Engine/Plugins/Runtime/HairStrands/Source/HairStrandsEditor/Private/HairStrandsEditor.cpp
@@ -157,8 +157,6 @@ void FGroomEditor::StartupModule()
 	UThumbnailManager::Get().RegisterCustomRenderer(UGroomAsset::StaticClass(), UGroomAssetThumbnailRenderer::StaticClass());
 	UThumbnailManager::Get().RegisterCustomRenderer(UGroomBindingAsset::StaticClass(), UGroomBindingAssetThumbnailRenderer::StaticClass());
 
-	// default no dataflow template
-	UE::Groom::RegisterGroomDataflowTemplate({ "SelectNoDataflowTemplate", "No Dataflow", "Don't add a Dataflow to this Groom Asset. Useful when a Dataflow already exists.", "", false });
 	UE::Groom::RegisterGroomDataflowTemplatePath(TEXT("/HairStrands/Dataflow/Templates"));
 }
 
diff --git a/Engine/Plugins/Runtime/HairStrands/Source/HairStrandsEditor/Private/HairStrandsFactory.cpp b/Engine/Plugins/Runtime/HairStrands/Source/HairStrandsEditor/Private/HairStrandsFactory.cpp
--- a/Engine/Plugins/Runtime/HairStrands/Source/HairStrandsEditor/Private/HairStrandsFactory.cpp
+++ b/Engine/Plugins/Runtime/HairStrands/Source/HairStrandsEditor/Private/HairStrandsFactory.cpp
@@ -13,6 +13,8 @@
 #include "GroomCacheImportOptions.h"
 #include "GroomBuilder.h"
 #include "GroomImportOptionsWindow.h"
+#include "Dataflow/AssetDefinition_DataflowAsset.h"
+#include "Dataflow/DataflowTemplateRegistry.h"
 #include "Dataflow/DataflowObject.h"
 #include "Dialog/SMessageDialog.h"
 #include "HairDescription.h"
@@ -150,6 +152,14 @@ namespace UE::Groom
 	void RegisterGroomDataflowTemplate(const FGroomDataflowTemplateData& DataflowTemplate)
 	{
 		FGroomDataflowTemplateManager::GetInstance().AddDataflowTemplateData(DataflowTemplate);
+		FDataflowTemplateRegistry::Get().RegisterTemplateAsset(
+			UGroomAsset::StaticClass(),
+			{
+				.AssetPath = FSoftObjectPath(DataflowTemplate.TemplatePath),
+				.DisplayName = FText::FromString(DataflowTemplate.TemplateTitle),
+				.Tooltip = FText::FromString(DataflowTemplate.TemplateTooltip),
+				.Icon = FSlateIcon("Groom", "ClassThumbnail.GroomAsset")
+			});
 	}
 
 	void UnregisterGroomDataflowTemplate(const FString& TemplateName)
@@ -160,6 +170,8 @@ namespace UE::Groom
 	void RegisterGroomDataflowTemplatePath(const FString& DataflowTemplatePath)
 	{
 		FGroomDataflowTemplateManager::GetInstance().AddDataflowTemplatePath(DataflowTemplatePath);
+		FDataflowTemplateRegistry::Get().RegisterTemplateFolder(UGroomAsset::StaticClass(),
+			{ DataflowTemplatePath, FSlateIcon("Groom", "ClassThumbnail.GroomAsset") });
 	}
 
 	void UnregisterGroomDataflowTemplatePath(const FString& DataflowTemplatePath)
@@ -169,51 +181,7 @@ namespace UE::Groom
 
 	bool BuildGroomDataflowAsset(UGroomAsset* GroomAsset)
 	{
-		// Create a new Dataflow asset
-		const FString DataflowPath = FPackageName::GetLongPackagePath(GroomAsset->GetOutermost()->GetName());
-		const FString GroomAssetName = GroomAsset->GetName();
-		FString DataflowName = FString(TEXT("DF_")) + (GroomAssetName.StartsWith(TEXT("GA_")) ? GroomAssetName.RightChop(3) : GroomAssetName);
-		FString DataflowPackageName = FPaths::Combine(DataflowPath, DataflowName);
-		if (FindPackage(nullptr, *DataflowPackageName))
-		{
-			// If a Dataflow asset already exists with this name, make a unique name from it to avoid clobbering it
-			DataflowPackageName = MakeUniqueObjectName(nullptr, UPackage::StaticClass(), FName(DataflowPackageName)).ToString();
-			DataflowName = FPaths::GetBaseFilename(DataflowPackageName);
-		}
-		UPackage* const DataflowPackage = CreatePackage(*DataflowPackageName);
-
-		FString GroomDataflowTemplate;
-		auto OnButtonClicked = [&GroomDataflowTemplate](FString TemplatePath)
-			{
-				GroomDataflowTemplate = TemplatePath;
-			};
-
-		TArray<SMessageDialog::FButton> Buttons = FGroomDataflowTemplateManager::GetInstance().BuildDataflowTemplateButtons(OnButtonClicked);
-
-		// Select the template
-		TSharedRef<SMessageDialog> SelectTemplateMessageDialog = SNew(SMessageDialog)
-			.Title(FText(LOCTEXT("SelectTemplateTitle", "Select a Groom Dataflow Template")))
-			.Message(LOCTEXT("SelectTemplateMessage", "Select a template for this groom asset:"))
-			.Buttons(Buttons);
-
-		// result will be set in GroomDataflowTemplate
-		SelectTemplateMessageDialog->ShowModal();
-
-		// Load the dataflow template into the groom asset
-		UDataflow* const Template = GroomDataflowTemplate.IsEmpty() ? nullptr : LoadObject<UDataflow>(nullptr, *GroomDataflowTemplate);
-		if (UDataflow* const Dataflow = Template ? DuplicateObject(Template, DataflowPackage, FName(DataflowName)) : nullptr)
-		{
-			Dataflow->MarkPackageDirty();
-
-			// Notify the asset registry
-			FAssetRegistryModule::AssetCreated(Dataflow);
-
-			// Set the Dataflow to the groom asset
-			GroomAsset->GetDataflowSettings().SetDataflowAsset(Dataflow);
-
-			return true;
-		}
-		return false;
+		return UE::DataflowAssetDefinitionHelpers::SetDataflowFromTemplatePicker(GroomAsset);
 	}
 }
 
