Dataflow groom : Add options on the resample curve node  to not extend the end of the curve and to allow for non power of two points count
- Also fix the resampling code in the curve facade that could sometimes produce one more point than expected

#Synced-CL 52448907

diff --git a/Engine/Plugins/Runtime/HairStrands/Source/HairStrandsDataflow/Private/ResampleGuidesPointsNode.cpp b/Engine/Plugins/Runtime/HairStrands/Source/HairStrandsDataflow/Private/ResampleGuidesPointsNode.cpp
--- a/Engine/Plugins/Runtime/HairStrands/Source/HairStrandsDataflow/Private/ResampleGuidesPointsNode.cpp
+++ b/Engine/Plugins/Runtime/HairStrands/Source/HairStrandsDataflow/Private/ResampleGuidesPointsNode.cpp
@@ -6,6 +6,8 @@
 
 #include UE_INLINE_GENERATED_CPP_BY_NAME(ResampleGuidesPointsNode)
 
+#define LOCTEXT_NAMESPACE "ResampleGuidesPointsNode"
+
 namespace UE::Groom::Private
 {
 	int32 ComputeCurvesPoints(const GeometryCollection::Facades::FCollectionCurveGeometryFacade& CurvesFacade, const FDataflowCurveSelection& CurveSelection, const int32 GuidePoints) 
@@ -71,7 +73,7 @@ namespace UE::Groom::Private
 		SamplePositions[SampleOffset] = 2 * SamplePositions[SampleOffset+1] - SamplePositions[SampleOffset+2];
 	}
 
-	void ResampleCurvesPoints(GeometryCollection::Facades::FCollectionCurveGeometryFacade& CurvesFacade, const FDataflowCurveSelection& CurveSelection, const int32 GuidePoints)
+	void ResampleCurvesPoints(GeometryCollection::Facades::FCollectionCurveGeometryFacade& CurvesFacade, const FDataflowCurveSelection& CurveSelection, const int32 GuidePoints, bool bForSimulation)
 	{
 		if(CurvesFacade.IsValid())
 		{
@@ -86,6 +88,7 @@ namespace UE::Groom::Private
 
 			const bool bValidSelection = CurveSelection.IsValidForCollection(CurvesFacade.GetManagedArrayCollection());
 
+			TArray<FVector3f> ResampledPoints; // used only when bForSimulation = false
 			int32 PointOffset = 0;
 			int32 SampleOffset = 0;
 			for(int32 CurveIndex = 0; CurveIndex < NumCurves; ++CurveIndex)
@@ -93,8 +96,21 @@ namespace UE::Groom::Private
 				const int32 NumPoints = CurvesFacade.GetCurvePointOffsets()[CurveIndex]-PointOffset;
 				if(!bValidSelection || (bValidSelection && CurveSelection.IsSelected(CurveIndex)))
 				{
-					UE::Groom::Private::BuildResampledPoints(CurvesFacade.GetPointRestPositions(), CurvesFacade.GetCurveGeometryIndices(), CurvesFacade.GetCurvePointOffsets(),
+					if (bForSimulation)
+					{
+						UE::Groom::Private::BuildResampledPoints(CurvesFacade.GetPointRestPositions(), CurvesFacade.GetCurveGeometryIndices(), CurvesFacade.GetCurvePointOffsets(),
 							{}, CurveIndex, PointOffset, SampleOffset, PointSampledPositions, GuidePoints);
+					}
+					else
+					{
+						ResampledPoints.Reset();
+						CurvesFacade.GetCurveResampledPositions(CurveIndex, GuidePoints, ResampledPoints);
+						const int32 NumResampledPoints = FMath::Min(GuidePoints, ResampledPoints.Num());
+						for (int32 Index = 0; Index < NumResampledPoints; ++Index)
+						{
+							PointSampledPositions[SampleOffset + Index] = ResampledPoints[Index];
+						}
+					}
 					SampleOffset += GuidePoints;
 				}
 				else
@@ -158,16 +174,29 @@ void FResampleCurvePointsDataflowNode::Evaluate(UE::Dataflow::FContext& Context,
 		FManagedArrayCollection GuidesCollection = GetValue<FManagedArrayCollection>(Context, &Collection);
 		const FDataflowCurveSelection& DataflowSelection = GetValue<FDataflowCurveSelection>(Context, &CurveSelection);
 		
-		const int32 GuidePoints = IsConnected(&NumPoints) ?
-			UE::Groom::Private::ClosestPowerOfTwo(GetValue<int32>(Context, &NumPoints)) : static_cast<int32>(PointsCount);
-		
+		const int32 InNumPoints = GetValue<int32>(Context, &NumPoints);
+		int32 GuidePoints = InNumPoints;
+		if (bResampleForSimulation)
+		{
+			GuidePoints = IsConnected(&NumPoints)
+				? UE::Groom::Private::ClosestPowerOfTwo(InNumPoints)
+				: static_cast<int32>(PointsCount);
+		}
+		if (GuidePoints < 2)
+		{
+			GuidePoints = 2;
+			Context.Warning(LOCTEXT("NumPointToolSmall", "Number of points cannot be smaller than 2, setting it to 2 for evaluation"));
+		}
+
 		GeometryCollection::Facades::FCollectionCurveGeometryFacade CurvesFacade(GuidesCollection);
 		
 		if(CurvesFacade.IsValid() && GuidePoints > 0)
 		{
-			UE::Groom::Private::ResampleCurvesPoints(CurvesFacade, DataflowSelection, GuidePoints);
+			UE::Groom::Private::ResampleCurvesPoints(CurvesFacade, DataflowSelection, GuidePoints, bResampleForSimulation);
 		}
 		
 		SetValue(Context, MoveTemp(GuidesCollection), &Collection);
 	}
 }
+
+#undef LOCTEXT_NAMESPACE
\ No newline at end of file
diff --git a/Engine/Plugins/Runtime/HairStrands/Source/HairStrandsDataflow/Public/ResampleGuidesPointsNode.h b/Engine/Plugins/Runtime/HairStrands/Source/HairStrandsDataflow/Public/ResampleGuidesPointsNode.h
--- a/Engine/Plugins/Runtime/HairStrands/Source/HairStrandsDataflow/Public/ResampleGuidesPointsNode.h
+++ b/Engine/Plugins/Runtime/HairStrands/Source/HairStrandsDataflow/Public/ResampleGuidesPointsNode.h
@@ -90,12 +90,20 @@ private:
 	UPROPERTY(meta = (DataflowInput, DisplayName = "Curve Selection"))
 	FDataflowCurveSelection CurveSelection;
 
+	/** 
+	* When true, the resampling is done in the context of simulation
+	* - Favors power of two point count ( optimized for GPU solvers )
+	* - Extend the begginning of the curve to maintain stiffness
+	*/
+	UPROPERTY(EditAnywhere, Category = Settings)
+	bool bResampleForSimulation = true;
+
 	/** Max number of points */
-	UPROPERTY(EditAnywhere, Category="Geometry", meta = (DisplayName = "Points Count"))
+	UPROPERTY(EditAnywhere, Category = Settings, meta = (DisplayName = "Points Count", EditCondition="bResampleForSimulation", EditConditionHides))
 	EGroomNumPoints PointsCount = EGroomNumPoints::Size16;
 	
 	/** Max number of points (to be able to plug into a variable since enum is not supported in dataflow) */
-	UPROPERTY(meta = (DataflowInput, DisplayName = "Points Count"))
+	UPROPERTY(EditAnywhere, Category = Settings, meta = (DataflowInput, DisplayName = "Points Count", ClampMin = "2", UIMin = "2", EditCondition = "!bResampleForSimulation", EditConditionHides))
 	int32 NumPoints = 16;
 };
 
diff --git a/Engine/Source/Runtime/Experimental/Chaos/Private/GeometryCollection/Facades/CollectionCurveFacade.cpp b/Engine/Source/Runtime/Experimental/Chaos/Private/GeometryCollection/Facades/CollectionCurveFacade.cpp
--- a/Engine/Source/Runtime/Experimental/Chaos/Private/GeometryCollection/Facades/CollectionCurveFacade.cpp
+++ b/Engine/Source/Runtime/Experimental/Chaos/Private/GeometryCollection/Facades/CollectionCurveFacade.cpp
@@ -379,10 +379,13 @@ namespace GeometryCollection::Facades
 			}
 			PointIndex = (NextIndex - 1);
 
-			// we may have reach the end of gone too far
-			if (SampledDistance >= DistancesAt.Last())
+			// we may have reach the end or gone too far
+			if (SampledDistance >= DistancesAt.Last()/* && NumSamples > OutPoints.Num()*/)
 			{
-				OutPoints.Add(Positions.Last());
+				if (OutPoints.Num() < NumSamples)
+				{
+					OutPoints.Add(Positions.Last());
+				}
 				break;
 			}
 		}
