groom : Guide to joints terminal node : attempt to fix an out of bounds crash

#Synced-CL 52663992

diff --git a/Engine/Plugins/Runtime/HairStrands/Source/HairStrandsDataflow/Private/GuidesToJointsNode.cpp b/Engine/Plugins/Runtime/HairStrands/Source/HairStrandsDataflow/Private/GuidesToJointsNode.cpp
--- a/Engine/Plugins/Runtime/HairStrands/Source/HairStrandsDataflow/Private/GuidesToJointsNode.cpp
+++ b/Engine/Plugins/Runtime/HairStrands/Source/HairStrandsDataflow/Private/GuidesToJointsNode.cpp
@@ -402,6 +402,7 @@ void FDataflowGuidesToJointsNode::UpdateSkeletalMesh(UE::Dataflow::FContext& Con
 
 		InOutSkeletalMesh->SetSkeleton(InOutSkeleton);
 		InOutSkeletalMesh->SetRefSkeleton(InOutSkeleton->GetReferenceSkeleton());
+		InOutSkeleton->RebuildLinkup(InOutSkeletalMesh);
 		InOutSkeletalMesh->CalculateInvRefMatrices();
 
 		if (!bIsCopied)
