PR #13932: Update GroomComponent.cpp

#Synced-CL 52685603

diff --git a/Engine/Plugins/Runtime/HairStrands/Source/HairStrandsCore/Private/GroomComponent.cpp b/Engine/Plugins/Runtime/HairStrands/Source/HairStrandsCore/Private/GroomComponent.cpp
--- a/Engine/Plugins/Runtime/HairStrands/Source/HairStrandsCore/Private/GroomComponent.cpp
+++ b/Engine/Plugins/Runtime/HairStrands/Source/HairStrandsCore/Private/GroomComponent.cpp
@@ -640,7 +640,7 @@ public:
 		FSceneInterface& LocalScene = GetScene();
 		for (TRefCountPtr<FHairGroupInstance>& Instance : HairGroupInstances)
 		{
-			if (Instance->IsValid() || Instance->Strands.ClusterResource)
+			if (Instance.IsValid() && (Instance->IsValid() || Instance->Strands.ClusterResource))
 			{
 				check(Instance->HairGroupPublicData != nullptr);
 				Instance->Debug.Proxy = this;
@@ -664,7 +664,7 @@ public:
 		FSceneInterface& LocalScene = GetScene();
 		for (TRefCountPtr<FHairGroupInstance>& Instance : HairGroupInstances)
 		{
-			if (Instance->IsValid() || Instance->Strands.ClusterResource)
+			if (Instance.IsValid() && (Instance->IsValid() || Instance->Strands.ClusterResource))
 			{
 				LocalScene.RemoveHairStrands(Instance);
 				Instance->Debug.Proxy = nullptr;
@@ -2088,7 +2088,7 @@ FPrimitiveSceneProxy* UGroomComponent::CreateSceneProxy()
 	bool bIsValid = false;
 	for (const TRefCountPtr<FHairGroupInstance>& Instance : HairGroupInstances)
 	{
-		bIsValid |= Instance->IsValid();
+		bIsValid |= Instance.IsValid() && Instance->IsValid();
 	}
 	if (!bIsValid)
 	{
