Fix crash in hair strands/cards streaming code with root resources.

#Synced-CL 52685603

diff --git a/Engine/Plugins/Runtime/HairStrands/Source/HairStrandsCore/Private/GroomResources.cpp b/Engine/Plugins/Runtime/HairStrands/Source/HairStrandsCore/Private/GroomResources.cpp
--- a/Engine/Plugins/Runtime/HairStrands/Source/HairStrandsCore/Private/GroomResources.cpp
+++ b/Engine/Plugins/Runtime/HairStrands/Source/HairStrandsCore/Private/GroomResources.cpp
@@ -1614,9 +1614,7 @@ FHairStrandsLODRestRootResource::FHairStrandsLODRestRootResource(TArray<FHairStr
 	FHairCommonResource(EHairStrandsAllocationType::Deferred, InResourceName, InOwnerName),
 	BulkDataArray(InBulkDataArray), BulkDataIndex(InBulkDataIndex), CurveType(InCurveType), MeshLODIndex(InMeshLODIndex)
 {
-	// MaxAvailableCurveCount is not set to 0, as we need to track the loaded curve data *per LOD*. 
-	// This is done through callback to InternalIsLODDataLoaded()
-	// MaxAvailableCurveCount = 0;
+	MaxAvailableCurveCount = 0;
 
 	Reset();
 }
diff --git a/Engine/Plugins/Runtime/HairStrands/Source/HairStrandsCore/Private/HairStrandsInterpolation.cpp b/Engine/Plugins/Runtime/HairStrands/Source/HairStrandsCore/Private/HairStrandsInterpolation.cpp
--- a/Engine/Plugins/Runtime/HairStrands/Source/HairStrandsCore/Private/HairStrandsInterpolation.cpp
+++ b/Engine/Plugins/Runtime/HairStrands/Source/HairStrandsCore/Private/HairStrandsInterpolation.cpp
@@ -700,17 +700,25 @@ void AddHairStrandsInterpolationPass(
 		ShaderPrint::SetParameters(GraphBuilder, *ShaderPrintData, Parameters->ShaderDrawParameters);
 	}
 	
-	const bool bSupportDynamicMesh = 
+	const bool bSupportDynamicMesh_Ren = 
 		RenRestRootResources &&
 		RenRestRootResources->GetRootCount() > 0 && 
 		MeshLODIndex >= 0 && 
 		RenRestRootResources->IsValid(MeshLODIndex) &&
 		RenDeformedRootResources->IsValid(MeshLODIndex);
 
+	const bool bSupportDynamicMesh_Sim =
+		SimRestRootResources &&
+		SimDeformedRootResources && 
+		SimRestRootResources->GetRootCount() > 0 && 
+		MeshLODIndex >= 0 &&
+		SimRestRootResources->IsValid(MeshLODIndex) &&
+		SimDeformedRootResources->IsValid(MeshLODIndex);
+
 	// Guides
 	// Guide interpolation with dynamic mesh geometry is now support with guide cache. It implicitly assumes that the guide-cache follows correctly the binding surface
 	bool bSupportGlobalInterpolation = false;
-	if (bSupportDynamicMesh && (Instance->Guides.bIsSimulationEnable || Instance->Guides.bHasGlobalInterpolation || Instance->Guides.bIsDeformationEnable || Instance->Guides.bIsSimulationCacheEnable))
+	if (bSupportDynamicMesh_Sim && (Instance->Guides.bIsSimulationEnable || Instance->Guides.bHasGlobalInterpolation || Instance->Guides.bIsDeformationEnable || Instance->Guides.bIsSimulationCacheEnable))
 	{
 		const FHairStrandsLODRestRootResource& Sim_RestLODDatas = *SimRestRootResources->GetLOD(MeshLODIndex);
 		const FHairStrandsLODDeformedRootResource& Sim_DeformedLODDatas = *SimDeformedRootResources->GetLOD(MeshLODIndex);
@@ -731,7 +739,7 @@ void AddHairStrandsInterpolationPass(
 	}
 
 	// Strands
-	if (bSupportDynamicMesh)
+	if (bSupportDynamicMesh_Ren)
 	{
 		const FHairStrandsLODRestRootResource& Ren_RestLODDatas = *RenRestRootResources->GetLOD(MeshLODIndex);
 		const FHairStrandsLODDeformedRootResource& Ren_DeformedLODDatas = *RenDeformedRootResources->GetLOD(MeshLODIndex);
@@ -758,6 +766,11 @@ void AddHairStrandsInterpolationPass(
 
 	const bool bHasLocalDeformation = Instance->Guides.bIsSimulationEnable || bSupportGlobalInterpolation || Instance->Guides.bIsDeformationEnable || Instance->Guides.bIsSimulationCacheEnable;
 
+	// Enable support for dynamic mesh (skinning) only if roots data are valids
+	// * If there is no need of Sim data, only rely on Ren resources
+	// * If there is need of Sim data, rely on both Sim and Ren resources
+	const bool bSupportDynamicMesh = bHasLocalDeformation ? (bSupportDynamicMesh_Ren && bSupportDynamicMesh_Sim) : bSupportDynamicMesh_Ren;
+
 	// Select dynamic geometry permutation, based on the Simulation/RBF/InterpolationType
 	int32 DynamicGeometryType = 0;
 	{
diff --git a/Engine/Plugins/Runtime/HairStrands/Source/HairStrandsCore/Public/GroomResources.h b/Engine/Plugins/Runtime/HairStrands/Source/HairStrandsCore/Public/GroomResources.h
--- a/Engine/Plugins/Runtime/HairStrands/Source/HairStrandsCore/Public/GroomResources.h
+++ b/Engine/Plugins/Runtime/HairStrands/Source/HairStrandsCore/Public/GroomResources.h
@@ -553,7 +553,6 @@ struct FHairStrandsLODDeformedRootResource : public FHairCommonResource
 	/* Init/Release buffers */
 	virtual void InternalAllocate(FRDGBuilder& GraphBuilder, uint32 InCurveCount, uint32 InPointCount) override;
 	virtual void InternalRelease() override;
-	//virtual bool InternalIsLODDataLoaded(uint32 InRequestedCurveCount, uint32 InRequestedPointCount) const override;
 
 	/* Get the resource name */
 	virtual FString GetFriendlyName() const override { return TEXT("FHairStrandsLODDeformedRootResource "); }
