Fix long shader compilation time on hair debug shaders on certain platforms.

#Synced-CL 52743879

diff --git a/Engine/Plugins/Runtime/HairStrands/Source/HairStrandsCore/Private/GroomDebug.cpp b/Engine/Plugins/Runtime/HairStrands/Source/HairStrandsCore/Private/GroomDebug.cpp
--- a/Engine/Plugins/Runtime/HairStrands/Source/HairStrandsCore/Private/GroomDebug.cpp
+++ b/Engine/Plugins/Runtime/HairStrands/Source/HairStrandsCore/Private/GroomDebug.cpp
@@ -763,8 +763,6 @@ public:
 		// Skip optimization for avoiding long compilation time due to large UAV writes
 		FGlobalShader::ModifyCompilationEnvironment(Parameters, OutEnvironment);
 		ShaderPrint::ModifyCompilationEnvironment(Parameters, OutEnvironment);
-
-		OutEnvironment.CompilerFlags.Add(CFLAG_Debug);
 		OutEnvironment.SetDefine(TEXT("SHADER_PRINT_INSTANCE"), 1);
 	}
 };
@@ -1010,8 +1008,6 @@ public:
 		// Skip optimization for avoiding long compilation time due to large UAV writes
 		FGlobalShader::ModifyCompilationEnvironment(Parameters, OutEnvironment);
 		ShaderPrint::ModifyCompilationEnvironment(Parameters, OutEnvironment);
-
-		OutEnvironment.CompilerFlags.Add(CFLAG_Debug);
 		OutEnvironment.SetDefine(TEXT("SHADER_PRINT_MEMORY"), 1);
 	}
 };
