Add back CFLAG_Debug for hair debug shader for certain platforms to avoid long compilation time.

#Synced-CL 52780392

diff --git a/Engine/Plugins/Runtime/HairStrands/Source/HairStrandsCore/Private/GroomDebug.cpp b/Engine/Plugins/Runtime/HairStrands/Source/HairStrandsCore/Private/GroomDebug.cpp
--- a/Engine/Plugins/Runtime/HairStrands/Source/HairStrandsCore/Private/GroomDebug.cpp
+++ b/Engine/Plugins/Runtime/HairStrands/Source/HairStrandsCore/Private/GroomDebug.cpp
@@ -70,6 +70,15 @@ enum class HairStrandsTriangleType
 	DeformedPose,
 };
 
+static void AddShaderPrintDebugFlag(EShaderPlatform InPlatform, FShaderCompilerEnvironment& OutEnvironment)
+{
+	const bool bNeedDebugFlags = FDataDrivenShaderPlatformInfo::GetIsLanguageNintendo(InPlatform);
+	if (bNeedDebugFlags)
+	{
+		OutEnvironment.CompilerFlags.Add(CFLAG_Debug);
+	}
+}
+
 ///////////////////////////////////////////////////////////////////////////////////////////////////
 
 class FHairProjectionHairDebugCS : public FGlobalShader
@@ -763,6 +772,7 @@ public:
 		// Skip optimization for avoiding long compilation time due to large UAV writes
 		FGlobalShader::ModifyCompilationEnvironment(Parameters, OutEnvironment);
 		ShaderPrint::ModifyCompilationEnvironment(Parameters, OutEnvironment);
+		AddShaderPrintDebugFlag(Parameters.Platform, OutEnvironment);
 		OutEnvironment.SetDefine(TEXT("SHADER_PRINT_INSTANCE"), 1);
 	}
 };
@@ -1008,6 +1018,7 @@ public:
 		// Skip optimization for avoiding long compilation time due to large UAV writes
 		FGlobalShader::ModifyCompilationEnvironment(Parameters, OutEnvironment);
 		ShaderPrint::ModifyCompilationEnvironment(Parameters, OutEnvironment);
+		AddShaderPrintDebugFlag(Parameters.Platform, OutEnvironment);
 		OutEnvironment.SetDefine(TEXT("SHADER_PRINT_MEMORY"), 1);
 	}
 };
